ChaosFlesh:
-Fixed GT/PT collection update through solver proxy, now proxy takes shared ptr
-Fixed deformable solver proxy deletion (tested with QA muscle flesh asset)
-Added flesh component propagation if flesh asset changes its collection

diff --git a/Engine/Plugins/Experimental/ChaosFlesh/Source/ChaosFleshEditor/Private/ChaosFlesh/Cmd/ChaosFleshCommands.cpp b/Engine/Plugins/Experimental/ChaosFlesh/Source/ChaosFleshEditor/Private/ChaosFlesh/Cmd/ChaosFleshCommands.cpp
--- a/Engine/Plugins/Experimental/ChaosFlesh/Source/ChaosFleshEditor/Private/ChaosFlesh/Cmd/ChaosFleshCommands.cpp
+++ b/Engine/Plugins/Experimental/ChaosFlesh/Source/ChaosFleshEditor/Private/ChaosFlesh/Cmd/ChaosFleshCommands.cpp
@@ -158,7 +158,7 @@ FChaosFleshCommands::FindQualifyingTetrahedra(const TArray<FString>& Args, UWorl
 					{
 						if (const UFleshAsset* RestCollection = FleshComponent->GetRestCollection())
 						{
-							const FFleshCollection* FleshCollection = RestCollection->GetCollection();
+							TSharedPtr<const FFleshCollection> FleshCollection = RestCollection->GetFleshCollection();
 
 							const TManagedArray<FIntVector4>* TetMesh =
 								FleshCollection->FindAttribute<FIntVector4>(
@@ -568,7 +568,7 @@ void FChaosFleshCommands::CreateGeometryCache(const TArray<FString>& Args, UWorl
 			{
 				continue;
 			}
-			const FFleshCollection* FleshCollection = RestCollection->GetCollection();
+			TSharedPtr<const FFleshCollection> FleshCollection = RestCollection->GetFleshCollection();
 			if (!FleshCollection)
 			{
 				continue;
diff --git a/Engine/Plugins/Experimental/ChaosFlesh/Source/ChaosFleshEngine/Private/ChaosCache/FleshComponentCacheAdapter.cpp b/Engine/Plugins/Experimental/ChaosFlesh/Source/ChaosFleshEngine/Private/ChaosCache/FleshComponentCacheAdapter.cpp
--- a/Engine/Plugins/Experimental/ChaosFlesh/Source/ChaosFleshEngine/Private/ChaosCache/FleshComponentCacheAdapter.cpp
+++ b/Engine/Plugins/Experimental/ChaosFlesh/Source/ChaosFleshEngine/Private/ChaosCache/FleshComponentCacheAdapter.cpp
@@ -692,7 +692,7 @@ namespace Chaos
 				}
 
 				const UFleshAsset* RestCollectionAsset = FleshComp->GetRestCollection();
-				const FFleshCollection* RestCollection = RestCollectionAsset ? RestCollectionAsset->GetCollection() : nullptr;
+				TSharedPtr<const FFleshCollection> RestCollection = RestCollectionAsset ? RestCollectionAsset->GetFleshCollection() : nullptr;
 				if (!RestCollection)
 				{
 					UE_LOG(LogChaosFleshCache, Error, TEXT("Failed to get rest collection from flesh component: '%s'"), *FleshComp->GetName());
diff --git a/Engine/Plugins/Experimental/ChaosFlesh/Source/ChaosFleshEngine/Private/ChaosFlesh/ChaosDeformableGameplayComponent.cpp b/Engine/Plugins/Experimental/ChaosFlesh/Source/ChaosFleshEngine/Private/ChaosFlesh/ChaosDeformableGameplayComponent.cpp
--- a/Engine/Plugins/Experimental/ChaosFlesh/Source/ChaosFleshEngine/Private/ChaosFlesh/ChaosDeformableGameplayComponent.cpp
+++ b/Engine/Plugins/Experimental/ChaosFlesh/Source/ChaosFleshEngine/Private/ChaosFlesh/ChaosDeformableGameplayComponent.cpp
@@ -75,7 +75,7 @@ void UDeformableGameplayComponent::DetectEnvironmentCollisions(const int32 MaxNu
 
 	// Init inputs, and return if we have no candidates.
 	GeometryCollection::Facades::FConstraintOverrideCandidateFacade CnstrCandidates(
-		*GetRestCollection()->GetCollection());
+		*GetRestCollection()->GetFleshCollection());
 	if (!CnstrCandidates.IsValid())
 	{
 		return;
@@ -95,7 +95,7 @@ void UDeformableGameplayComponent::DetectEnvironmentCollisions(const int32 MaxNu
 
 	// Get bone weights for skinning vertices of the tetrahedron mesh.
 	GeometryCollection::Facades::FVertexBoneWeightsFacade BoneWeightsFacade(
-		*GetRestCollection()->GetCollection());
+		*GetRestCollection()->GetFleshCollection());
 	if (!BoneWeightsFacade.IsValid())
 	{
 		return;
@@ -106,7 +106,7 @@ void UDeformableGameplayComponent::DetectEnvironmentCollisions(const int32 MaxNu
 	const TManagedArray<FTransform3f>* RestTransforms = nullptr;
 	if (const UFleshAsset* FleshAsset = GetRestCollection())
 	{
-		if (const FFleshCollection* Rest = FleshAsset->GetCollection())
+		if (TSharedPtr<const FFleshCollection> Rest = FleshAsset->GetFleshCollection())
 		{
 			GeometryCollection::Facades::FTransformSource TransformSource(*Rest);
 			if (TransformSource.IsValid())
@@ -199,7 +199,7 @@ void UDeformableGameplayComponent::DetectEnvironmentCollisions(const int32 MaxNu
 					//FTransform RestInComponent = ComponentSpaceMatrix(SkeletalMeshComponent->GetSkeletalMeshAsset()->GetSkeleton()->GetReferenceSkeleton(), RefMat, Bones[j]);
 					//FVector LocalPoint = RestInComponent.TransformPosition(ToDouble(RestPos));
 					
-					const FFleshCollection* Collection = GetRestCollection()->GetCollection();
+					TSharedPtr<const FFleshCollection> Collection = GetRestCollection()->GetFleshCollection();
 					FVector3f LocalPoint = (*RestTransforms)[Bones[j]].InverseTransformPosition(RestPos);
 					if (const TManagedArray<FTransform3f>* ConstTransformsPtr = Collection->FindAttributeTyped<FTransform3f>(FName("ComponentTransform"), FName("ComponentTransformGroup")))
 					{
diff --git a/Engine/Plugins/Experimental/ChaosFlesh/Source/ChaosFleshEngine/Private/ChaosFlesh/ChaosDeformableTetrahedralComponent.cpp b/Engine/Plugins/Experimental/ChaosFlesh/Source/ChaosFleshEngine/Private/ChaosFlesh/ChaosDeformableTetrahedralComponent.cpp
--- a/Engine/Plugins/Experimental/ChaosFlesh/Source/ChaosFleshEngine/Private/ChaosFlesh/ChaosDeformableTetrahedralComponent.cpp
+++ b/Engine/Plugins/Experimental/ChaosFlesh/Source/ChaosFleshEngine/Private/ChaosFlesh/ChaosDeformableTetrahedralComponent.cpp
@@ -96,6 +96,9 @@ void UDeformableTetrahedralComponent::EndPlay(const EEndPlayReason::Type ReasonE
 void UDeformableTetrahedralComponent::SetRestCollection(const UFleshAsset* InRestCollection)
 {
 	RestCollection = InRestCollection;
+	ResetDynamicCollection();
+	// Rebuild the physics state since the dynamic collection has changed
+	RecreatePhysicsState();
 	Invalidate();
 	UpdateLocalBounds();
 	ResetProceduralMesh();
@@ -109,7 +112,7 @@ UDeformablePhysicsComponent::FThreadingProxy* UDeformableTetrahedralComponent::N
 	UpdateSimSpaceTransformIndex();
 	if (const UFleshAsset* RestAsset = GetRestCollection())
 	{
-		if (const FFleshCollection* Rest = RestAsset->GetCollection())
+		if (TSharedPtr<const FFleshCollection> Rest = RestAsset->GetFleshCollection())
 		{
 			if (Rest->NumElements(FGeometryCollection::VerticesGroup))
 			{
@@ -122,7 +125,7 @@ UDeformablePhysicsComponent::FThreadingProxy* UDeformableTetrahedralComponent::N
 					SimulationCollection = NewObject<USimulationAsset>(this, TEXT("Flesh Simulation Asset"));
 				}
 
-				GetDynamicCollection()->Reset(Rest);
+				GetDynamicCollection()->Reset(Rest.Get());
 				if (const FManagedArrayCollection* Dynamic = GetDynamicCollection()->GetCollection())
 				{
 					// Mesh points are in component space, such that the exterior hull aligns with the
@@ -134,7 +137,7 @@ UDeformablePhysicsComponent::FThreadingProxy* UDeformableTetrahedralComponent::N
 						ComponentToWorldXf,
 						ComponentToSimXf,
 						SimulationSpace.SimSpace,
-						*Rest,
+						Rest.ToSharedRef(),
 						*Dynamic);
 				}
 			}
@@ -152,7 +155,7 @@ UDeformablePhysicsComponent::FDataMapValue UDeformableTetrahedralComponent::NewD
 	{
 		if (const UFleshAsset* FleshAsset = GetRestCollection())
 		{
-			if (const FFleshCollection* Rest = FleshAsset->GetCollection())
+			if (TSharedPtr<const FFleshCollection> Rest = FleshAsset->GetFleshCollection())
 			{
 				FTransformSource TransformSource(*Rest);
 				FMuscleActivationFacade MuscleActivationFacade(*Rest);
@@ -349,7 +352,6 @@ FTransform UDeformableTetrahedralComponent::GetSimSpaceRestTransform() const
 void UDeformableTetrahedralComponent::UpdateFromSimulation(const FDataMapValue* SimulationBuffer)
 {
 	PERF_SCOPE(STAT_ChaosDeformable_UDeformableTetrahedralComponent_UpdateFromSimualtion);
-
 	if (const FFleshThreadingProxy::FFleshOutputBuffer* FleshBuffer = (*SimulationBuffer)->As<FFleshThreadingProxy::FFleshOutputBuffer>())
 	{
 		if (GetDynamicCollection())
@@ -362,9 +364,9 @@ void UDeformableTetrahedralComponent::UpdateFromSimulation(const FDataMapValue*
 			const TManagedArray<FVector3f>& SimulationVertex = FleshBuffer->Dynamic.GetAttribute<FVector3f>("Vertex", FGeometryCollection::VerticesGroup);
 
 			TArray<FTransform3f> Transforms;
-			const TManagedArray<int32>& Parent = GetRestCollection()->GetCollection()->GetAttribute<int32>(FTransformCollection::ParentAttribute, FTransformCollection::TransformGroup);
-			const TManagedArray<FTransform3f>& Transform = GetRestCollection()->GetCollection()->GetAttribute<FTransform3f>(FTransformCollection::TransformAttribute, FTransformCollection::TransformGroup);
-			const TManagedArray<int32>& BoneIndex = GetRestCollection()->GetCollection()->GetAttribute<int32>("BoneMap", FGeometryCollection::VerticesGroup);
+			const TManagedArray<int32>& Parent = GetRestCollection()->GetFleshCollection()->GetAttribute<int32>(FTransformCollection::ParentAttribute, FTransformCollection::TransformGroup);
+			const TManagedArray<FTransform3f>& Transform = GetRestCollection()->GetFleshCollection()->GetAttribute<FTransform3f>(FTransformCollection::TransformAttribute, FTransformCollection::TransformGroup);
+			const TManagedArray<int32>& BoneIndex = GetRestCollection()->GetFleshCollection()->GetAttribute<int32>("BoneMap", FGeometryCollection::VerticesGroup);
 			GeometryCollectionAlgo::GlobalMatrices(Transform, Parent, Transforms);
 			// Simulator produces results in component space.
 			for (int i = DynamicVertex.Num() - 1; i >= 0; i--)
@@ -417,7 +419,7 @@ void UDeformableTetrahedralComponent::UpdateLocalBounds()
 				Collection->UpdateBoundingBox();
 			}
 		}
-		BoundingBox = RestCollection->GetCollection()->GetBoundingBox();
+		BoundingBox = RestCollection->GetFleshCollection()->GetBoundingBox();
 		bBoundsNeedsUpdate = false;
 	}
 
@@ -436,17 +438,13 @@ void UDeformableTetrahedralComponent::ResetDynamicCollection()
 		{
 			DynamicCollection = NewObject<UFleshDynamicAsset>(this, TEXT("Flesh Dynamic Asset"));
 		}
+		GetDynamicCollection()->Reset(RestAsset->GetFleshCollection().Get()); //always reset from a new copy to avoid shrinking group size
 
 		if (!GetDynamicCollection()->GetCollection() || 
 			!GetDynamicCollection()->GetCollection()->NumElements(FGeometryCollection::VerticesGroup))
 		{
-			GetDynamicCollection()->Reset(RestAsset->GetCollection());
 			ResetProceduralMesh();
 		}
-		else
-		{
-			GetDynamicCollection()->ResetAttributesFrom(RestAsset->GetCollection());
-		}
 	}
 }
 
@@ -514,7 +512,7 @@ void UDeformableTetrahedralComponent::RenderProceduralMesh()
 				{
 					PERF_SCOPE(STAT_ChaosDeformable_UDeformableTetrahedralComponent_RenderProceduralMesh);
 
-					if (const FFleshCollection* FleshCollection = FleshAsset->GetCollection())
+					if (TSharedPtr<const FFleshCollection> FleshCollection = FleshAsset->GetFleshCollection())
 					{
 						const Chaos::FFleshCollectionFacade Flesh(*FleshCollection);
 
@@ -701,12 +699,12 @@ TArray<FVector3f> UDeformableTetrahedralComponent::GetGeometryCachePositions(con
 	TArray<FVector3f> Positions;
 	if (RestCollection && DynamicCollection && SkeletalMesh)
 	{
-		const FFleshCollection* FleshCollection = RestCollection->GetCollection();
+		TSharedPtr<const FFleshCollection> FleshCollection = RestCollection->GetFleshCollection();
 		const TManagedArray<FVector3f>* RestVertices = RestCollection->FindPositions();
 		const TManagedArray<FVector3f>* SimulatedVertices = DynamicCollection->FindPositions();
 		if (FleshCollection && RestVertices && SimulatedVertices)
 		{
-			ChaosFlesh::BoundSurfacePositions(SkeletalMesh, FleshCollection, RestVertices, SimulatedVertices, Positions);
+			ChaosFlesh::BoundSurfacePositions(SkeletalMesh, FleshCollection.Get(), RestVertices, SimulatedVertices, Positions);
 		}
 	}
 	return Positions;
@@ -717,12 +715,12 @@ TArray<FVector3f> UDeformableTetrahedralComponent::GetGeometryCachePositions(con
 	TArray<FVector3f> Positions;
 	if (RestCollection && DynamicCollection && StaticMesh)
 	{
-		const FFleshCollection* FleshCollection = RestCollection->GetCollection();
+		TSharedPtr<const FFleshCollection> FleshCollection = RestCollection->GetFleshCollection();
 		const TManagedArray<FVector3f>* RestVertices = RestCollection->FindPositions();
 		const TManagedArray<FVector3f>* SimulatedVertices = DynamicCollection->FindPositions();
 		if (FleshCollection && RestVertices && SimulatedVertices)
 		{
-			ChaosFlesh::BoundSurfacePositions(StaticMesh, FleshCollection, RestVertices, SimulatedVertices, Positions);
+			ChaosFlesh::BoundSurfacePositions(StaticMesh, FleshCollection.Get(), RestVertices, SimulatedVertices, Positions);
 		}
 	}
 	return Positions;
@@ -762,7 +760,7 @@ TOptional<TArray<int32>> UDeformableTetrahedralComponent::GetMeshImportVertexMap
 	const TManagedArray<FVector3f>* RestVertices = RestCollection->FindPositions();
 	if (SkeletalMeshAsset && RestCollection && RestVertices)
 	{
-		ChaosFlesh::BoundSurfacePositions(SkeletalMeshAsset, RestCollection->GetCollection(), RestVertices, RestVertices, Positions);
+		ChaosFlesh::BoundSurfacePositions(SkeletalMeshAsset, RestCollection->GetFleshCollection().Get(), RestVertices, RestVertices, Positions);
 
 		int32 NumSections = MLDLOD.Sections.Num();
 		int32 VertexOffset = 0;
@@ -814,7 +812,7 @@ bool UDeformableTetrahedralComponent::IsSkeletalMeshAnimationCompatible(const US
 	{
 		if (const UFleshAsset* FleshAsset = GetRestCollection())
 		{
-			if (const FFleshCollection* Rest = FleshAsset->GetCollection())
+			if (TSharedPtr<const FFleshCollection> Rest = FleshAsset->GetFleshCollection())
 			{
 				GeometryCollection::Facades::FTransformSource TransformSource(*Rest);
 				if (TransformSource.IsValid())
@@ -1099,7 +1097,7 @@ TArray<FVector> UDeformableTetrahedralComponent::GetEmbeddedPositionsInternal(
 
 	if (const UFleshAsset* RestAsset = GetRestCollection())
 	{
-		if (const FFleshCollection* Rest = RestAsset->GetCollection())
+		if (TSharedPtr<const FFleshCollection> Rest = RestAsset->GetFleshCollection())
 		{
 			GeometryCollection::Facades::FTetrahedralSkeletalBindings TetBindings(*Rest);
 
@@ -1159,7 +1157,7 @@ TArray<FVector> UDeformableTetrahedralComponent::GetSkeletalMeshBindingPositions
 		FName SkeletalMeshName(InSkeletalMesh->GetName());
 		if (const UFleshAsset* RestAsset = GetRestCollection())
 		{
-			if (const FFleshCollection* Rest = RestAsset->GetCollection())
+			if (TSharedPtr<const FFleshCollection> Rest = RestAsset->GetFleshCollection())
 			{
 				GeometryCollection::Facades::FTetrahedralSkeletalBindings TetBindings(*Rest);
 
diff --git a/Engine/Plugins/Experimental/ChaosFlesh/Source/ChaosFleshEngine/Private/ChaosFlesh/ChaosFleshDeformerBufferManager.cpp b/Engine/Plugins/Experimental/ChaosFlesh/Source/ChaosFleshEngine/Private/ChaosFlesh/ChaosFleshDeformerBufferManager.cpp
--- a/Engine/Plugins/Experimental/ChaosFlesh/Source/ChaosFleshEngine/Private/ChaosFlesh/ChaosFleshDeformerBufferManager.cpp
+++ b/Engine/Plugins/Experimental/ChaosFlesh/Source/ChaosFleshEngine/Private/ChaosFlesh/ChaosFleshDeformerBufferManager.cpp
@@ -715,8 +715,7 @@ const FFleshCollection* FChaosFleshDeformableGPUManager::GetRestCollection() con
 		return nullptr;
 	}
 
-	const FFleshCollection* RestCollection = RestCollectionAsset->GetCollection();
-	return RestCollection; // can be null
+	return RestCollectionAsset->GetFleshCollection().Get(); // can be null
 }
 
 const TManagedArray<FVector3f>* FChaosFleshDeformableGPUManager::GetRestVertexArray() const
diff --git a/Engine/Plugins/Experimental/ChaosFlesh/Source/ChaosFleshEngine/Private/ChaosFlesh/FleshAsset.cpp b/Engine/Plugins/Experimental/ChaosFlesh/Source/ChaosFleshEngine/Private/ChaosFlesh/FleshAsset.cpp
--- a/Engine/Plugins/Experimental/ChaosFlesh/Source/ChaosFleshEngine/Private/ChaosFlesh/FleshAsset.cpp
+++ b/Engine/Plugins/Experimental/ChaosFlesh/Source/ChaosFleshEngine/Private/ChaosFlesh/FleshAsset.cpp
@@ -9,6 +9,10 @@
 #include "Dataflow/DataflowContent.h"
 #include "Engine/SkeletalMesh.h"
 #include "GeometryCollection/TransformCollection.h"
+#if WITH_EDITOR
+#include "ChaosFlesh/FleshComponent.h"
+#include "UObject/UObjectIterator.h"
+#endif
 
 #include UE_INLINE_GENERATED_CPP_BY_NAME(FleshAsset)
 
@@ -47,6 +51,11 @@ void UFleshAsset::SetCollection(FFleshCollection* InCollection)
 	Modify();
 }
 
+void UFleshAsset::SetFleshCollection(TUniquePtr<FFleshCollection>&& InCollection)
+{
+	FleshCollection = TSharedPtr<FFleshCollection, ESPMode::ThreadSafe>(InCollection.Release());
+	Modify();
+}
 
 void UFleshAsset::PostEditCallback()
 {
@@ -149,6 +158,25 @@ void UFleshAsset::PostEditChangeProperty(FPropertyChangedEvent& PropertyChangedE
 	}
 	InvalidateDataflowContents();
 }
+
+void UFleshAsset::PostEditUndo()
+{
+	PropagateTransformUpdateToComponents();
+	Super::PostEditUndo();
+}
+
+void UFleshAsset::PropagateTransformUpdateToComponents() const
+{
+	for (TObjectIterator<UFleshComponent> It(RF_ClassDefaultObject, false, EInternalObjectFlags::Garbage); It; ++It)
+	{
+		if (It->GetRestCollection() == this)
+		{
+			// make sure to reset the rest collection to make sure the internal state of the components is up to date 
+			// but we do not apply asset default to avoid overriding the existing overrides
+			It->SetRestCollection(this);
+		}
+	}
+}
 #endif //if WITH_EDITOR
 
 UDataflowFleshContent::UDataflowFleshContent() : Super()
diff --git a/Engine/Plugins/Experimental/ChaosFlesh/Source/ChaosFleshEngine/Private/DataInterfaces/DIFleshDeformer.cpp b/Engine/Plugins/Experimental/ChaosFlesh/Source/ChaosFleshEngine/Private/DataInterfaces/DIFleshDeformer.cpp
--- a/Engine/Plugins/Experimental/ChaosFlesh/Source/ChaosFleshEngine/Private/DataInterfaces/DIFleshDeformer.cpp
+++ b/Engine/Plugins/Experimental/ChaosFlesh/Source/ChaosFleshEngine/Private/DataInterfaces/DIFleshDeformer.cpp
@@ -298,7 +298,7 @@ FDIFleshDeformerProviderProxy::IsValid(FValidationData const& InValidationData)
 		return false;
 	}
 
-	const FFleshCollection* FleshCollection = RestCollectionAsset->GetCollection();
+	TSharedPtr<const FFleshCollection> FleshCollection = RestCollectionAsset->GetFleshCollection();
 	if (!FleshCollection)
 	{
 		return false;
diff --git a/Engine/Plugins/Experimental/ChaosFlesh/Source/ChaosFleshEngine/Public/ChaosFlesh/FleshAsset.h b/Engine/Plugins/Experimental/ChaosFlesh/Source/ChaosFleshEngine/Public/ChaosFlesh/FleshAsset.h
--- a/Engine/Plugins/Experimental/ChaosFlesh/Source/ChaosFleshEngine/Public/ChaosFlesh/FleshAsset.h
+++ b/Engine/Plugins/Experimental/ChaosFlesh/Source/ChaosFleshEngine/Public/ChaosFlesh/FleshAsset.h
@@ -85,14 +85,23 @@ public:
 #if WITH_EDITOR
 	/** Post edit change property */
 	virtual void PostEditChangeProperty(FPropertyChangedEvent& PropertyChangedEvent) override;
+	virtual void PostEditUndo() override;
+	/** Propogate the fact that transform have been updated to all components */
+	void PropagateTransformUpdateToComponents() const;
 #endif //if WITH_EDITOR
 
 
 	/**Editing the collection should only be through the edit object.*/
+	UE_DEPRECATED(5.7, "Use SetFleshCollection instead")
 	void SetCollection(FFleshCollection* InCollection);
+	UE_DEPRECATED(5.7, "Use GetFleshCollection instead")
 	const FFleshCollection* GetCollection() const { return FleshCollection.Get(); }
+	UE_DEPRECATED(5.7, "Access to non-const FleshCollection is now deprecated, use SetFleshCollection to modify FleshCollection.")
 	FFleshCollection* GetCollection() { return FleshCollection.Get(); }
 
+	void SetFleshCollection(TUniquePtr<FFleshCollection>&& InCollection);
+	TSharedPtr<const FFleshCollection> GetFleshCollection() const { return FleshCollection; }
+
 	TManagedArray<FVector3f>& GetPositions();
 	const TManagedArray<FVector3f>* FindPositions() const;
 
diff --git a/Engine/Plugins/Experimental/ChaosFlesh/Source/ChaosFleshNodes/Private/Dataflow/ChaosFleshFleshAssetTerminalNode.cpp b/Engine/Plugins/Experimental/ChaosFlesh/Source/ChaosFleshNodes/Private/Dataflow/ChaosFleshFleshAssetTerminalNode.cpp
--- a/Engine/Plugins/Experimental/ChaosFlesh/Source/ChaosFleshNodes/Private/Dataflow/ChaosFleshFleshAssetTerminalNode.cpp
+++ b/Engine/Plugins/Experimental/ChaosFlesh/Source/ChaosFleshNodes/Private/Dataflow/ChaosFleshFleshAssetTerminalNode.cpp
@@ -16,7 +16,11 @@ void FFleshAssetTerminalDataflowNode::SetAssetValue(TObjectPtr<UObject> Asset, U
 	if (UFleshAsset* InFleshAsset = Cast<UFleshAsset>(Asset.Get()))
 	{
 		const FManagedArrayCollection& InCollection = GetValue<FManagedArrayCollection>(Context, &Collection);
-		InFleshAsset->SetCollection(InCollection.NewCopy<FFleshCollection>());
+		TUniquePtr<FFleshCollection> NewFleshCollection(InCollection.NewCopy<FFleshCollection>());
+		InFleshAsset->SetFleshCollection(MoveTemp(NewFleshCollection));
+#if WITH_EDITOR
+		InFleshAsset->PropagateTransformUpdateToComponents();
+#endif
 	}
 }
 
diff --git a/Engine/Plugins/Experimental/ChaosFlesh/Source/ChaosFleshNodes/Private/Dataflow/ChaosFleshGetFleshAssetNode.cpp b/Engine/Plugins/Experimental/ChaosFlesh/Source/ChaosFleshNodes/Private/Dataflow/ChaosFleshGetFleshAssetNode.cpp
--- a/Engine/Plugins/Experimental/ChaosFlesh/Source/ChaosFleshNodes/Private/Dataflow/ChaosFleshGetFleshAssetNode.cpp
+++ b/Engine/Plugins/Experimental/ChaosFlesh/Source/ChaosFleshNodes/Private/Dataflow/ChaosFleshGetFleshAssetNode.cpp
@@ -22,7 +22,7 @@ void FGetFleshAssetDataflowNode::Evaluate(UE::Dataflow::FContext& Context, const
 
 		if (FleshAssetValue)
 		{
-			if (const FFleshCollection* AssetCollection = FleshAssetValue->GetCollection())
+			if (TSharedPtr<const FFleshCollection> AssetCollection = FleshAssetValue->GetFleshCollection())
 			{
 				SetValue(Context, (const FManagedArrayCollection&)(*AssetCollection), &Output);
 			}
diff --git a/Engine/Source/Runtime/Experimental/Chaos/Private/Chaos/Deformable/ChaosDeformableSolver.cpp b/Engine/Source/Runtime/Experimental/Chaos/Private/Chaos/Deformable/ChaosDeformableSolver.cpp
--- a/Engine/Source/Runtime/Experimental/Chaos/Private/Chaos/Deformable/ChaosDeformableSolver.cpp
+++ b/Engine/Source/Runtime/Experimental/Chaos/Private/Chaos/Deformable/ChaosDeformableSolver.cpp
@@ -933,15 +933,17 @@ namespace Chaos::Softs
 
 		const Chaos::FRange& Range = Proxy.GetSolverParticleRange();
 		const FManagedArrayCollection& Rest = Proxy.GetRestCollection();
-		const TManagedArray<FIntVector4>& Tetrahedron = Rest.GetAttribute<FIntVector4>("Tetrahedron", "Tetrahedral");
-		if (uint32 NumElements = Tetrahedron.Num())
+		if (const TManagedArray<FIntVector4>* Tetrahedron = Rest.FindAttributeTyped<FIntVector4>("Tetrahedron", "Tetrahedral"))
 		{
-			const Chaos::Softs::FSolverParticles& P = Evolution->Particles();
-			for (uint32 edx = 0; edx < NumElements; ++edx)
+			if (uint32 NumElements = Tetrahedron->Num())
 			{
-				auto T = ChaosTet(Tetrahedron[edx], Range.Start);
-				Chaos::FDebugDrawQueue::GetInstance().DrawDebugPoint(
-					DoubleVert(P.GetX(T[0])), FColor::Blue, false, -1.0f, 0, GDeformableDebugParams.ParticleRadius);
+				const Chaos::Softs::FSolverParticles& P = Evolution->Particles();
+				for (uint32 edx = 0; edx < NumElements; ++edx)
+				{
+					auto T = ChaosTet((*Tetrahedron)[edx], Range.Start);
+					Chaos::FDebugDrawQueue::GetInstance().DrawDebugPoint(
+						DoubleVert(P.GetX(T[0])), FColor::Blue, false, -1.0f, 0, GDeformableDebugParams.ParticleRadius);
+				}
 			}
 		}
 #endif
@@ -988,163 +990,165 @@ namespace Chaos::Softs
 			}
 		}
 
-		const TManagedArray<FIntVector4>& Tetrahedron = Rest.GetAttribute<FIntVector4>("Tetrahedron", "Tetrahedral");
-		if (uint32 NumElements = Tetrahedron.Num())
+		if (const TManagedArray<FIntVector4>* Tetrahedron = Rest.FindAttributeTyped<FIntVector4>("Tetrahedron", "Tetrahedral"))
 		{
-			// Add Tetrahedral Elements Node
-			TArray<Chaos::TVec4<int32>> Elements;
-			Elements.SetNum(NumElements);
-			for (uint32 edx = 0; edx < NumElements; ++edx)
+			if (uint32 NumElements = Tetrahedron->Num())
 			{
-				Elements[edx] = ChaosTet(Tetrahedron[edx], Range.Start);
-			}
-
-			if (Property.bUseGridBasedConstraints)
-			{
-				int32 ElementsOffset = AllElements->Num();
-				AllElements->SetNum(ElementsOffset + NumElements);
+				// Add Tetrahedral Elements Node
+				TArray<Chaos::TVec4<int32>> Elements;
+				Elements.SetNum(NumElements);
 				for (uint32 edx = 0; edx < NumElements; ++edx)
 				{
-					(*AllElements)[edx + ElementsOffset] = ChaosTet(Tetrahedron[edx], Range.Start);
+					Elements[edx] = ChaosTet((*Tetrahedron)[edx], Range.Start);
 				}
 
-			}
+				if (Property.bUseGridBasedConstraints)
+				{
+					int32 ElementsOffset = AllElements->Num();
+					AllElements->SetNum(ElementsOffset + NumElements);
+					for (uint32 edx = 0; edx < NumElements; ++edx)
+					{
+						(*AllElements)[edx + ElementsOffset] = ChaosTet((*Tetrahedron)[edx], Range.Start);
+					}
 
-			if (Rest.HasAttributes({ FManagedArrayCollection::TManagedType<FSolverReal>("Stiffness", FGeometryCollection::VerticesGroup) }))
-			{
-				uint32 NumParticles = Rest.NumElements(FGeometryCollection::VerticesGroup);
+				}
 
-				//StiffnessWithMultiplier.Init(0.f, NumParticles);
-				FSolverReal StiffnessMultiplier = 1.f;
-				FSolverReal IncompressibilityMultiplier = 1.f;
-				FSolverReal InflationMultiplier = 1.f;
+				if (Rest.HasAttributes({ FManagedArrayCollection::TManagedType<FSolverReal>("Stiffness", FGeometryCollection::VerticesGroup) }))
+				{
+					uint32 NumParticles = Rest.NumElements(FGeometryCollection::VerticesGroup);
 
-				if (const UObject* Owner = this->MObjects[Range.Start]) {
-					FFleshThreadingProxy::FFleshInputBuffer* FleshInputBuffer = nullptr;
-					if (this->CurrentInputPackage->ObjectMap.Contains(Owner))
-					{
-						FleshInputBuffer = this->CurrentInputPackage->ObjectMap[Owner]->As<FFleshThreadingProxy::FFleshInputBuffer>();
-						if (FleshInputBuffer)
+					//StiffnessWithMultiplier.Init(0.f, NumParticles);
+					FSolverReal StiffnessMultiplier = 1.f;
+					FSolverReal IncompressibilityMultiplier = 1.f;
+					FSolverReal InflationMultiplier = 1.f;
+
+					if (const UObject* Owner = this->MObjects[Range.Start]) {
+						FFleshThreadingProxy::FFleshInputBuffer* FleshInputBuffer = nullptr;
+						if (this->CurrentInputPackage->ObjectMap.Contains(Owner))
 						{
-							IncompressibilityMultiplier = FleshInputBuffer->IncompressibilityMultiplier;
-							InflationMultiplier = FleshInputBuffer->InflationMultiplier;
+							FleshInputBuffer = this->CurrentInputPackage->ObjectMap[Owner]->As<FFleshThreadingProxy::FFleshInputBuffer>();
+							if (FleshInputBuffer)
+							{
+								IncompressibilityMultiplier = FleshInputBuffer->IncompressibilityMultiplier;
+								InflationMultiplier = FleshInputBuffer->InflationMultiplier;
+							}
 						}
 					}
-				}
-				const TManagedArray<FSolverReal>* StiffnessArray = Rest.FindAttribute<FSolverReal>("Stiffness", FGeometryCollection::VerticesGroup);
-				TArray<FSolverReal> TetStiffness;
-				TetStiffness.Init(0.f, Elements.Num());
-				if (StiffnessArray)
-				{
-					for (int32 edx = 0; edx < Elements.Num(); edx++)
+					const TManagedArray<FSolverReal>* StiffnessArray = Rest.FindAttribute<FSolverReal>("Stiffness", FGeometryCollection::VerticesGroup);
+					TArray<FSolverReal> TetStiffness;
+					TetStiffness.Init(0.f, Elements.Num());
+					if (StiffnessArray)
 					{
-						TetStiffness[edx] = (StiffnessWithMultiplier[Tetrahedron[edx].X] + StiffnessWithMultiplier[Tetrahedron[edx].Y]
-							+ StiffnessWithMultiplier[Tetrahedron[edx].Z] + StiffnessWithMultiplier[Tetrahedron[edx].W]) / 4.f;
+						for (int32 edx = 0; edx < Elements.Num(); edx++)
+						{
+							TetStiffness[edx] = (StiffnessWithMultiplier[(*Tetrahedron)[edx].X] + StiffnessWithMultiplier[(*Tetrahedron)[edx].Y]
+								+ StiffnessWithMultiplier[(*Tetrahedron)[edx].Z] + StiffnessWithMultiplier[(*Tetrahedron)[edx].W]) / 4.f;
+						}
 					}
-				}
 
-				const TManagedArray<FSolverReal>* IncompressibilityArray = Rest.FindAttribute<FSolverReal>("Incompressibility", FGeometryCollection::VerticesGroup);
-				TArray<FSolverReal> TetNu, AlphaJMesh, IncompressibilityWithMultiplier, InflationWithMultiplier;
-				TetNu.Init(.3f, Elements.Num());
+					const TManagedArray<FSolverReal>* IncompressibilityArray = Rest.FindAttribute<FSolverReal>("Incompressibility", FGeometryCollection::VerticesGroup);
+					TArray<FSolverReal> TetNu, AlphaJMesh, IncompressibilityWithMultiplier, InflationWithMultiplier;
+					TetNu.Init(.3f, Elements.Num());
 
-				IncompressibilityWithMultiplier.Init(0.f, NumParticles);
-				InflationWithMultiplier.Init(0.f, NumParticles);
-				if (IncompressibilityArray)
-				{
-					for (uint32 vdx = 0; vdx < NumParticles; ++vdx)
-					{
-						IncompressibilityWithMultiplier[vdx] = (*IncompressibilityArray)[vdx] * IncompressibilityMultiplier;
-					}
-					for (int32 edx = 0; edx < Elements.Num(); edx++)
+					IncompressibilityWithMultiplier.Init(0.f, NumParticles);
+					InflationWithMultiplier.Init(0.f, NumParticles);
+					if (IncompressibilityArray)
 					{
-						TetNu[edx] = (IncompressibilityWithMultiplier[Tetrahedron[edx].X] + IncompressibilityWithMultiplier[Tetrahedron[edx].Y]
-							+ IncompressibilityWithMultiplier[Tetrahedron[edx].Z] + IncompressibilityWithMultiplier[Tetrahedron[edx].W]) / 4.f;
+						for (uint32 vdx = 0; vdx < NumParticles; ++vdx)
+						{
+							IncompressibilityWithMultiplier[vdx] = (*IncompressibilityArray)[vdx] * IncompressibilityMultiplier;
+						}
+						for (int32 edx = 0; edx < Elements.Num(); edx++)
+						{
+							TetNu[edx] = (IncompressibilityWithMultiplier[(*Tetrahedron)[edx].X] + IncompressibilityWithMultiplier[(*Tetrahedron)[edx].Y]
+								+ IncompressibilityWithMultiplier[(*Tetrahedron)[edx].Z] + IncompressibilityWithMultiplier[(*Tetrahedron)[edx].W]) / 4.f;
+						}
 					}
-				}
 
-				const TManagedArray<FSolverReal>* InflationArray = Rest.FindAttribute<FSolverReal>("Inflation", FGeometryCollection::VerticesGroup);
-				AlphaJMesh.Init(1.f, Elements.Num());
-				if (InflationArray)
-				{
-					for (uint32 vdx = 0; vdx < NumParticles; ++vdx)
-					{
-						InflationWithMultiplier[vdx] = (*InflationArray)[vdx] * InflationMultiplier;
-					}
-					for (int32 edx = 0; edx < Elements.Num(); edx++)
+					const TManagedArray<FSolverReal>* InflationArray = Rest.FindAttribute<FSolverReal>("Inflation", FGeometryCollection::VerticesGroup);
+					AlphaJMesh.Init(1.f, Elements.Num());
+					if (InflationArray)
 					{
-						AlphaJMesh[edx] = (InflationWithMultiplier[Tetrahedron[edx].X] + InflationWithMultiplier[Tetrahedron[edx].Y]
-							+ InflationWithMultiplier[Tetrahedron[edx].Z] + InflationWithMultiplier[Tetrahedron[edx].W]) / 4.f;
+						for (uint32 vdx = 0; vdx < NumParticles; ++vdx)
+						{
+							InflationWithMultiplier[vdx] = (*InflationArray)[vdx] * InflationMultiplier;
+						}
+						for (int32 edx = 0; edx < Elements.Num(); edx++)
+						{
+							AlphaJMesh[edx] = (InflationWithMultiplier[(*Tetrahedron)[edx].X] + InflationWithMultiplier[(*Tetrahedron)[edx].Y]
+								+ InflationWithMultiplier[(*Tetrahedron)[edx].Z] + InflationWithMultiplier[(*Tetrahedron)[edx].W]) / 4.f;
+						}
 					}
-				}
 
-				if (Property.bUseGaussSeidelConstraints)
-				{
-					int32 ElementsOffset = AllTetEMeshArray->Num();
-					AllTetEMeshArray->SetNum(ElementsOffset + NumElements);
-					AllTetNuMeshArray->SetNum(ElementsOffset + NumElements);
-					AllTetAlphaJArray->SetNum(ElementsOffset + NumElements);
-						
-					for (uint32 edx = 0; edx < NumElements; ++edx)
+					if (Property.bUseGaussSeidelConstraints)
 					{
-						(*AllTetEMeshArray)[edx + ElementsOffset] = TetStiffness[edx];
-						(*AllTetNuMeshArray)[edx + ElementsOffset] = TetNu[edx];
-						(*AllTetAlphaJArray)[edx + ElementsOffset] = AlphaJMesh[edx];
-					}
-				}
-
+						int32 ElementsOffset = AllTetEMeshArray->Num();
+						AllTetEMeshArray->SetNum(ElementsOffset + NumElements);
+						AllTetNuMeshArray->SetNum(ElementsOffset + NumElements);
+						AllTetAlphaJArray->SetNum(ElementsOffset + NumElements);
 
-				if (Property.bEnableCorotatedConstraints)
-				{
+						for (uint32 edx = 0; edx < NumElements; ++edx)
+						{
+							(*AllTetEMeshArray)[edx + ElementsOffset] = TetStiffness[edx];
+							(*AllTetNuMeshArray)[edx + ElementsOffset] = TetNu[edx];
+							(*AllTetAlphaJArray)[edx + ElementsOffset] = AlphaJMesh[edx];
+						}
+					}
 
-					int32 InitIndex = Evolution->AddConstraintInitRange(1, true);
-					int32 ConstraintIndex = Evolution->AddConstraintRuleRange(1, true);
 
-					if (Property.bDoBlended)
+					if (Property.bEnableCorotatedConstraints)
 					{
-						FBlendedXPBDCorotatedConstraints<FSolverReal, FSolverParticles>* BlendedCorotatedConstraint =
-							new FBlendedXPBDCorotatedConstraints<FSolverReal, FSolverParticles>(
-								Evolution->Particles(), Elements, TetStiffness, (FSolverReal).3,/*bRecordMetric = */false, Property.BlendedZeta);
 
-						Evolution->ConstraintInits()[InitIndex] =
-							[BlendedCorotatedConstraint](FSolverParticles& InParticles, const FSolverReal Dt)
-						{
-							BlendedCorotatedConstraint->Init();
-						};
+						int32 InitIndex = Evolution->AddConstraintInitRange(1, true);
+						int32 ConstraintIndex = Evolution->AddConstraintRuleRange(1, true);
 
-						Evolution->ConstraintRules()[ConstraintIndex] =
-							[BlendedCorotatedConstraint](FSolverParticles& InParticles, const FSolverReal Dt)
+						if (Property.bDoBlended)
 						{
-							BlendedCorotatedConstraint->ApplyInParallel(InParticles, Dt);
-						};
+							FBlendedXPBDCorotatedConstraints<FSolverReal, FSolverParticles>* BlendedCorotatedConstraint =
+								new FBlendedXPBDCorotatedConstraints<FSolverReal, FSolverParticles>(
+									Evolution->Particles(), Elements, TetStiffness, (FSolverReal).3,/*bRecordMetric = */false, Property.BlendedZeta);
 
-						BlendedCorotatedConstraints.Add(TUniquePtr<FBlendedXPBDCorotatedConstraints<FSolverReal, FSolverParticles>>(BlendedCorotatedConstraint));
+							Evolution->ConstraintInits()[InitIndex] =
+								[BlendedCorotatedConstraint](FSolverParticles& InParticles, const FSolverReal Dt)
+								{
+									BlendedCorotatedConstraint->Init();
+								};
 
-					}
-					else
-					{
-						FXPBDCorotatedConstraints<FSolverReal, FSolverParticles>* CorotatedConstraint =
-							new FXPBDCorotatedConstraints<FSolverReal, FSolverParticles>(
-								Evolution->Particles(), Elements, TetStiffness, TetNu, MoveTemp(AlphaJMesh), GDeformableXPBDCorotatedParams);
+							Evolution->ConstraintRules()[ConstraintIndex] =
+								[BlendedCorotatedConstraint](FSolverParticles& InParticles, const FSolverReal Dt)
+								{
+									BlendedCorotatedConstraint->ApplyInParallel(InParticles, Dt);
+								};
 
-						Evolution->ConstraintInits()[InitIndex] =
-							[CorotatedConstraint](FSolverParticles& InParticles, const FSolverReal Dt)
-						{
-							CorotatedConstraint->Init();
-						};
+							BlendedCorotatedConstraints.Add(TUniquePtr<FBlendedXPBDCorotatedConstraints<FSolverReal, FSolverParticles>>(BlendedCorotatedConstraint));
 
-						Evolution->ConstraintRules()[ConstraintIndex] =
-							[CorotatedConstraint](FSolverParticles& InParticles, const FSolverReal Dt)
+						}
+						else
 						{
-							CorotatedConstraint->ApplyInParallel(InParticles, Dt);
-						};
+							FXPBDCorotatedConstraints<FSolverReal, FSolverParticles>* CorotatedConstraint =
+								new FXPBDCorotatedConstraints<FSolverReal, FSolverParticles>(
+									Evolution->Particles(), Elements, TetStiffness, TetNu, MoveTemp(AlphaJMesh), GDeformableXPBDCorotatedParams);
+
+							Evolution->ConstraintInits()[InitIndex] =
+								[CorotatedConstraint](FSolverParticles& InParticles, const FSolverReal Dt)
+								{
+									CorotatedConstraint->Init();
+								};
+
+							Evolution->ConstraintRules()[ConstraintIndex] =
+								[CorotatedConstraint](FSolverParticles& InParticles, const FSolverReal Dt)
+								{
+									CorotatedConstraint->ApplyInParallel(InParticles, Dt);
+								};
 
-						CorotatedConstraints.Add(TUniquePtr<FXPBDCorotatedConstraints<FSolverReal, FSolverParticles>>(CorotatedConstraint));
+							CorotatedConstraints.Add(TUniquePtr<FXPBDCorotatedConstraints<FSolverReal, FSolverParticles>>(CorotatedConstraint));
+						}
 					}
+
+
 				}
-			
-				
-			}
 
+			}
 		}
 
 		Chaos::FRange AllUnconstrainedSurfaceElementsSkinRange(INDEX_NONE, INDEX_NONE), AllUnconstrainedSurfaceElementsCorotatedCodRange(INDEX_NONE, INDEX_NONE);
@@ -1276,19 +1280,19 @@ namespace Chaos::Softs
 		if (Property.bUseGridBasedConstraints)
 		{
 			auto ChaosTet = [](FIntVector4 V, int32 dp) { return Chaos::TVec4<int32>(dp + V.X, dp + V.Y, dp + V.Z, dp + V.W); };
-
 			const FManagedArrayCollection& Rest = Proxy.GetRestCollection();
-			const TManagedArray<FIntVector4>& Tetrahedron = Rest.GetAttribute<FIntVector4>("Tetrahedron", "Tetrahedral");
-
-			if (uint32 NumElements = Tetrahedron.Num())
+			if (const TManagedArray<FIntVector4>* Tetrahedron = Rest.FindAttributeTyped<FIntVector4>("Tetrahedron", "Tetrahedral"))
 			{
-				const Chaos::FRange& Range = Proxy.GetSolverParticleRange();
-
-				int32 ElementsOffset = AllElements->Num();
-				AllElements->SetNum(ElementsOffset + NumElements);
-				for (uint32 edx = 0; edx < NumElements; ++edx)
+				if (uint32 NumElements = Tetrahedron->Num())
 				{
-					(*AllElements)[edx + ElementsOffset] = ChaosTet(Tetrahedron[edx], Range.Start);
+					const Chaos::FRange& Range = Proxy.GetSolverParticleRange();
+
+					int32 ElementsOffset = AllElements->Num();
+					AllElements->SetNum(ElementsOffset + NumElements);
+					for (uint32 edx = 0; edx < NumElements; ++edx)
+					{
+						(*AllElements)[edx + ElementsOffset] = ChaosTet((*Tetrahedron)[edx], Range.Start);
+					}
 				}
 			}
 		}
@@ -1301,37 +1305,36 @@ namespace Chaos::Softs
 		if (Property.bUseGaussSeidelConstraints)
 		{
 			auto ChaosTet = [](FIntVector4 V, int32 dp) { return Chaos::TVec4<int32>(dp + V.X, dp + V.Y, dp + V.Z, dp + V.W); };
-
 			const FManagedArrayCollection& Rest = Proxy.GetRestCollection();
-			const TManagedArray<FIntVector4>& Tetrahedron = Rest.GetAttribute<FIntVector4>("Tetrahedron", "Tetrahedral");
-			const TManagedArray<TArray<int32>>& Tetrahedron1 = Rest.GetAttribute<TArray<int32>>("IncidentElements", "Vertices");
-			const TManagedArray<TArray<int32>>& Tetrahedron2 = Rest.GetAttribute<TArray<int32>>("IncidentElementsLocalIndex", "Vertices");
+			const TManagedArray<FIntVector4>* Tetrahedron = Rest.FindAttributeTyped<FIntVector4>("Tetrahedron", "Tetrahedral");
 			const TManagedArray<TArray<int32>>* IncidentElementsPointer = Rest.FindAttribute<TArray<int32>>("IncidentElements", "Vertices");
 			const TManagedArray<TArray<int32>>* IncidentElementsLocalPointer = Rest.FindAttribute<TArray<int32>>("IncidentElementsLocalIndex", "Vertices");
-
-			if (uint32 NumElements = Tetrahedron.Num())
+			if (Tetrahedron && IncidentElementsPointer && IncidentElementsLocalPointer)
 			{
-				const Chaos::FRange& Range = Proxy.GetSolverParticleRange();
-
-				int32 ElementsOffset = AllElements->Num();
-				AllElements->SetNum(ElementsOffset + NumElements);
-				for (uint32 edx = 0; edx < NumElements; ++edx)
-				{
-					(*AllElements)[edx + ElementsOffset] = ChaosTet(Tetrahedron[edx], Range.Start);
-				}
-				if (uint32 NumIncident = IncidentElementsPointer->Num())
+				if (uint32 NumElements = Tetrahedron->Num())
 				{
-					int32 IncidentOffSet = Range.Start;
-					AllIncidentElements->SetNum(IncidentOffSet + NumIncident); 
-					AllIncidentElementsLocal->SetNum(IncidentOffSet + NumIncident);
-					for (uint32 i = 0; i < NumIncident; ++i)
+					const Chaos::FRange& Range = Proxy.GetSolverParticleRange();
+
+					int32 ElementsOffset = AllElements->Num();
+					AllElements->SetNum(ElementsOffset + NumElements);
+					for (uint32 edx = 0; edx < NumElements; ++edx)
 					{
-						(*AllIncidentElements)[i + IncidentOffSet] = (*IncidentElementsPointer)[i];
-						for (int32 j = 0; j < (*AllIncidentElements)[i + IncidentOffSet].Num(); j++)
+						(*AllElements)[edx + ElementsOffset] = ChaosTet((*Tetrahedron)[edx], Range.Start);
+					}
+					if (uint32 NumIncident = IncidentElementsPointer->Num())
+					{
+						int32 IncidentOffSet = Range.Start;
+						AllIncidentElements->SetNum(IncidentOffSet + NumIncident);
+						AllIncidentElementsLocal->SetNum(IncidentOffSet + NumIncident);
+						for (uint32 i = 0; i < NumIncident; ++i)
 						{
-							(*AllIncidentElements)[i + IncidentOffSet][j] += ElementsOffset;
+							(*AllIncidentElements)[i + IncidentOffSet] = (*IncidentElementsPointer)[i];
+							for (int32 j = 0; j < (*AllIncidentElements)[i + IncidentOffSet].Num(); j++)
+							{
+								(*AllIncidentElements)[i + IncidentOffSet][j] += ElementsOffset;
+							}
+							(*AllIncidentElementsLocal)[i + IncidentOffSet] = (*IncidentElementsLocalPointer)[i];
 						}
-						(*AllIncidentElementsLocal)[i + IncidentOffSet] = (*IncidentElementsLocalPointer)[i];
 					}
 				}
 			}
@@ -1468,72 +1471,71 @@ namespace Chaos::Softs
 		SurfaceElements->SetNum(0);
 		TetmeshSurfaceElements->SetNum(0);
 		ParticleComponentIndex->SetNum(0);
+		SurfaceCollisionVertices->SetNum(0);
 		int ComponentOffset = 0;
 
 		for (FThreadingProxy* InProxy : UninitializedProxys_Internal)
 		{
 			if (FFleshThreadingProxy* Proxy = InProxy->As<FFleshThreadingProxy>())
 			{
-				if (const FManagedArrayCollection* Rest = &Proxy->GetRestCollection())
+				const FManagedArrayCollection& Rest = Proxy->GetRestCollection();
+				if (const TManagedArray<FVector3f>* Vertex = Rest.FindAttribute<FVector3f>("Vertex", FGeometryCollection::VerticesGroup))
 				{
-					if (const TManagedArray<FVector3f>* Vertex = Rest->FindAttribute<FVector3f>("Vertex", FGeometryCollection::VerticesGroup))
+					if (const TManagedArray<FIntVector>* Indices = Rest.FindAttribute<FIntVector>("Indices", FGeometryCollection::FacesGroup))
 					{
-						if (const TManagedArray<FIntVector>* Indices = Rest->FindAttribute<FIntVector>("Indices", FGeometryCollection::FacesGroup))
+						int32 SurfaceOffset = SurfaceElements->Num();
+						SurfaceElements->SetNum(SurfaceOffset + Indices->Num());
+						TArray<int32> SurfaceVertices;
+						for (int32 i = 0; i < Indices->Num(); i++)
 						{
-							int32 SurfaceOffset = SurfaceElements->Num();
-							SurfaceElements->SetNum(SurfaceOffset + Indices->Num());
-							TArray<int32> SurfaceVertices;
-							for (int32 i = 0; i < Indices->Num(); i++)
+							for (int32 j = 0; j < 3; j++)
 							{
-								for (int32 j = 0; j < 3; j++)
-								{
-									(*SurfaceElements)[i + SurfaceOffset][j] = VertexOffset + (*Indices)[i][j];
-									SurfaceVertices.Add(VertexOffset + (*Indices)[i][j]);
-								}
+								(*SurfaceElements)[i + SurfaceOffset][j] = VertexOffset + (*Indices)[i][j];
+								SurfaceVertices.Add(VertexOffset + (*Indices)[i][j]);
 							}
-							TSet<int32> UniqueSurfaceVertices(SurfaceVertices);
-							GeometryCollection::Facades::FCollisionFacade CollisionFacade(*Rest);
-							if (CollisionFacade.IsValid()) //Collision vertices are specified
+						}
+						TSet<int32> UniqueSurfaceVertices(SurfaceVertices);
+						GeometryCollection::Facades::FCollisionFacade CollisionFacade(Rest);
+						if (CollisionFacade.IsValid()) //Collision vertices are specified
+						{
+							for (int32 SurfaceVertexIdx : UniqueSurfaceVertices)
 							{
-								for (int32 SurfaceVertexIdx: UniqueSurfaceVertices)
+								if (CollisionFacade.IsCollisionEnabled(SurfaceVertexIdx - VertexOffset))
 								{
-									if (CollisionFacade.IsCollisionEnabled(SurfaceVertexIdx - VertexOffset))
-									{
-										SurfaceCollisionVertices->Add(SurfaceVertexIdx);
-									}
+									SurfaceCollisionVertices->Add(SurfaceVertexIdx);
 								}
 							}
-							else
-							{
-								SurfaceCollisionVertices->Append(UniqueSurfaceVertices.Array());
-							}
 						}
-
-						VertexOffset += Vertex->Num();
-						
-						if (!Property.bDoInComponentSpringCollision || Property.bDoSphereRepulsion) //Component-Component collisions or sphere repulsion
+						else
 						{
-							int32 Offset = ParticleComponentIndex->Num();
-							ParticleComponentIndex->SetNum(ParticleComponentIndex->Num() + Vertex->Num());
-							for (int32 i = 0; i < Vertex->Num(); i++) {
-								(*ParticleComponentIndex)[i + Offset] = ComponentOffset;
+							SurfaceCollisionVertices->Append(UniqueSurfaceVertices.Array());
+						}
+					}
+
+					VertexOffset += Vertex->Num();
+
+					if (!Property.bDoInComponentSpringCollision || Property.bDoSphereRepulsion) //Component-Component collisions or sphere repulsion
+					{
+						int32 Offset = ParticleComponentIndex->Num();
+						ParticleComponentIndex->SetNum(ParticleComponentIndex->Num() + Vertex->Num());
+						for (int32 i = 0; i < Vertex->Num(); i++) {
+							(*ParticleComponentIndex)[i + Offset] = ComponentOffset;
+						}
+						int NewComponentOffset = ComponentOffset;
+						GeometryCollection::Facades::FCollectionMeshFacade MeshFacade(Rest);
+						TArray<int32> ComponentIndex = MeshFacade.GetGeometryGroupIndexArray();
+						for (int32 i = 0; i < ComponentIndex.Num(); i++) {
+							if (ComponentIndex[i] < 0)
+							{
+								(*ParticleComponentIndex)[i + Offset] = ComponentIndex[i]; //Isolated Nodes
 							}
-							int NewComponentOffset = ComponentOffset;
-							GeometryCollection::Facades::FCollectionMeshFacade MeshFacade(*Rest);
-							TArray<int32> ComponentIndex = MeshFacade.GetGeometryGroupIndexArray();
-							for (int32 i = 0; i < ComponentIndex.Num(); i++) {
-								if (ComponentIndex[i] < 0)
-								{
-									(*ParticleComponentIndex)[i + Offset] = ComponentIndex[i]; //Isolated Nodes
-								}
-								else
-								{
-									(*ParticleComponentIndex)[i + Offset] = ComponentOffset + ComponentIndex[i];
-									NewComponentOffset = NewComponentOffset < (*ParticleComponentIndex)[i + Offset] ? (*ParticleComponentIndex)[i + Offset] : NewComponentOffset;
-								}
+							else
+							{
+								(*ParticleComponentIndex)[i + Offset] = ComponentOffset + ComponentIndex[i];
+								NewComponentOffset = NewComponentOffset < (*ParticleComponentIndex)[i + Offset] ? (*ParticleComponentIndex)[i + Offset] : NewComponentOffset;
 							}
-							ComponentOffset = NewComponentOffset + 1;
 						}
+						ComponentOffset = NewComponentOffset + 1;
 					}
 				}
 			}
@@ -1547,60 +1549,58 @@ namespace Chaos::Softs
 		{
 			if (FFleshThreadingProxy* Proxy = InProxy->As<FFleshThreadingProxy>())
 			{
-				if (const FManagedArrayCollection* Rest = &Proxy->GetRestCollection())
+				const FManagedArrayCollection& Rest = Proxy->GetRestCollection();
+				if (const TManagedArray<FVector3f>* Vertex = Rest.FindAttribute<FVector3f>("Vertex", FGeometryCollection::VerticesGroup))
 				{
-					if (const TManagedArray<FVector3f>* Vertex = Rest->FindAttribute<FVector3f>("Vertex", FGeometryCollection::VerticesGroup))
+					if (const TManagedArray<FIntVector>* Indices = Rest.FindAttribute<FIntVector>("Indices", FGeometryCollection::FacesGroup))
 					{
-						if (const TManagedArray<FIntVector>* Indices = Rest->FindAttribute<FIntVector>("Indices", FGeometryCollection::FacesGroup))
+						if (const TManagedArray<int32>* TriangleMeshIndices = Rest.FindAttribute<int32>("ObjectIndices", "TriangleMesh"))
 						{
-							if (const TManagedArray<int32>* TriangleMeshIndices = Rest->FindAttribute<int32>("ObjectIndices", "TriangleMesh"))
+							if (const TManagedArray<int32>* FaceStarts = Rest.FindAttribute<int32>("FaceStart", FGeometryCollection::GeometryGroup))
 							{
-								if (const TManagedArray<int32>* FaceStarts = Rest->FindAttribute<int32>("FaceStart", FGeometryCollection::GeometryGroup))
+								if (const TManagedArray<int32>* FaceCounts = Rest.FindAttribute<int32>("FaceCount", FGeometryCollection::GeometryGroup))
 								{
-									if (const TManagedArray<int32>* FaceCounts = Rest->FindAttribute<int32>("FaceCount", FGeometryCollection::GeometryGroup))
+									TSet<int32> TriMeshObjects;
+									for (const int32 ObjectIndex : (*TriangleMeshIndices))
 									{
-										TSet<int32> TriMeshObjects;
-										for (const int32 ObjectIndex: (*TriangleMeshIndices))
-										{
-											TriMeshObjects.Add(ObjectIndex);
-										}
-										for (int32 i = 0; i < FaceStarts->Num(); i++) 
+										TriMeshObjects.Add(ObjectIndex);
+									}
+									for (int32 i = 0; i < FaceStarts->Num(); i++)
+									{
+										if (!TriMeshObjects.Contains(i))
 										{
-											if (!TriMeshObjects.Contains(i))
-											{
-												const int32 FaceStartIndex = (*FaceStarts)[i];
-												const int32 FaceNum = (*FaceCounts)[i];
+											const int32 FaceStartIndex = (*FaceStarts)[i];
+											const int32 FaceNum = (*FaceCounts)[i];
 
-												int32 SurfaceOffset = TetmeshSurfaceElements->Num();
-												TetmeshSurfaceElements->SetNum(SurfaceOffset + FaceNum);
-												for (int32 e = FaceStartIndex; e < FaceStartIndex + FaceNum; e++)
+											int32 SurfaceOffset = TetmeshSurfaceElements->Num();
+											TetmeshSurfaceElements->SetNum(SurfaceOffset + FaceNum);
+											for (int32 e = FaceStartIndex; e < FaceStartIndex + FaceNum; e++)
+											{
+												for (int32 j = 0; j < 3; j++)
 												{
-													for (int32 j = 0; j < 3; j++)
-													{
-														(*TetmeshSurfaceElements)[e - FaceStartIndex + SurfaceOffset][j] = (*Indices)[e][j];
-													}
+													(*TetmeshSurfaceElements)[e - FaceStartIndex + SurfaceOffset][j] = (*Indices)[e][j];
 												}
 											}
 										}
 									}
 								}
-								
 							}
-							else
+
+						}
+						else
+						{
+							int32 SurfaceOffset = TetmeshSurfaceElements->Num();
+							TetmeshSurfaceElements->SetNum(SurfaceOffset + Indices->Num());
+							for (int32 i = 0; i < Indices->Num(); i++)
 							{
-								int32 SurfaceOffset = TetmeshSurfaceElements->Num();
-								TetmeshSurfaceElements->SetNum(SurfaceOffset + Indices->Num());
-								for (int32 i = 0; i < Indices->Num(); i++)
+								for (int32 j = 0; j < 3; j++)
 								{
-									for (int32 j = 0; j < 3; j++)
-									{
-										(*TetmeshSurfaceElements)[i + SurfaceOffset][j] = VertexOffset + (*Indices)[i][j];
-									}
+									(*TetmeshSurfaceElements)[i + SurfaceOffset][j] = VertexOffset + (*Indices)[i][j];
 								}
 							}
 						}
-						VertexOffset += Vertex->Num();
 					}
+					VertexOffset += Vertex->Num();
 				}
 			}
 		}
@@ -2028,6 +2028,26 @@ namespace Chaos::Softs
 		}
 	}
 
+	void FDeformableSolver::ResetGaussSeidelVariables()
+	{
+		AllElements.Reset(new TArray<Chaos::TVec4<int32>>());
+		AllIncidentElements.Reset(new TArray<TArray<int32>>());
+		AllIncidentElementsLocal.Reset(new TArray<TArray<int32>>());
+		AllTetEMeshArray.Reset(new TArray<FSolverReal>());
+		AllTetNuMeshArray.Reset(new TArray<FSolverReal>());
+		AllTetAlphaJArray.Reset(new TArray<FSolverReal>());
+		AllCorotatedCodEMeshArray.Reset(new TArray<FSolverReal>());
+		AllSkinEMeshArray.Reset(new TArray<FSolverReal>());	
+		MuscleIndexOffset.Empty();
+		AllUnconstrainedSurfaceElementsCorotatedCod.Reset(new TArray<Chaos::TVec3<int32>>());
+		AllUnconstrainedSurfaceElementsSkin.Reset(new TArray<Chaos::TVec3<int32>>());
+		GSWeakConstraints.Reset(new FGaussSeidelWeakConstraints<FSolverReal, FSolverParticles>({}, {}, {}, {}, {}, GDeformableXPBDWeakConstraintParams));
+		GSDynamicWeakConstraints.Reset(new FGaussSeidelDynamicWeakConstraints<FSolverReal, FSolverParticles>(GDeformableXPBDWeakConstraintParams));
+		GSSphereRepulsionConstraints.Reset(new FGaussSeidelSphereRepulsionConstraints<FSolverReal, FSolverParticles>(Property.SphereRepulsionRadius, Property.SphereRepulsionStiffness, Evolution->Particles(), GDeformableXPBDWeakConstraintParams));
+		GSVolumeConstraints.Reset(new FGaussSeidelUnilateralTetConstraints<Softs::FSolverReal, Softs::FSolverParticles>(Evolution->Particles(), {}, {}));
+		MuscleActivationConstraints.Reset(new FMuscleActivationConstraints<FSolverReal, FSolverParticles>());
+	}
+
 	void FDeformableSolver::RemoveSimulationObjects()
 	{
 		PERF_SCOPE(STAT_ChaosDeformableSolver_RemoveSimulationObjects);
@@ -2097,11 +2117,14 @@ namespace Chaos::Softs
 				}
 			}
 
+			ResetGaussSeidelVariables();
+
 			// regenerate all constraints
 			for (TPair< FThreadingProxy::FKey, TUniquePtr<FThreadingProxy> >& BaseProxyPair : Proxies)
 			{
 				if (FFleshThreadingProxy* Proxy = BaseProxyPair.Value->As<FFleshThreadingProxy>())
 				{
+					InitializeWeakConstraint(*Proxy);
 					InitializeMuscleActivation(*Proxy);
 					InitializeTetrahedralOrTriangleConstraint(*Proxy);
 					InitializeGridBasedConstraints(*Proxy);
@@ -2529,10 +2552,9 @@ namespace Chaos::Softs
 		TRACE_CPUPROFILER_EVENT_SCOPE(ChaosDeformableSolver_WriteFrame);
 		if (FFleshThreadingProxy* Proxy = InProxy.As<FFleshThreadingProxy>())
 		{
-			if (const FManagedArrayCollection* Rest = &Proxy->GetRestCollection())
+			const FManagedArrayCollection& Rest = Proxy->GetRestCollection();
+			if (SurfaceElements)
 			{
-				const TManagedArray<FIntVector>& Indices = Rest->GetAttribute<FIntVector>("Indices", FGeometryCollection::FacesGroup);
-
 				WriteTrisGEO(Evolution->Particles(), *SurfaceElements);
 				FString file = FPaths::ProjectDir();
 				file.Append(TEXT("/DebugOutput/DtLog.txt"));
diff --git a/Engine/Source/Runtime/Experimental/Chaos/Public/Chaos/Deformable/ChaosDeformableSolver.h b/Engine/Source/Runtime/Experimental/Chaos/Public/Chaos/Deformable/ChaosDeformableSolver.h
--- a/Engine/Source/Runtime/Experimental/Chaos/Public/Chaos/Deformable/ChaosDeformableSolver.h
+++ b/Engine/Source/Runtime/Experimental/Chaos/Public/Chaos/Deformable/ChaosDeformableSolver.h
@@ -172,6 +172,7 @@ namespace Chaos::Softs
 		CHAOS_API void PushRestartPackage(int32 InFrame, FDeformableDataMap&& InPackage);
 
 		const FDeformableSolverProperties& GetProperties() const { return Property; }
+		void ResetGaussSeidelVariables();
 
 	private:
 
diff --git a/Engine/Source/Runtime/Experimental/Chaos/Public/Chaos/Deformable/ChaosDeformableSolverProxy.h b/Engine/Source/Runtime/Experimental/Chaos/Public/Chaos/Deformable/ChaosDeformableSolverProxy.h
--- a/Engine/Source/Runtime/Experimental/Chaos/Public/Chaos/Deformable/ChaosDeformableSolverProxy.h
+++ b/Engine/Source/Runtime/Experimental/Chaos/Public/Chaos/Deformable/ChaosDeformableSolverProxy.h
@@ -84,7 +84,7 @@ namespace Chaos::Softs
 		FTransform PrevComponentToBoneXf = FTransform::Identity; // bone xf
 		FTransform PrevBoneToComponentXf = FTransform::Identity; // inv bone xf
 
-		const FManagedArrayCollection& Rest;
+		TSharedRef<const FManagedArrayCollection> Rest;
 		FManagedArrayCollection Dynamic;
 		Chaos::FRange SolverParticleRange = Chaos::FRange(0, 0);
 		ChaosDeformableSimSpace SimSpace = ChaosDeformableSimSpace::World;
@@ -97,7 +97,7 @@ namespace Chaos::Softs
 			const FTransform& InWorldToComponentXf, // ComponentXf
 			const FTransform& InComponentToBoneXf,
 			const ChaosDeformableSimSpace InSimSpace,
-			const FManagedArrayCollection& InRest,
+			TSharedRef<const FManagedArrayCollection> InRest,
 			const FManagedArrayCollection& InDynamic)
 			: Super(InOwner, TypeName())
 			, WorldToComponentXf(InWorldToComponentXf)
@@ -118,7 +118,7 @@ namespace Chaos::Softs
 
 		bool CanSimulate() const
 		{
-			return Rest.NumElements(FGeometryCollection::VerticesGroup) > 0;
+			return Rest->NumElements(FGeometryCollection::VerticesGroup) > 0;
 		}
 
 		bool IsBoneSpace() const
@@ -245,7 +245,7 @@ namespace Chaos::Softs
 
 		FManagedArrayCollection& GetDynamicCollection()				{ return Dynamic; }
 		const FManagedArrayCollection& GetDynamicCollection() const { return Dynamic; }
-		const FManagedArrayCollection& GetRestCollection() const	{ return Rest; }
+		const FManagedArrayCollection& GetRestCollection() const	{ return *Rest; }
 
 		class FFleshInputBuffer : public FThreadingProxy::FBuffer
 		{
