Use correct host group name length when computing cache key for generic HTTP requests.

#rnx

diff --git a/Engine/Source/Runtime/Experimental/IoStore/OnDemand/Private/OnDemandHttpIoDispacher.cpp b/Engine/Source/Runtime/Experimental/IoStore/OnDemand/Private/OnDemandHttpIoDispacher.cpp
--- a/Engine/Source/Runtime/Experimental/IoStore/OnDemand/Private/OnDemandHttpIoDispacher.cpp
+++ b/Engine/Source/Runtime/Experimental/IoStore/OnDemand/Private/OnDemandHttpIoDispacher.cpp
@@ -84,8 +84,8 @@ static FIoHash GetCacheKey(const FName& HostGroup, const FIoRelativeUrl Url, con
 	FIoHashBuilder HashBuilder;
 
 	TCHAR NameString[NAME_SIZE];
-	HostGroup.GetPlainNameString(NameString);
-	HashBuilder.Update(NameString, NAME_SIZE);
+	const uint32 NameLength = HostGroup.GetPlainNameString(NameString);
+	HashBuilder.Update(NameString, NameLength);
 	HashBuilder.Update(Url.ToString(), Url.Len());
 	
 	TConstArrayView<FAnsiString> HeadersView = Headers.ToArrayView();
