Rename cvar 'iax.EnableHttpCancellation' to 'iax.HttpEnableInflightCancellation' to match existing naming conventions.

#rnx

diff --git a/Engine/Source/Runtime/Experimental/IoStore/OnDemand/Private/OnDemandHttpThread.cpp b/Engine/Source/Runtime/Experimental/IoStore/OnDemand/Private/OnDemandHttpThread.cpp
--- a/Engine/Source/Runtime/Experimental/IoStore/OnDemand/Private/OnDemandHttpThread.cpp
+++ b/Engine/Source/Runtime/Experimental/IoStore/OnDemand/Private/OnDemandHttpThread.cpp
@@ -177,10 +177,10 @@ static FAutoConsoleVariableRef CVar_GIasHttpPollTimeoutMs(
 	TEXT("Http tick poll timeout in milliseconds")
 );
 
-bool GIaxEnableHttpCancellation = true;
-static FAutoConsoleVariableRef CVar_IaxEnableHttpCancellation(
-	TEXT("iax.EnableHttpCancellation"),
-	GIaxEnableHttpCancellation,
+bool GIaxHttpEnableInflightCancellation = true;
+static FAutoConsoleVariableRef CVar_IaxHttpEnableInflightCancellation(
+	TEXT("iax.HttpEnableInflightCancellation"),
+	GIaxHttpEnableInflightCancellation,
 	TEXT("When enabled the system will attempt to cancel inflight http requests")
 );
 
@@ -374,7 +374,7 @@ void FOnDemandHttpThread::CancelRequest(FRequestHandle RequestHandle)
 		
 		Request->bCancelled = true;
 
-		if (GIaxEnableHttpCancellation)
+		if (GIaxHttpEnableInflightCancellation)
 		{
 			HttpClient->CancelRequest(Request->HttpTicketId);
 		}
