Add additional logging when an inflight IAX http request is canceled.

#rnx

diff --git a/Engine/Source/Runtime/Experimental/IoStore/OnDemand/Private/OnDemandHttpClient.cpp b/Engine/Source/Runtime/Experimental/IoStore/OnDemand/Private/OnDemandHttpClient.cpp
--- a/Engine/Source/Runtime/Experimental/IoStore/OnDemand/Private/OnDemandHttpClient.cpp
+++ b/Engine/Source/Runtime/Experimental/IoStore/OnDemand/Private/OnDemandHttpClient.cpp
@@ -321,6 +321,10 @@ void FMultiEndpointHttpClient::CancelRequest(FHttpTicketId TicketId)
 
 		if (Ticket->HttpTicket != 0)
 		{
+#if !NO_LOGGING
+			FMsg::Logf(__FILE__, __LINE__, Config.LogCategory->GetCategoryName(), ELogVerbosity::Verbose, TEXT("Canceling FHttpTicketId %u with FTicket: %llu"), TicketId, Ticket->HttpTicket);
+#endif // !NO_LOGGING
+
 			EventLoop.Cancel(Ticket->HttpTicket);
 		}
 	}
