Fix the chaos transactional failure by using a transactionally-safe queue.

diff --git a/Engine/Source/Runtime/Experimental/Chaos/Public/Chaos/SimCallbackObject.h b/Engine/Source/Runtime/Experimental/Chaos/Public/Chaos/SimCallbackObject.h
--- a/Engine/Source/Runtime/Experimental/Chaos/Public/Chaos/SimCallbackObject.h
+++ b/Engine/Source/Runtime/Experimental/Chaos/Public/Chaos/SimCallbackObject.h
@@ -1,7 +1,7 @@
 // Copyright Epic Games, Inc. All Rights Reserved.
 #pragma once
 
-#include "Containers/SpscQueue.h"
+#include "Containers/TransactionallySafeSpscQueue.h"
 #include "Chaos/Defines.h"
 #include "Chaos/SimCallbackInput.h"
 #include "Chaos/CollisionResolutionTypes.h"
@@ -633,7 +633,7 @@ private:
 	}
 
 	template <typename T>
-	T* NewDataHelper(TArray<TUniquePtr<T>>& Backing, TSpscQueue<T*>& Queue)
+	T* NewDataHelper(TArray<TUniquePtr<T>>& Backing, TTransactionallySafeSpscQueue<T*>& Queue)
 	{
 		T* Result;
 		if (!Queue.Dequeue(Result))
@@ -667,12 +667,12 @@ private:
 		OutputQueue.Enqueue(static_cast<TOutputType*>(BaseOutput));
 	}
 
-	TSpscQueue<TInputType*> InputPool;
+	TTransactionallySafeSpscQueue<TInputType*> InputPool;
 	TArray<TUniquePtr<TInputType>> InputBacking;
 
-	TSpscQueue<TOutputType*> OutputPool;
+	TTransactionallySafeSpscQueue<TOutputType*> OutputPool;
 	TArray<TUniquePtr<TOutputType>> OutputBacking;
-	TSpscQueue<TOutputType*> OutputQueue;	//holds the outputs in order
+	TTransactionallySafeSpscQueue<TOutputType*> OutputQueue;	//holds the outputs in order
 
 	TOutputType* CurrentOutput_External;	//the earliest output we can consume
 };
