GeometryCollection: Add Nanite Minimum Residency setting.
Previously nanite geom collections always had a single root page. This could cause pop-in on the first frames because a single page is unlikely to cover all the simple clusters for a fractured mesh (at least one cluster per fracture).
Default behavior doesn't change. But now it is possible to increase minimum residency for GCs that require it.
#FailedValidations Preflight

diff --git a/Engine/Source/Runtime/Experimental/GeometryCollectionEngine/Private/GeometryCollection/DerivedDataGeometryCollectionCooker.cpp b/Engine/Source/Runtime/Experimental/GeometryCollectionEngine/Private/GeometryCollection/DerivedDataGeometryCollectionCooker.cpp
--- a/Engine/Source/Runtime/Experimental/GeometryCollectionEngine/Private/GeometryCollection/DerivedDataGeometryCollectionCooker.cpp
+++ b/Engine/Source/Runtime/Experimental/GeometryCollectionEngine/Private/GeometryCollection/DerivedDataGeometryCollectionCooker.cpp
@@ -102,7 +102,15 @@ bool FDerivedDataGeometryCollectionRenderDataCooker::Build(TArray<uint8>& OutDat
 	{
 		Chaos::FErrorReporter ErrorReporter(GeometryCollection.GetName());
 
-		TUniquePtr<FGeometryCollectionRenderData> RenderData = FGeometryCollectionRenderData::Create(*Collection, GeometryCollection.EnableNanite, GeometryCollection.bEnableNaniteFallback, GeometryCollection.bUseFullPrecisionUVs, GeometryCollection.bConvertVertexColorsToSRGB);
+		FGeometryCollectionRenderData::FCreateDesc CreateDesc = {
+			.bEnableNanite = GeometryCollection.EnableNanite,
+			.bEnableNaniteFallback = GeometryCollection.bEnableNaniteFallback,
+			.bUseFullPrecisionUVs = GeometryCollection.bUseFullPrecisionUVs, 
+			.bConvertVertexColorsToSRGB = GeometryCollection.bConvertVertexColorsToSRGB,
+			.NaniteMinimumResidencyInKB = GeometryCollection.NaniteMinimumResidencyInKB
+		};
+
+		TUniquePtr<FGeometryCollectionRenderData> RenderData = FGeometryCollectionRenderData::Create(*Collection, CreateDesc);
 		RenderData->Serialize(ChaosAr, GeometryCollection);
 
 		if (false && ErrorReporter.EncounteredAnyErrors())
diff --git a/Engine/Source/Runtime/Experimental/GeometryCollectionEngine/Private/GeometryCollection/GeometryCollectionObject.cpp b/Engine/Source/Runtime/Experimental/GeometryCollectionEngine/Private/GeometryCollection/GeometryCollectionObject.cpp
--- a/Engine/Source/Runtime/Experimental/GeometryCollectionEngine/Private/GeometryCollection/GeometryCollectionObject.cpp
+++ b/Engine/Source/Runtime/Experimental/GeometryCollectionEngine/Private/GeometryCollection/GeometryCollectionObject.cpp
@@ -115,6 +115,7 @@ UGeometryCollection::UGeometryCollection(const FObjectInitializer& ObjectInitial
 	, bStripRenderDataOnCook(false)
 	, EnableNanite(false)
 	, bEnableNaniteFallback(false)
+	, NaniteMinimumResidencyInKB(0)
 #if WITH_EDITORONLY_DATA
 	, CollisionType_DEPRECATED(ECollisionTypeEnum::Chaos_Volumetric)
 	, ImplicitType_DEPRECATED(EImplicitTypeEnum::Chaos_Implicit_Convex)
@@ -1499,7 +1500,15 @@ void UGeometryCollection::RebuildRenderData()
 		}
 
 		// Create new RenderData and initialize render resources (happens on render thread).
-		RenderData = FGeometryCollectionRenderData::Create(*GetGeometryCollection(), EnableNanite, bEnableNaniteFallback, bUseFullPrecisionUVs, bConvertVertexColorsToSRGB);
+		FGeometryCollectionRenderData::FCreateDesc CreateDesc = {
+			.bEnableNanite = EnableNanite,
+			.bEnableNaniteFallback = bEnableNaniteFallback,
+			.bUseFullPrecisionUVs = bUseFullPrecisionUVs,
+			.bConvertVertexColorsToSRGB = bConvertVertexColorsToSRGB,
+			.NaniteMinimumResidencyInKB = NaniteMinimumResidencyInKB
+		};
+
+		RenderData = FGeometryCollectionRenderData::Create(*GetGeometryCollection(), CreateDesc);
 		InitResources();
 
 		PropagateMarkDirtyToComponents();
@@ -1896,6 +1905,11 @@ void UGeometryCollection::PostEditChangeProperty(struct FPropertyChangedEvent& P
 			bDoInvalidateCollection = true;
 			bRebuildRenderData = true;
 		}
+		else if (PropertyChangedEvent.Property->GetFName() == GET_MEMBER_NAME_CHECKED(UGeometryCollection, NaniteMinimumResidencyInKB))
+		{
+			bDoInvalidateCollection = true;
+			bRebuildRenderData = true;
+		}
 		else if (PropertyChangedEvent.Property->GetFName() == GET_MEMBER_NAME_CHECKED(UGeometryCollection, bUseFullPrecisionUVs))
 		{
 			bDoInvalidateCollection = true;
diff --git a/Engine/Source/Runtime/Experimental/GeometryCollectionEngine/Private/GeometryCollection/GeometryCollectionRenderData.cpp b/Engine/Source/Runtime/Experimental/GeometryCollectionEngine/Private/GeometryCollection/GeometryCollectionRenderData.cpp
--- a/Engine/Source/Runtime/Experimental/GeometryCollectionEngine/Private/GeometryCollection/GeometryCollectionRenderData.cpp
+++ b/Engine/Source/Runtime/Experimental/GeometryCollectionEngine/Private/GeometryCollection/GeometryCollectionRenderData.cpp
@@ -160,7 +160,7 @@ struct FGeometryCollectionBuiltMeshData
 };
 
 /** Get the data needed to build render data from a Geometry Collection. */
-bool BuildMeshDataFromGeometryCollection(FGeometryCollection& InCollection, FGeometryCollectionBuiltMeshData& OutMeshData, bool bConvertVertexColorsToSRGB)
+bool BuildMeshDataFromGeometryCollection(FGeometryCollection& InCollection, FGeometryCollectionRenderData::FCreateDesc const& InDesc, FGeometryCollectionBuiltMeshData& OutMeshData)
 {
 	// Vertices Group
 	const TManagedArray<FVector3f>& VertexArray = InCollection.Vertex;
@@ -244,7 +244,7 @@ bool BuildMeshDataFromGeometryCollection(FGeometryCollection& InCollection, FGeo
 
 				if (bHasColors)
 				{
-					BuildVertexData.Color[DestVertexStart + VertexIndex] = ColorArray[VertexStart + VertexIndex].ToFColor(bConvertVertexColorsToSRGB /* sRGB */);
+					BuildVertexData.Color[DestVertexStart + VertexIndex] = ColorArray[VertexStart + VertexIndex].ToFColor(InDesc.bConvertVertexColorsToSRGB /* sRGB */);
 				}
 
 				for (int32 TexCoord = 0; TexCoord < NumTexCoords; ++TexCoord)
@@ -604,7 +604,7 @@ FBoxSphereBounds BuildPreSkinnedBounds(FGeometryCollection& InCollection)
 	return bBoundsInit ? FBoxSphereBounds(PreSkinnedBounds) : FBoxSphereBounds(ForceInit);
 }
 
-void CreateMeshData(FGeometryCollection& InCollection, FGeometryCollectionBuiltMeshData const& InMeshData, bool bInUseFullPrecisionUVs, FGeometryCollectionRenderData& OutRenderData)
+void CreateMeshData(FGeometryCollectionBuiltMeshData const& InMeshData, FGeometryCollectionRenderData::FCreateDesc const& InDesc, FGeometryCollectionRenderData& OutRenderData)
 {
 	TRACE_CPUPROFILER_EVENT_SCOPE(FGeometryCollectionRenderData::CreateMeshData);
 
@@ -618,7 +618,7 @@ void CreateMeshData(FGeometryCollection& InCollection, FGeometryCollectionBuiltM
 	{
 		TRACE_CPUPROFILER_EVENT_SCOPE(FGeometryCollectionRenderData::BufferInit);
 		OutRenderData.MeshResource.PositionVertexBuffer.Init(VertexView);
-		OutRenderData.MeshResource.StaticMeshVertexBuffer.SetUseFullPrecisionUVs(bInUseFullPrecisionUVs);
+		OutRenderData.MeshResource.StaticMeshVertexBuffer.SetUseFullPrecisionUVs(InDesc.bUseFullPrecisionUVs);
 		OutRenderData.MeshResource.StaticMeshVertexBuffer.Init(VertexView);
 		OutRenderData.MeshResource.ColorVertexBuffer.Init(VertexView);
 		OutRenderData.MeshResource.BoneMapVertexBuffer.Init(InMeshData.BonesPerVertex);
@@ -635,7 +635,7 @@ void CreateMeshData(FGeometryCollection& InCollection, FGeometryCollectionBuiltM
 	OutRenderData.MeshDescription.SubSections = MoveTemp(SectionData.SubSections);
 }
 
-void CreateNaniteData(FGeometryCollectionBuiltMeshData&& InMeshData, FGeometryCollectionRenderData& OutRenderData)
+void CreateNaniteData(FGeometryCollectionBuiltMeshData&& InMeshData, FGeometryCollectionRenderData::FCreateDesc const& InDesc, FGeometryCollectionRenderData& OutRenderData)
 {
 	TRACE_CPUPROFILER_EVENT_SCOPE(FGeometryCollectionRenderData::CreateNaniteData);
 
@@ -643,7 +643,7 @@ void CreateNaniteData(FGeometryCollectionBuiltMeshData&& InMeshData, FGeometryCo
 
 	FMeshNaniteSettings NaniteSettings = {};
 	NaniteSettings.bEnabled = true;
-	NaniteSettings.TargetMinimumResidencyInKB = 0;	// Default to smallest possible, which is a single page
+	NaniteSettings.TargetMinimumResidencyInKB = InDesc.NaniteMinimumResidencyInKB;
 	NaniteSettings.KeepPercentTriangles = 1.0f;
 	NaniteSettings.TrimRelativeError = 0.0f;
 	NaniteSettings.FallbackPercentTriangles = 1.0f; // 100% - no reduction
@@ -675,23 +675,23 @@ void CreateNaniteData(FGeometryCollectionBuiltMeshData&& InMeshData, FGeometryCo
 	}
 }
 
-TUniquePtr<FGeometryCollectionRenderData> FGeometryCollectionRenderData::Create(FGeometryCollection& InCollection, bool bInEnableNanite, bool bInEnableNaniteFallback, bool bInUseFullPrecisionUVs, bool bConvertVertexColorsToSRGB)
+TUniquePtr<FGeometryCollectionRenderData> FGeometryCollectionRenderData::Create(FGeometryCollection& InCollection, FCreateDesc const& InDesc)
 {
 	TRACE_CPUPROFILER_EVENT_SCOPE(FGeometryCollectionRenderData::Create);
 
 	TUniquePtr<FGeometryCollectionRenderData> RenderData = MakeUnique<FGeometryCollectionRenderData>();
 
 	FGeometryCollectionBuiltMeshData MeshBuildData;
-	if (BuildMeshDataFromGeometryCollection(InCollection, MeshBuildData, bConvertVertexColorsToSRGB))
+	if (BuildMeshDataFromGeometryCollection(InCollection, InDesc, MeshBuildData))
 	{
-		if (!bInEnableNanite || bInEnableNaniteFallback)
+		if (!InDesc.bEnableNanite || InDesc.bEnableNaniteFallback)
 		{
-			CreateMeshData(InCollection, MeshBuildData, bInUseFullPrecisionUVs, *RenderData.Get());
+			CreateMeshData(MeshBuildData, InDesc, *RenderData.Get());
 		}
 
-		if (bInEnableNanite)
+		if (InDesc.bEnableNanite)
 		{
-			CreateNaniteData(MoveTemp(MeshBuildData), *RenderData.Get());
+			CreateNaniteData(MoveTemp(MeshBuildData), InDesc, *RenderData.Get());
 		}
 
 		RenderData->PreSkinnedBounds = BuildPreSkinnedBounds(InCollection);
diff --git a/Engine/Source/Runtime/Experimental/GeometryCollectionEngine/Private/GeometryCollection/GeometryCollectionRenderData.h b/Engine/Source/Runtime/Experimental/GeometryCollectionEngine/Private/GeometryCollection/GeometryCollectionRenderData.h
--- a/Engine/Source/Runtime/Experimental/GeometryCollectionEngine/Private/GeometryCollection/GeometryCollectionRenderData.h
+++ b/Engine/Source/Runtime/Experimental/GeometryCollectionEngine/Private/GeometryCollection/GeometryCollectionRenderData.h
@@ -113,7 +113,16 @@ public:
 	~FGeometryCollectionRenderData();
 
 #if WITH_EDITOR
-	static TUniquePtr<FGeometryCollectionRenderData> Create(FGeometryCollection& InCollection, bool bInEnableNanite, bool bInEnableNaniteFallback, bool bInUseFullPrecisionUVs, bool bConvertVertexColorsToSRGB);
+	struct FCreateDesc
+	{
+		bool bEnableNanite = false;
+		bool bEnableNaniteFallback = false;
+		bool bUseFullPrecisionUVs = false;
+		bool bConvertVertexColorsToSRGB = false;
+		uint32 NaniteMinimumResidencyInKB = 0;
+	};
+
+	static TUniquePtr<FGeometryCollectionRenderData> Create(FGeometryCollection& InCollection, FCreateDesc const& InDesc);
 #endif
 
 	bool IsInitialized()
diff --git a/Engine/Source/Runtime/Experimental/GeometryCollectionEngine/Public/GeometryCollection/GeometryCollectionObject.h b/Engine/Source/Runtime/Experimental/GeometryCollectionEngine/Public/GeometryCollection/GeometryCollectionObject.h
--- a/Engine/Source/Runtime/Experimental/GeometryCollectionEngine/Public/GeometryCollection/GeometryCollectionObject.h
+++ b/Engine/Source/Runtime/Experimental/GeometryCollectionEngine/Public/GeometryCollection/GeometryCollectionObject.h
@@ -635,6 +635,12 @@ public:
 	UPROPERTY(EditAnywhere, Category = "Nanite", Meta = (EditCondition = "EnableNanite"))
 	bool bEnableNaniteFallback;
 
+	/**
+	 * How much of the nanite resource should always be resident (In KB). Approximate due to paging. 0: Minimum size (single page). MAX_uint32: Entire mesh.
+	 */
+	UPROPERTY(EditAnywhere, Category = "Nanite", Meta = (EditCondition = "EnableNanite", DisplayName = "Nanite Minimum Residency"))
+	uint32 NaniteMinimumResidencyInKB;
+
 	UFUNCTION(BlueprintCallable, Category = "Rendering")
 	GEOMETRYCOLLECTIONENGINE_API void SetConvertVertexColorsToSRGB(bool bValue);
 
