Add assert to verify that each IAX http request is only being completed once.

#rnx

- We recently had a bug in the lower level code that could cause the requests sink to be invoked multiple times with a completion state. This assert would have made the issue more obvious.

diff --git a/Engine/Source/Runtime/Experimental/IoStore/OnDemand/Private/OnDemandHttpClient.cpp b/Engine/Source/Runtime/Experimental/IoStore/OnDemand/Private/OnDemandHttpClient.cpp
--- a/Engine/Source/Runtime/Experimental/IoStore/OnDemand/Private/OnDemandHttpClient.cpp
+++ b/Engine/Source/Runtime/Experimental/IoStore/OnDemand/Private/OnDemandHttpClient.cpp
@@ -437,6 +437,10 @@ FMultiEndpointHttpClient::FHttpTicketId FMultiEndpointHttpClient::IssueRequest(F
 			case FTicketStatus::EId::Error:
 			case FTicketStatus::EId::Cancelled:
 			{
+#if DO_CHECK
+				check(Request.SinkCounter++ == Request.RetryCount);
+#endif //DO_CHECK
+
 				Log(Request, TicketStatus);
 
 				// We only want to retry the request if there was an internal or server problem. There is no point retrying a 404 error for example/
diff --git a/Engine/Source/Runtime/Experimental/IoStore/OnDemand/Private/OnDemandHttpClient.h b/Engine/Source/Runtime/Experimental/IoStore/OnDemand/Private/OnDemandHttpClient.h
--- a/Engine/Source/Runtime/Experimental/IoStore/OnDemand/Private/OnDemandHttpClient.h
+++ b/Engine/Source/Runtime/Experimental/IoStore/OnDemand/Private/OnDemandHttpClient.h
@@ -176,6 +176,11 @@ private:
 		int32				Host = INDEX_NONE;
 		bool				bIsChunkedTransfer = false;
 		int8				CDNCacheStatus = -1;
+
+#if DO_CHECK
+		/** Used to make sure that the sink is only called with a completed state one per request attempt */
+		int32				SinkCounter = 0;
+#endif // DO_CHECK
 	};
 
 										FMultiEndpointHttpClient(const FMultiEndpointHttpClientConfig& Config);
