Constraints: only run the task's evaluation function when added

#rnx

diff --git a/Engine/Source/Runtime/Experimental/Animation/Constraints/Private/Transform/AnimationEvaluation.cpp b/Engine/Source/Runtime/Experimental/Animation/Constraints/Private/Transform/AnimationEvaluation.cpp
--- a/Engine/Source/Runtime/Experimental/Animation/Constraints/Private/Transform/AnimationEvaluation.cpp
+++ b/Engine/Source/Runtime/Experimental/Animation/Constraints/Private/Transform/AnimationEvaluation.cpp
@@ -102,15 +102,18 @@ FTransform FAnimationEvaluator::GetGlobalTransform(const FName InSocketName) con
 	return RawSkeletalMeshComponent->GetSocketTransform(InSocketName);
 }
 
-void FAnimationEvaluator::AddPostEvaluationTask(const FAnimationEvaluationTask& InTask)
+bool FAnimationEvaluator::AddPostEvaluationTask(const FAnimationEvaluationTask& InTask)
 {
 	if (USkeletalMeshComponent* RawSkeletalMeshComponent = SkeletalMeshComponent.Get())
 	{
-		if (InTask.IsValid(RawSkeletalMeshComponent))
+		if (InTask.IsValid(RawSkeletalMeshComponent) && !PostEvaluationTasks.Contains(InTask.Guid))
 		{
 			PostEvaluationTasks.FindOrAdd(InTask.Guid, InTask);
+			return true;
 		}
 	}
+	
+	return false;
 }
 	
 void FAnimationEvaluator::UpdateContext()
@@ -390,7 +393,7 @@ const FAnimationEvaluator& FAnimationEvaluationCache::GetEvaluator(USkeletalMesh
 
 	if (FAnimationEvaluator* Found = PerSkeletalMeshEvaluator.Find(InSkeletalMeshComponent))
 	{
-		Found->AddPostEvaluationTask(InTask);
+		const bool bAdded = Found->AddPostEvaluationTask(InTask);
 		if (!Found->IsValid())
 		{
 			if (bLogNewEvaluation)
@@ -399,7 +402,7 @@ const FAnimationEvaluator& FAnimationEvaluationCache::GetEvaluator(USkeletalMesh
 			}
 			Found->Update(bRefreshBones);
 		}
-		else if (InTask.IsValid(InSkeletalMeshComponent))
+		else if (bAdded && InTask.IsValid(InSkeletalMeshComponent))
 		{
 			InTask.PostEvaluationFunction();
 		}
diff --git a/Engine/Source/Runtime/Experimental/Animation/Constraints/Public/Transform/AnimationEvaluation.h b/Engine/Source/Runtime/Experimental/Animation/Constraints/Public/Transform/AnimationEvaluation.h
--- a/Engine/Source/Runtime/Experimental/Animation/Constraints/Public/Transform/AnimationEvaluation.h
+++ b/Engine/Source/Runtime/Experimental/Animation/Constraints/Public/Transform/AnimationEvaluation.h
@@ -51,8 +51,8 @@ struct FAnimationEvaluator
 	/** Returns InSocketName component space transform composed with the skeletal mesh component's global transform. */
 	FTransform GetGlobalTransform(const FName InSocketName) const;
 
-	/** Adds a post-evaluation task to this evaluator if it has not already been added. */
-	void AddPostEvaluationTask(const FAnimationEvaluationTask& InTask);
+	/** Adds a post-evaluation task to this evaluator if it has not already been added (returns true if added). */
+	bool AddPostEvaluationTask(const FAnimationEvaluationTask& InTask);
 	
 private:
 
