Prevent the new IAX backend from treating all http errors as decoding errors as well.

#rnx

- Previously, the code assumed if no decoding took place then we had a decoding error but it was not taking into account if the request actually succeeded and if a decode was actually attempted or not.
- Now we only check for decode errors if a decode was attempted.

#FailedValidations Preflight

diff --git a/Engine/Source/Runtime/Experimental/IoStore/OnDemand/Private/OnDemandHttpIoDispatcherBackend.cpp b/Engine/Source/Runtime/Experimental/IoStore/OnDemand/Private/OnDemandHttpIoDispatcherBackend.cpp
--- a/Engine/Source/Runtime/Experimental/IoStore/OnDemand/Private/OnDemandHttpIoDispatcherBackend.cpp
+++ b/Engine/Source/Runtime/Experimental/IoStore/OnDemand/Private/OnDemandHttpIoDispatcherBackend.cpp
@@ -274,6 +274,23 @@ void FHttpIoDispatcherBackend::CompleteRequest(FResolvedRequest& ResolvedRequest
 			? EIoDecodeFlags::WriteOnly : EIoDecodeFlags::None;
 
 		bDecoded = FIoChunkEncoding::Decode(Params, Chunk.GetView(), DispatcherRequest.GetBuffer().GetMutableView(), Options);
+
+		if (!bDecoded)
+		{
+			if (HttpResponse.IsCached())
+			{
+				Stats.OnCacheDecodeError();
+
+				if (FIoStatus Status = FHttpIoDispatcher::EvictFromCache(HttpResponse); !Status.IsOk())
+				{
+					UE_LOG(LogHttpIoDispatcher, Error, TEXT("Evict HTTP cache failed, reason '%s'"), *Status.ToString());
+				}
+			}
+			else
+			{
+				Stats.OnHttpDecodeError(EHttpRequestType::Streaming);
+			}
+		}
 	}
 
 	const uint64 DurationMs = ResolvedRequest.StartTimeCycles > 0
@@ -298,20 +315,6 @@ void FHttpIoDispatcherBackend::CompleteRequest(FResolvedRequest& ResolvedRequest
 	{
 		DispatcherRequest.SetFailed();
 		Stats.OnIoRequestError();
-
-		if (HttpResponse.IsCached())
-		{
-			Stats.OnCacheDecodeError();
-
-			if (FIoStatus Status = FHttpIoDispatcher::EvictFromCache(HttpResponse); !Status.IsOk())
-			{
-				UE_LOG(LogHttpIoDispatcher, Error, TEXT("Evict HTTP cache failed, reason '%s'"), *Status.ToString());
-			}
-		}
-		else
-		{
-			Stats.OnHttpDecodeError(EHttpRequestType::Streaming);
-		}
 	}
 	
 	LogRequest(ResolvedRequest, DispatcherRequest, bDecoded, HttpResponse.IsCached(), bCancelled, DurationMs);
