Fixed a minor bug where a http request that has been queued for a retry would have it's retry count incremented even though that retry was never issued as the request was canceled.

#rnx

diff --git a/Engine/Source/Runtime/Experimental/IoStore/OnDemand/Private/OnDemandHttpClient.cpp b/Engine/Source/Runtime/Experimental/IoStore/OnDemand/Private/OnDemandHttpClient.cpp
--- a/Engine/Source/Runtime/Experimental/IoStore/OnDemand/Private/OnDemandHttpClient.cpp
+++ b/Engine/Source/Runtime/Experimental/IoStore/OnDemand/Private/OnDemandHttpClient.cpp
@@ -594,7 +594,7 @@ void FMultiEndpointHttpClient::ProcessRetryAttempts(uint32 TicketCount)
 					.Reason = TEXT("Canceled"),
 					.DurationMilliseconds = uint64(FPlatformTime::ToMilliseconds64(FPlatformTime::Cycles64() - Request.StartTime)),
 					.StatusCode = Request.StatusCode,
-					.RetryCount = Request.RetryCount,
+					.RetryCount = Request.RetryCount -1, // Reduce retry count by 1 as this retry attempt didn't really do anything
 					.HostIndex = Request.Host,
 					.bCanceled = true,
 					.CDNCacheStatus = Request.CDNCacheStatus
