Deactivate partial island sleep if the ratio of motionless particles in an island is below a user-defined threshold:
- Setting the threshold to 0 will always activate partial island sleep
- Minimum number of ticks below threshold before switching between full and partial island sleep to prevent too many switches.
- Improve performance to automatically disable partial island sleep if there are only few motionless particles in the island.

diff --git a/Engine/Source/Programs/HeadlessChaos/Private/HeadlessChaosTestGraphEvolution.cpp b/Engine/Source/Programs/HeadlessChaos/Private/HeadlessChaosTestGraphEvolution.cpp
--- a/Engine/Source/Programs/HeadlessChaos/Private/HeadlessChaosTestGraphEvolution.cpp
+++ b/Engine/Source/Programs/HeadlessChaos/Private/HeadlessChaosTestGraphEvolution.cpp
@@ -49,6 +49,11 @@ namespace ChaosTest
 			bPartialWakePreIntegration_EditorDefault = CVarPartialWakePreIntegration->GetBool();
 			CVarPartialWakePreIntegration->Set(bPartialWakePreIntegration_UnitTestDefault);
 
+			CVarMinMotionlessRatio = IConsoleManager::Get().FindConsoleVariable(TEXT("p.Chaos.Solver.Sleep.PartialIslandSleep.MinimumMotionlessRatio"), false);
+			check(CVarMinMotionlessRatio && CVarMinMotionlessRatio->IsVariableFloat());
+			MinimumMotionlessThreshold_EditorDefault = CVarMinMotionlessRatio->GetFloat();
+			CVarMinMotionlessRatio->Set(MinimumMotionlessThreshold_UnitTestDefault);
+
 			CVarLinearWakeThresholdMultiplier = IConsoleManager::Get().FindConsoleVariable(TEXT("p.Chaos.Solver.Sleep.PartialIslandSleep.LinearWakeThresholdMultiplier"), false);
 			check(CVarLinearWakeThresholdMultiplier && CVarLinearWakeThresholdMultiplier->IsVariableFloat());
 			LinearWakeThresholdMultiplier_EditorDefault = CVarLinearWakeThresholdMultiplier->GetFloat();
@@ -66,6 +71,7 @@ namespace ChaosTest
 			CVarMomentumPropagation->Set(MomentumPropagation_EditorDefault);
 			CVarPartialSleepCollisionConstraintsOnly->Set(bPartialSleepCollisionConstraintsOnly_EditorDefault);
 			CVarPartialWakePreIntegration->Set(bPartialWakePreIntegration_EditorDefault);
+			CVarMinMotionlessRatio->Set(MinimumMotionlessThreshold_EditorDefault);
 			CVarLinearWakeThresholdMultiplier->Set(LinearWakeThresholdMultiplier_EditorDefault);
 			CVarAngularWakeThresholdMultiplier->Set(AngularWakeThresholdMultiplier_EditorDefault);
 		}
@@ -75,6 +81,7 @@ namespace ChaosTest
 		FRealSingle MomentumPropagation_UnitTestDefault = 0.0f;
 		bool bPartialSleepCollisionConstraintsOnly_UnitTestDefault = false;
 		bool bPartialWakePreIntegration_UnitTestDefault = false;
+		FRealSingle MinimumMotionlessThreshold_UnitTestDefault = 0.0f;
 		FRealSingle LinearWakeThresholdMultiplier_UnitTestDefault = 1.0f;
 		FRealSingle AngularWakeThresholdMultiplier_UnitTestDefault = 1.0f;
 
@@ -83,6 +90,7 @@ namespace ChaosTest
 		FRealSingle MomentumPropagation_EditorDefault;
 		bool bPartialSleepCollisionConstraintsOnly_EditorDefault;
 		bool bPartialWakePreIntegration_EditorDefault;
+		FRealSingle MinimumMotionlessThreshold_EditorDefault;
 		FRealSingle LinearWakeThresholdMultiplier_EditorDefault;
 		FRealSingle AngularWakeThresholdMultiplier_EditorDefault;
 
@@ -92,6 +100,7 @@ namespace ChaosTest
 		IConsoleVariable* CVarMomentumPropagation;
 		IConsoleVariable* CVarPartialSleepCollisionConstraintsOnly;
 		IConsoleVariable* CVarPartialWakePreIntegration;
+		IConsoleVariable* CVarMinMotionlessRatio;
 		IConsoleVariable* CVarLinearWakeThresholdMultiplier;
 		IConsoleVariable* CVarAngularWakeThresholdMultiplier;
 	};
@@ -2633,5 +2642,103 @@ namespace ChaosTest
 		CVarMomentumPropagation->Set(MomentumPropagationDefault);
 	}
 
+	TEST_P(GraphEvolutionTests, TestConstraintGraph_SwitchBewtweenFullAndPartialIslandSleep)
+	{
+		const bool bEnablePartialSleep = true;
+		const FRealSingle MinMotionlessRatio = GetParam() ? 0.5f : 0.1f;
+
+		const FRealSingle MinMotionlessRatioDefault = CVarMinMotionlessRatio->GetFloat();
+		CVarMinMotionlessRatio->Set(MinMotionlessRatio);
+
+		const IConsoleVariable* CVarSleepModeSwitchCounterThreshold = IConsoleManager::Get().FindConsoleVariable(TEXT("p.Chaos.Solver.Sleep.PartialIslandSleep.SleepModeSwitchCounterThreshold"), false);
+		check(CVarSleepModeSwitchCounterThreshold && CVarSleepModeSwitchCounterThreshold->IsVariableInt());
+		const int32 SleepModeSwitchCounterThreshold = CVarSleepModeSwitchCounterThreshold->GetInt();
+		const int32 Steps = SleepModeSwitchCounterThreshold;
+
+		FGraphEvolutionTest Test(6, bEnablePartialSleep);
+		Test.MakeChain();
+		Test.Evolution.SetParticleObjectState(Test.ParticleHandles[0], EObjectStateType::Kinematic);
+
+		Test.AdvanceUntilSleeping();
+
+		FParticleDynamics Dynamics;
+		Dynamics.SetAcceleration(FVec3(0, 0, 0));
+		Dynamics.SetAngularAcceleration(FVec3(0, 0, 0));
+		Dynamics.SetLinearImpulseVelocity(FVec3(0, 0.1, 0));
+		Dynamics.SetAngularImpulseVelocity(FVec3(0, 0, 0));
+
+		// Explicitly wake up one of the top 3 particles and apply and impulse
+		Test.Evolution.SetParticleObjectState(Test.ParticleHandles[3], EObjectStateType::Dynamic);
+		Test.Evolution.SetParticleObjectState(Test.ParticleHandles[4], EObjectStateType::Dynamic);
+		Test.Evolution.SetParticleObjectState(Test.ParticleHandles[5], EObjectStateType::Dynamic);
+		Test.Evolution.SetParticleDynamics(Test.ParticleHandles[3], Dynamics);
+		Test.Evolution.SetParticleDynamics(Test.ParticleHandles[4], Dynamics);
+		Test.Evolution.SetParticleDynamics(Test.ParticleHandles[5], Dynamics);
+		
+		for (int32 I = 0; I < Steps; ++I)
+		{
+			Test.Advance();
+		}
+
+		// The particles explicitly woken and the particle next to it are awake.
+		EXPECT_TRUE(FConstGenericParticleHandle(Test.ParticleHandles[0])->IsKinematic());
+		EXPECT_TRUE(FConstGenericParticleHandle(Test.ParticleHandles[1])->Sleeping()); // sleeping
+		EXPECT_FALSE(FConstGenericParticleHandle(Test.ParticleHandles[2])->Sleeping()); // awake, motionless
+		EXPECT_FALSE(FConstGenericParticleHandle(Test.ParticleHandles[3])->Sleeping()); // awake, moving
+		EXPECT_FALSE(FConstGenericParticleHandle(Test.ParticleHandles[4])->Sleeping()); // awake, moving
+		EXPECT_FALSE(FConstGenericParticleHandle(Test.ParticleHandles[5])->Sleeping()); // awake, moving
+
+		// Partial island sleep remains enabled because particles have not moved enough steps
+		EXPECT_TRUE(Test.IslandManager->GetIsland(0)->IsPartialSleepEnabled());
+
+		Test.Advance();
+		
+		if (MinMotionlessRatio == 0.5f)
+		{
+			// Partial island sleeping is internally disabled since 2/5 = 40% of dynamic particles are motionless.
+			EXPECT_FALSE(Test.IslandManager->GetIsland(0)->IsPartialSleepEnabled());
+		}
+		else
+		{
+			// Partial island sleeping is internally enabled since 2/5 = 40% of dynamic particles are motionless.
+			EXPECT_TRUE(Test.IslandManager->GetIsland(0)->IsPartialSleepEnabled());
+		}
+
+		Test.Advance();
+
+		if (MinMotionlessRatio == 0.5f)
+		{
+			// All particles wake up
+			EXPECT_TRUE(FConstGenericParticleHandle(Test.ParticleHandles[0])->IsKinematic());
+			EXPECT_FALSE(FConstGenericParticleHandle(Test.ParticleHandles[1])->Sleeping());
+			EXPECT_FALSE(FConstGenericParticleHandle(Test.ParticleHandles[2])->Sleeping());
+			EXPECT_FALSE(FConstGenericParticleHandle(Test.ParticleHandles[3])->Sleeping());
+			EXPECT_FALSE(FConstGenericParticleHandle(Test.ParticleHandles[4])->Sleeping());
+			EXPECT_FALSE(FConstGenericParticleHandle(Test.ParticleHandles[5])->Sleeping());
+		}
+		else
+		{
+			EXPECT_TRUE(FConstGenericParticleHandle(Test.ParticleHandles[0])->IsKinematic());
+			EXPECT_TRUE(FConstGenericParticleHandle(Test.ParticleHandles[1])->Sleeping()); // sleeping
+			EXPECT_FALSE(FConstGenericParticleHandle(Test.ParticleHandles[2])->Sleeping()); // awake, motionless
+			EXPECT_FALSE(FConstGenericParticleHandle(Test.ParticleHandles[3])->Sleeping()); // awake, moving
+			EXPECT_FALSE(FConstGenericParticleHandle(Test.ParticleHandles[4])->Sleeping()); // awake, moving
+			EXPECT_FALSE(FConstGenericParticleHandle(Test.ParticleHandles[5])->Sleeping()); // awake, moving
+		}
+
+		Test.CVarPartialSleeping->Set(false);
+		Test.Advance();
+
+		// All particles wake up when disabling partial island sleep explicitly (CVar).
+		EXPECT_TRUE(FConstGenericParticleHandle(Test.ParticleHandles[0])->IsKinematic());
+		EXPECT_FALSE(FConstGenericParticleHandle(Test.ParticleHandles[1])->Sleeping());
+		EXPECT_FALSE(FConstGenericParticleHandle(Test.ParticleHandles[2])->Sleeping());
+		EXPECT_FALSE(FConstGenericParticleHandle(Test.ParticleHandles[3])->Sleeping());
+		EXPECT_FALSE(FConstGenericParticleHandle(Test.ParticleHandles[4])->Sleeping());
+		EXPECT_FALSE(FConstGenericParticleHandle(Test.ParticleHandles[5])->Sleeping());
+
+		CVarMinMotionlessRatio->Set(MinMotionlessRatioDefault);
+	}
+
 	INSTANTIATE_TEST_SUITE_P(, GraphEvolutionTests,	testing::Bool());
 }
diff --git a/Engine/Source/Programs/HeadlessChaos/Private/HeadlessChaosTestSleeping.cpp b/Engine/Source/Programs/HeadlessChaos/Private/HeadlessChaosTestSleeping.cpp
--- a/Engine/Source/Programs/HeadlessChaos/Private/HeadlessChaosTestSleeping.cpp
+++ b/Engine/Source/Programs/HeadlessChaos/Private/HeadlessChaosTestSleeping.cpp
@@ -46,6 +46,11 @@ namespace ChaosTest
 			bPartialWakePreIntegration_EditorDefault = CVarPartialWakePreIntegration->GetBool();
 			CVarPartialWakePreIntegration->Set(bPartialWakePreIntegration_UnitTestDefault);
 
+			CVarMinMotionlessRatio = IConsoleManager::Get().FindConsoleVariable(TEXT("p.Chaos.Solver.Sleep.PartialIslandSleep.MinimumMotionlessRatio"), false);
+			check(CVarMinMotionlessRatio && CVarMinMotionlessRatio->IsVariableFloat());
+			MinimumMotionlessThreshold_EditorDefault = CVarMinMotionlessRatio->GetFloat();
+			CVarMinMotionlessRatio->Set(MinimumMotionlessThreshold_UnitTestDefault);
+
 			CVarLinearWakeThresholdMultiplier = IConsoleManager::Get().FindConsoleVariable(TEXT("p.Chaos.Solver.Sleep.PartialIslandSleep.LinearWakeThresholdMultiplier"), false);
 			check(CVarLinearWakeThresholdMultiplier && CVarLinearWakeThresholdMultiplier->IsVariableFloat());
 			LinearWakeThresholdMultiplier_EditorDefault = CVarLinearWakeThresholdMultiplier->GetFloat();
@@ -63,6 +68,7 @@ namespace ChaosTest
 			CVarMomentumPropagation->Set(MomentumPropagation_EditorDefault);
 			CVarPartialSleepCollisionConstraintsOnly->Set(bPartialSleepCollisionConstraintsOnly_EditorDefault);
 			CVarPartialWakePreIntegration->Set(bPartialWakePreIntegration_EditorDefault);
+			CVarMinMotionlessRatio->Set(MinimumMotionlessThreshold_EditorDefault);
 			CVarLinearWakeThresholdMultiplier->Set(LinearWakeThresholdMultiplier_EditorDefault);
 			CVarAngularWakeThresholdMultiplier->Set(AngularWakeThresholdMultiplier_EditorDefault);
 		}
@@ -72,6 +78,7 @@ namespace ChaosTest
 		FRealSingle MomentumPropagation_UnitTestDefault = 0.0f;
 		bool bPartialSleepCollisionConstraintsOnly_UnitTestDefault = false;
 		bool bPartialWakePreIntegration_UnitTestDefault = false;
+		FRealSingle MinimumMotionlessThreshold_UnitTestDefault = 0.0f;
 		FRealSingle LinearWakeThresholdMultiplier_UnitTestDefault = 1.0f;
 		FRealSingle AngularWakeThresholdMultiplier_UnitTestDefault = 1.0f;
 
@@ -80,6 +87,7 @@ namespace ChaosTest
 		FRealSingle MomentumPropagation_EditorDefault;
 		bool bPartialSleepCollisionConstraintsOnly_EditorDefault;
 		bool bPartialWakePreIntegration_EditorDefault;
+		FRealSingle MinimumMotionlessThreshold_EditorDefault;
 		FRealSingle LinearWakeThresholdMultiplier_EditorDefault;
 		FRealSingle AngularWakeThresholdMultiplier_EditorDefault;
 
@@ -89,6 +97,7 @@ namespace ChaosTest
 		IConsoleVariable* CVarMomentumPropagation;
 		IConsoleVariable* CVarPartialSleepCollisionConstraintsOnly;
 		IConsoleVariable* CVarPartialWakePreIntegration;
+		IConsoleVariable* CVarMinMotionlessRatio;
 		IConsoleVariable* CVarLinearWakeThresholdMultiplier;
 		IConsoleVariable* CVarAngularWakeThresholdMultiplier;
 	};
diff --git a/Engine/Source/Runtime/Experimental/Chaos/Private/Chaos/Island/IslandManager.cpp b/Engine/Source/Runtime/Experimental/Chaos/Private/Chaos/Island/IslandManager.cpp
--- a/Engine/Source/Runtime/Experimental/Chaos/Private/Chaos/Island/IslandManager.cpp
+++ b/Engine/Source/Runtime/Experimental/Chaos/Private/Chaos/Island/IslandManager.cpp
@@ -90,6 +90,14 @@ namespace Chaos::CVars
 	bool bChaosSolverPartialWakePreIntegration = true;
 	FAutoConsoleVariableRef CVarChaosSolverPartialWakePreIntegration(TEXT("p.Chaos.Solver.Sleep.PartialIslandSleep.PartialWakePreIntegration"), bChaosSolverPartialWakePreIntegration, TEXT("[EXPERIMENTAL] Safeguard only: Partial wake events also happen pre-integration when the user explicilty sets a particle transform or applies external impulses or forces.[def:true]"));
 
+	/** Cvar to set the minimum ratio of particles which needs to be motionless to activate partial island sleep. */
+	FRealSingle ChaosSolverMinMotionlessRatio = 0.1f;
+	FAutoConsoleVariableRef CVarChaosSolverMinMotionlessRatio(TEXT("p.Chaos.Solver.Sleep.PartialIslandSleep.MinimumMotionlessRatio"), ChaosSolverMinMotionlessRatio, TEXT("[EXPERIMENTAL] Ratio of the particles which need to be motionless in an island for partial island sleep to become active.[def:0.1]"));
+
+	/** Cvar to set the minimum number of ticks the ratio of motionless particles needs to be exceeded before switching between partial and full island sleep. */
+	int32 ChaosSolverSleepModeSwitchThreshold = 10;
+	FAutoConsoleVariableRef CVarChaosSolverSleepModeSwitchThreshold(TEXT("p.Chaos.Solver.Sleep.PartialIslandSleep.SleepModeSwitchCounterThreshold"), ChaosSolverSleepModeSwitchThreshold, TEXT("[EXPERIMENTAL] Safeguard only: Number of steps MinimumMotionlessRatio needs to be satisfied before switching the between full and partial island sleep.[def:10]"));
+
 	/** Cvar to adjust the linear wake threshold for partial waking */
 	FRealSingle LinearPartialWakeThresholdMultiplier = 10.0f;
 	FAutoConsoleVariableRef CVarChaosSolverLinearWakeThresholdMultiplier(TEXT("p.Chaos.Solver.Sleep.PartialIslandSleep.LinearWakeThresholdMultiplier"), LinearPartialWakeThresholdMultiplier, TEXT("[EXPERIMENTAL] A multiplier applied to LinearSleepThreshold for particle waking.[def:10]"));
@@ -97,7 +105,7 @@ namespace Chaos::CVars
 	/** Cvar to adjust the angular wake threshold for partial waking */
 	FRealSingle AngularPartialWakeThresholdMultiplier = 10.0f;
 	FAutoConsoleVariableRef CVarChaosSolverAngularWakeThresholdMultiplier(TEXT("p.Chaos.Solver.Sleep.PartialIslandSleep.AngularWakeThresholdMultiplier"), AngularPartialWakeThresholdMultiplier, TEXT("[EXPERIMENTAL] A multiplier applied to AngularSleepThreshold for particle waking.[def:10]"));
-
+	
 	/** Cvar to override the sleep counter threshold if necessary */
 	int32 ChaosSolverCollisionDefaultSleepCounterThreshold = 20;
 	FAutoConsoleVariableRef CVarChaosSolverCollisionDefaultSleepCounterThreshold(TEXT("p.Chaos.Solver.Sleep.Defaults.SleepCounterThreshold"), ChaosSolverCollisionDefaultSleepCounterThreshold, TEXT("Default counter threshold for sleeping.[def:20]"));
@@ -498,8 +506,9 @@ namespace Chaos::Private
 	void FPBDIsland::Reuse()
 	{
 		Flags.bIsSleepAllowed = true;
-		Flags.bIsPartialSleepAllowed = true;
+		Flags.bIsPartialSleepAllowedInAllNodes = true;
 		Flags.bHasSleepingNode = false;
+		Flags.bEnoughMotionlessNodesForPartialSleep = true;
 		Flags.bIsSleeping = true;
 	}
 
@@ -539,6 +548,11 @@ namespace Chaos::Private
 		}
 	}
 
+	bool FPBDIsland::IsPartialSleepEnabled() const
+	{
+		return CVars::bChaosSolverPartialIslandSleep && !!Flags.bEnoughMotionlessNodesForPartialSleep && !!Flags.bIsPartialSleepAllowedInAllNodes;
+	}
+
 	//////////////////////////////////////////////////////////////////////////////////////////////////
 	//////////////////////////////////////////////////////////////////////////////////////////////////
 	// 
@@ -844,7 +858,7 @@ namespace Chaos::Private
 			{
 				// We permit sleeping for this island if partial island sleeping and pre-integration waking are enabled.
 				// This will allows us to perform partial island waking if a dynamic particle has been woken by the user or had its transform set explicitly.
-				if (CVars::bChaosSolverPartialIslandSleep && CVars::bChaosSolverPartialWakePreIntegration)
+				if (Node->Island->IsPartialSleepEnabled() && CVars::bChaosSolverPartialWakePreIntegration)
 				{
 					bIsSleepAllowed = true;
 				}
@@ -1663,7 +1677,7 @@ namespace Chaos::Private
 		const FPBDRigidParticleHandle* Rigid = Particle->CastToRigidParticle();
 		if (Island != nullptr && Rigid != nullptr)
 		{
-			Island->Flags.bIsPartialSleepAllowed = (Island->Flags.bIsPartialSleepAllowed != 0) && Rigid->PartialIslandSleepAllowed();
+			Island->Flags.bIsPartialSleepAllowedInAllNodes = (Island->Flags.bIsPartialSleepAllowedInAllNodes != 0) && Rigid->PartialIslandSleepAllowed();
 		}
 	}
 
@@ -1671,7 +1685,7 @@ namespace Chaos::Private
 	{
 		for (FPBDIsland* Island : Islands)
 		{
-			Island->Flags.bIsPartialSleepAllowed = true;
+			Island->Flags.bIsPartialSleepAllowedInAllNodes = true;
 		}
 	}
 
@@ -1944,7 +1958,7 @@ namespace Chaos::Private
 			ParentIsland->Flags.bItemsAdded = true;
 			ParentIsland->Flags.bItemsRemoved = !!ParentIsland->Flags.bItemsRemoved || !!ChildIsland->Flags.bItemsRemoved;
 			ParentIsland->Flags.bIsSleepAllowed = !!ParentIsland->Flags.bIsSleepAllowed && !!ChildIsland->Flags.bIsSleepAllowed;
-			ParentIsland->Flags.bIsPartialSleepAllowed = !!ParentIsland->Flags.bIsPartialSleepAllowed && !!ChildIsland->Flags.bIsPartialSleepAllowed;
+			ParentIsland->Flags.bIsPartialSleepAllowedInAllNodes = !!ParentIsland->Flags.bIsPartialSleepAllowedInAllNodes && !!ChildIsland->Flags.bIsPartialSleepAllowedInAllNodes;
 			ParentIsland->Flags.bHasSleepingNode = !!ParentIsland->Flags.bHasSleepingNode || !!ChildIsland->Flags.bHasSleepingNode;
 			ParentIsland->Flags.bIsSleeping = !!ParentIsland->Flags.bIsSleeping && !!ChildIsland->Flags.bIsSleeping;
 			ParentIsland->Flags.bWasSleeping = !!ParentIsland->Flags.bWasSleeping || !!ChildIsland->Flags.bWasSleeping;
@@ -2042,7 +2056,7 @@ namespace Chaos::Private
 
 		// Reset partial sleep allowed flag
 		// Will be updated when the island is repopulated (call to UpdateIslandPartialSleepAllowed) 
-		CurrentIsland->Flags.bIsPartialSleepAllowed = true;
+		CurrentIsland->Flags.bIsPartialSleepAllowedInAllNodes = true;
 
 		// Reset has sleeping node flag
 		// Will be updated when the island is repopulated (call to UpdateIslandHasSleepingNodes) 
@@ -2285,7 +2299,7 @@ namespace Chaos::Private
 			// @todo(chaos) : (island sleeping only) we should only need to do this when the state changes... ?
 			if (!Island->Flags.bIsSleeping || (Island->Flags.bIsSleeping != Island->Flags.bWasSleeping) || !!Island->Flags.bCheckSleep)
 			{
-				if (CVars::bChaosSolverPartialIslandSleep && !!Island->Flags.bIsSleepAllowed && !!Island->Flags.bIsPartialSleepAllowed && 
+				if (Island->IsPartialSleepEnabled() && !!Island->Flags.bIsSleepAllowed &&
 					(CVars::bChaosSolverForcePartialIslandWake || !!Island->Flags.bHasSleepingNode))
 				{
 					if (!Island->Flags.bIsSleeping)
@@ -2404,7 +2418,7 @@ namespace Chaos::Private
 
 				// @todo(Chaos): Island->Flags.bHasSleepingNode is always false here since it's reset in ProcessSleep and will be repopulated in CreateConstraintGraph.
 				// Since this code is rarely executed, we tolerate this and just do partial wakes even though it's not necessary (i.e. flag is not used here).
-				if (CVars::bChaosSolverPartialIslandSleep && CVars::bChaosSolverPartialWakePreIntegration && !!Island->Flags.bIsSleepAllowed && !!Island->Flags.bIsPartialSleepAllowed)
+				if (Island->IsPartialSleepEnabled() && CVars::bChaosSolverPartialWakePreIntegration && !!Island->Flags.bIsSleepAllowed)
 				{
 					if (!Island->Flags.bIsSleeping)
 					{
@@ -2427,6 +2441,11 @@ namespace Chaos::Private
 					// constraint sleep state matches the particle sleep state. Ideally we would only
 					// do this if we know that they don't match, but that's hard to determine.
 					PropagateIslandSleep(Island);
+
+					if (!Island->Flags.bIsPartialSleepAllowedInAllNodes)
+					{
+						ResetParticleSleepCounter(Island);
+					}
 				}
 
 				Island->Flags.bCheckSleep = false;
@@ -2453,7 +2472,14 @@ namespace Chaos::Private
 		{
 			if (!Island->Flags.bIsSleeping && !!Island->Flags.bIsSleepAllowed && !Island->Flags.bIsUsingCache)
 			{
-				if (!CVars::bChaosSolverPartialIslandSleep || !Island->Flags.bIsPartialSleepAllowed)
+				if (Island->IsPartialSleepEnabled())
+				{
+					// Update the sleep state based on particle movement
+					// Deactivate non-moving particles
+					// Deactivate constraints between sleeping particles
+					ProcessPartialIslandSleep(Island, Dt);
+				}
+				else
 				{
 					// Update the sleep state based on particle movement etc
 					ProcessIslandSleep(Island, Dt);
@@ -2463,13 +2489,6 @@ namespace Chaos::Private
 						PropagateIslandSleep(Island);
 					}
 				}
-				else
-				{
-					// Update the sleep state based on particle movement
-					// Deactivate non-moving particles
-					// Deactivate constraints between sleeping particles
-					ProcessPartialIslandSleep(Island, Dt);
-				}
 			}
 
 			// This gets set to false again next tick if there are moving kinematics in the island
@@ -2539,6 +2558,7 @@ namespace Chaos::Private
 	{
 		bool bWithinSleepThreshold = true;
 		int32 SleepCounterThreshold = 0;
+		int32 MotionlessCount = 0;
 		for (FPBDIslandParticle* Node : Island->GetParticles())
 		{
 			// All zeroes means never sleep/disable
@@ -2562,12 +2582,18 @@ namespace Chaos::Private
 				|| (Rigid->WSmooth().SizeSquared() > Node->SleepAngularThresholdSq))
 			{
 				bWithinSleepThreshold = false;
+				Rigid->SetSleepCounter(int8(0));
 
 				// NOTE: We will not sleep if any particle exceeds the threshold, so we could "break" here.
 				// However we still want to update the SleepMetrics for all particles because we want to
 				// update the smoothed velocity based on current state, so we must continue to remaining particles
 				continue;
 			}
+			++MotionlessCount;
+
+			// Particle sleep counter needs to be updated for particles to be able to sleep if we switch to partial island sleep
+			const int32 SleepCounter = FMath::Min(int32(Rigid->SleepCounter()) + 1, Node->SleepCounterThreshold);
+			Rigid->SetSleepCounter(int8(FMath::Min(SleepCounter, TNumericLimits<int8>::Max())));
 
 			// Take the longest sleep time
 			SleepCounterThreshold = FMath::Max(SleepCounterThreshold, Node->SleepCounterThreshold);
@@ -2585,6 +2611,8 @@ namespace Chaos::Private
 			}
 		}
 		Island->SleepCounter = SleepCounter;
+
+		SwitchIslandSleepType(Island, MotionlessCount);
 	}
 
 	void FPBDIslandManager::PropagateIslandSleep(FPBDIsland* Island)
@@ -2686,6 +2714,7 @@ namespace Chaos::Private
 	void FPBDIslandManager::SleepNonMovingParticles(FPBDIsland* Island, const FRealSingle Dt)
 	{
 		bool bAllNodesSleeping = true;
+		int32 MotionlessCount = 0;
 
 		TArray<FGeometryParticleHandle*> SleptParticles;
 		for (FPBDIslandParticle* Node : Island->GetParticles())
@@ -2700,6 +2729,7 @@ namespace Chaos::Private
 			// No need to check sleeping nodes again
 			if (!!Node->Flags.bIsSleeping)
 			{
+				++MotionlessCount;
 				continue;
 			}
 
@@ -2727,6 +2757,8 @@ namespace Chaos::Private
 				}
 				else // motionless
 				{
+					++MotionlessCount;
+
 					// If we get here we might want to sleep
 					// Update the sleep counter and sleep if we exceed it
 					static_assert(sizeof(decltype(Rigid->SleepCounter())) == 1, "Expected int8 for SleepCounter(). Update clamp below");
@@ -2745,16 +2777,25 @@ namespace Chaos::Private
 			}
 			Rigid->SetSleepCounter(int8(SleepCounter));
 		}
+
 		// Flag the island if all particles are sleeping
 		if (bAllNodesSleeping)
 		{
 			Island->Flags.bIsSleeping = true;
+			Island->SleepCounter = 0;
+		}
+		else if (MotionlessCount < Island->GetNumParticles())
+		{
+			// Reset island sleep counter if there are any moving particles
+			Island->SleepCounter = 0;
 		}
 
 		if (!SleptParticles.IsEmpty())
 		{
 			Particles.DeactivateParticles(SleptParticles);
 		}
+
+		SwitchIslandSleepType(Island, MotionlessCount);
 	}
 
 	void FPBDIslandManager::RewakeConstraintsNotSupportingPartialSleep(FPBDIsland* Island)
@@ -3005,6 +3046,40 @@ namespace Chaos::Private
 		WakeUpParticles(WakeUpQueue, /*bResetSleepCounter=*/false);
 	}
 
+	void FPBDIslandManager::SwitchIslandSleepType(FPBDIsland* Island, const int32 MotionlessCount)
+	{
+		if (Island->GetNumParticles() > 0)
+		{
+			const FReal RatioMotionless = MotionlessCount / (FReal)Island->GetNumParticles();
+
+			const bool bIsAboveMinMotionlessRatio = RatioMotionless > CVars::ChaosSolverMinMotionlessRatio;
+			const bool bExceedSleepCounterSwitchThreshold = Island->SleepTypeSwitchCounter >= CVars::ChaosSolverSleepModeSwitchThreshold;
+
+			const bool bPotentialSwitchToPartialIslandSleep = !Island->Flags.bEnoughMotionlessNodesForPartialSleep && bIsAboveMinMotionlessRatio;
+			const bool bPotentialSwitchToFullIslandSleep = !!Island->Flags.bEnoughMotionlessNodesForPartialSleep && !bIsAboveMinMotionlessRatio;
+
+			// Switch sleep type: flip the switch and reset the counter
+			const bool bSwitchSleepType = (bPotentialSwitchToPartialIslandSleep || bPotentialSwitchToFullIslandSleep) && bExceedSleepCounterSwitchThreshold;
+
+			// Increase counter: potential switch but not enough ticks
+			const bool bIncreaseCounter = (bPotentialSwitchToPartialIslandSleep || bPotentialSwitchToFullIslandSleep) && !bExceedSleepCounterSwitchThreshold;
+
+			if (bSwitchSleepType)
+			{
+				Island->Flags.bEnoughMotionlessNodesForPartialSleep = !Island->Flags.bEnoughMotionlessNodesForPartialSleep ? true : false;
+				Island->SleepTypeSwitchCounter = 0;
+			}
+			else if (bIncreaseCounter)
+			{
+				++(Island->SleepTypeSwitchCounter);
+			}
+			else // Reset counter
+			{
+				Island->SleepTypeSwitchCounter = 0;
+			}
+		}
+	}
+
 	// If (bProcessEdgeBasedWake), we first process all updated nodes and triggered wake events, then we process all new edges and wake events.
 	// If (!bProcessEdgeBasedWake), we only process updated nodes and their wake events.
 	void FPBDIslandManager::ProcessPartialIslandWake(FPBDIsland* Island, const bool bProcessEdgeBasedWake)
@@ -3717,7 +3792,7 @@ namespace Chaos::Private
 				CHAOS_CONSTRAINTGRAPH_VALIDATE_TEST(Edge->Island->ContainerEdges[Edge->ContainerIndex][Edge->IslandArrayIndex] == Edge);
 
 				// Make sure the island and edge sleep states match
-				if (!CVars::bChaosSolverPartialIslandSleep && Edge->Constraint->SupportsSleeping())
+				if (!Edge->Island->IsPartialSleepEnabled() && Edge->Constraint->SupportsSleeping())
 				{
 					CHAOS_CONSTRAINTGRAPH_VALIDATE_TEST(Edge->Island->Flags.bIsSleeping == Edge->Flags.bIsSleeping);
 				}
diff --git a/Engine/Source/Runtime/Experimental/Chaos/Public/Chaos/Island/IslandManager.h b/Engine/Source/Runtime/Experimental/Chaos/Public/Chaos/Island/IslandManager.h
--- a/Engine/Source/Runtime/Experimental/Chaos/Public/Chaos/Island/IslandManager.h
+++ b/Engine/Source/Runtime/Experimental/Chaos/Public/Chaos/Island/IslandManager.h
@@ -342,6 +342,8 @@ namespace Chaos::Private
 			Iterations = InIterations;
 		}
 
+		bool IsPartialSleepEnabled() const;
+
 	private:
 		friend class FPBDIslandManager;
 
@@ -358,8 +360,9 @@ namespace Chaos::Private
 				uint32 bItemsAdded : 1;
 				uint32 bItemsRemoved : 1;
 				uint32 bIsSleepAllowed : 1;
-				uint32 bIsPartialSleepAllowed : 1;
+				uint32 bIsPartialSleepAllowedInAllNodes : 1;
 				uint32 bHasSleepingNode : 1;
+				uint32 bEnoughMotionlessNodesForPartialSleep : 1;
 				uint32 bIsSleeping : 1;
 				uint32 bWasSleeping : 1;
 				uint32 bCheckSleep : 1;
@@ -396,6 +399,10 @@ namespace Chaos::Private
 		// When DisableCounter hits a threshold the island particle will be disabled
 		int32 DisableCounter = 0;
 
+		// When SleepTypeSwitchCounter hits a threshold, the island sleep type (partial or full) 
+		// will change if CVars::bChaosSolverPartialIslandSleep == true.
+		int32 SleepTypeSwitchCounter = 0;
+
 		int32 ResimFrame = INDEX_NONE;
 
 		FFlags Flags;
@@ -623,6 +630,7 @@ namespace Chaos::Private
 		CHAOS_API void RewakeConstraintsNotSupportingPartialSleep(FPBDIsland* Island);
 		CHAOS_API void PropagateParticleSleepToConstraints(FPBDIsland* Island);
 		CHAOS_API void PostStepPartialWake(FPBDIsland* Island);
+		CHAOS_API void SwitchIslandSleepType(FPBDIsland* Island, const int32 MotionlessCount);
 
 		// Partial Island Waking
 		CHAOS_API void ProcessPartialIslandWake(FPBDIsland* Island, const bool bProcessEdgeBasedWake);
