Also removed completion from the cancels slot mask.

#rnx

diff --git a/Engine/Source/Runtime/Experimental/IoStore/HttpClient/Private/Loop.inl b/Engine/Source/Runtime/Experimental/IoStore/HttpClient/Private/Loop.inl
--- a/Engine/Source/Runtime/Experimental/IoStore/HttpClient/Private/Loop.inl
+++ b/Engine/Source/Runtime/Experimental/IoStore/HttpClient/Private/Loop.inl
@@ -1162,9 +1162,9 @@ uint32 FEventLoop::FImpl::Tick(int32 PollTimeoutMs)
 		PrevFreeSlots += ReturnedSlots;
 	}
 
-	if (CancelsLoad)
+	if (uint32 Mask = CancelsLoad | ReturnedSlots; Mask)
 	{
-		Cancels.fetch_and(~CancelsLoad, std::memory_order_relaxed);
+		Cancels.fetch_and(~Mask, std::memory_order_relaxed);
 	}
 
 	return BusyCount + BusyBias;
