Incorrect integer size.

#rnx

diff --git a/Engine/Source/Runtime/Experimental/IoStore/HttpClient/Private/Loop.inl b/Engine/Source/Runtime/Experimental/IoStore/HttpClient/Private/Loop.inl
--- a/Engine/Source/Runtime/Experimental/IoStore/HttpClient/Private/Loop.inl
+++ b/Engine/Source/Runtime/Experimental/IoStore/HttpClient/Private/Loop.inl
@@ -1162,7 +1162,7 @@ uint32 FEventLoop::FImpl::Tick(int32 PollTimeoutMs)
 		PrevFreeSlots += ReturnedSlots;
 	}
 
-	if (uint32 Mask = CancelsLoad | ReturnedSlots; Mask)
+	if (uint64 Mask = CancelsLoad | ReturnedSlots; Mask)
 	{
 		Cancels.fetch_and(~Mask, std::memory_order_relaxed);
 	}
