OnDemandContentInstaller
- Add CVar_ContentInstallerFlushOnIdle which flushes any pending chunks when the installer runs out of tasks.  Default to disabling this behavior because current access patterns produce too many flushes when its enabled.
- Fix TryExecuteNextRequest not getting called from CompleteVerifyRequest and CompleteFlushLastAccessRequest.

diff --git a/Engine/Source/Runtime/Experimental/IoStore/OnDemand/Private/OnDemandContentInstaller.cpp b/Engine/Source/Runtime/Experimental/IoStore/OnDemand/Private/OnDemandContentInstaller.cpp
--- a/Engine/Source/Runtime/Experimental/IoStore/OnDemand/Private/OnDemandContentInstaller.cpp
+++ b/Engine/Source/Runtime/Experimental/IoStore/OnDemand/Private/OnDemandContentInstaller.cpp
@@ -27,7 +27,7 @@ namespace UE::IoStore
 	namespace CVars
 	{
 #if !UE_BUILD_SHIPPING
-		FString IoStoreErrorOnRequest = "";
+		static FString IoStoreErrorOnRequest = "";
 		static FAutoConsoleVariableRef CVar_IoStoreErrorOnRequest(
 			TEXT("iostore.ErrorOnRequest"),
 			IoStoreErrorOnRequest,
@@ -41,6 +41,13 @@ namespace UE::IoStore
 			TEXT("When the request with a debug name partially matching this cvar is found iostore will error with a random error.")
 		);
 #endif
+
+		static bool bContentInstallerFlushOnIdle = false;
+		static FAutoConsoleVariableRef CVar_ContentInstallerFlushOnIdle(
+			TEXT("iostore.ContentInstallerFlushOnIdle"), 
+			bContentInstallerFlushOnIdle,
+			TEXT("When the content installer runs out of tasks, it will flush any pending iochunks to disk.")
+			);
 	}
 namespace Private
 {
@@ -439,10 +446,34 @@ void FOnDemandContentInstaller::TryExecuteNextRequest()
 	if (NextRequest != nullptr)
 	{
 		InstallerPipe.Launch(
-			TEXT("ExecuteRequest"),
+			TEXT("OnDemandContentInstaller::ExecuteRequest"),
 			[this, NextRequest] { ExecuteRequest(*NextRequest); },
 			UE::Tasks::ETaskPriority::BackgroundLow);
 	}
+	else if (UE::IoStore::CVars::bContentInstallerFlushOnIdle)
+	{
+		// Flush any pending chunks, but only if the installer is otherwise idle
+		InstallerPipe.Launch(
+			TEXT("OnDemandContentInstaller::Flush"),
+			[this] 
+			{
+				bool bExecuteFlush = false;
+				{
+					TUniqueLock Lock(Mutex);
+					bExecuteFlush = RequestQueue.IsEmpty() && CurrentRequest == nullptr;
+				}
+
+				if (bExecuteFlush)
+				{
+					FResult Result = IoStore.InstallCache->Flush();
+					if (Result.HasError())
+					{
+						UE_LOG(LogIoStoreOnDemand, Error, TEXT("Failed to flush pending chunks, Reason %s"), *LexToString(Result.GetError()));
+					}
+				}
+			},
+			UE::Tasks::ETaskPriority::BackgroundLow);
+	}
 }
 
 void FOnDemandContentInstaller::ExecuteRequest(FRequest& Request)
@@ -1300,6 +1331,8 @@ void FOnDemandContentInstaller::CompleteVerifyRequest(FRequest& Request)
 	}
 	DestroyRequest(Request);
 
+	TryExecuteNextRequest();
+
 	if (!OnCompleted)
 	{
 		return;
@@ -1329,6 +1362,8 @@ void FOnDemandContentInstaller::CompleteFlushLastAccessRequest(FRequest& Request
 	}
 	DestroyRequest(Request);
 
+	TryExecuteNextRequest();
+
 	if (!OnCompleted)
 	{
 		return;
