OnDemandContentInstaller
- Fix flush on idle check.  Still disabled

diff --git a/Engine/Source/Runtime/Experimental/IoStore/OnDemand/Private/OnDemandContentInstaller.cpp b/Engine/Source/Runtime/Experimental/IoStore/OnDemand/Private/OnDemandContentInstaller.cpp
--- a/Engine/Source/Runtime/Experimental/IoStore/OnDemand/Private/OnDemandContentInstaller.cpp
+++ b/Engine/Source/Runtime/Experimental/IoStore/OnDemand/Private/OnDemandContentInstaller.cpp
@@ -434,9 +434,11 @@ void FOnDemandContentInstaller::TryExecuteNextRequest()
 	}
 
 	FRequest* NextRequest = nullptr;
+	bool bRequestQueueIsEmpty = false;
 	{
 		TUniqueLock Lock(Mutex);
-		if (CurrentRequest == nullptr && RequestQueue.IsEmpty() == false)
+		bRequestQueueIsEmpty = RequestQueue.IsEmpty();
+		if (CurrentRequest == nullptr && bRequestQueueIsEmpty == false)
 		{
 			RequestQueue.HeapPop(NextRequest, RequestSortPredicate, EAllowShrinking::No);
 			CurrentRequest = NextRequest;
@@ -450,7 +452,7 @@ void FOnDemandContentInstaller::TryExecuteNextRequest()
 			[this, NextRequest] { ExecuteRequest(*NextRequest); },
 			UE::Tasks::ETaskPriority::BackgroundLow);
 	}
-	else if (UE::IoStore::CVars::bContentInstallerFlushOnIdle)
+	else if (bRequestQueueIsEmpty && UE::IoStore::CVars::bContentInstallerFlushOnIdle)
 	{
 		// Flush any pending chunks, but only if the installer is otherwise idle
 		InstallerPipe.Launch(
