Chaos : fix force based fields applied directly to geometry collection component not woring properly
- the array storing the geometry collection proxies affected by fields in the solver was cleared after exeuting code for the parameter fields
- this caused the force based one in the step iteration loop to have no proxies to process for force based fields
- fixed by removing the first clearing of the proxy array
- this is knock on from optimization CL 24688493

diff --git a/Engine/Source/Runtime/Experimental/Chaos/Private/PBDRigidsSolver.cpp b/Engine/Source/Runtime/Experimental/Chaos/Private/PBDRigidsSolver.cpp
--- a/Engine/Source/Runtime/Experimental/Chaos/Private/PBDRigidsSolver.cpp
+++ b/Engine/Source/Runtime/Experimental/Chaos/Private/PBDRigidsSolver.cpp
@@ -482,7 +482,6 @@ namespace Chaos
 				{
 					GeoclObj->FieldParameterUpdateCallback(MSolver);
 				}
-				MSolver->GetGeometryCollectionPhysicsProxiesField_Internal().Reset();
 
 				MSolver->GetEvolution()->GetBroadPhase().GetIgnoreCollisionManager().ProcessPendingQueues(*MSolver);
 			}
@@ -544,6 +543,7 @@ namespace Chaos
 					{
 						GeoCollectionObj->FieldForcesUpdateCallback(MSolver);
 					}
+					// we only want to process once per step because we don't want to apply the same forces multiple time during this solver tick
 					MSolver->GetGeometryCollectionPhysicsProxiesField_Internal().Reset();
 
 					if(FRewindData* RewindData = MSolver->GetRewindData())
