[ChaosVD] Modifications to support the secure socket relay transport mode.

As part of this change we now :

- Broadcast the trace connection details when they are needed instead of every second with the rest of the recording status. This reduces the bandwidth usage in message bus
- Deprecated fields no longer needed (and directly removed ones we added during 5.7 development but will not be used)
- Removed the Message Bus relay transport implementation. This was recently submitted as a prototype and last resort implementation in case the direct TLS transport mode could not be implemented. As that mode is already in development,
I decided to remove the bus relay transport implementation as keeping it working would require make some compromises (like support writes in the game thread in the ITraceRelayTransport interface).

diff --git a/Engine/Plugins/ChaosVD/Source/ChaosVD/Private/Trace/ChaosVDTraceManager.cpp b/Engine/Plugins/ChaosVD/Source/ChaosVD/Private/Trace/ChaosVDTraceManager.cpp
--- a/Engine/Plugins/ChaosVD/Source/ChaosVD/Private/Trace/ChaosVDTraceManager.cpp
+++ b/Engine/Plugins/ChaosVD/Source/ChaosVD/Private/Trace/ChaosVDTraceManager.cpp
@@ -120,6 +120,12 @@ namespace Chaos::VD
 #endif
 	}
 
+	void FRelayDataStream::EnqueueRelayedData(const TConstArrayView<uint8> InTraceDataBuffer)
+	{
+		FSerializedDataBuffer DataBuffer(InTraceDataBuffer);
+		DataQueue.Enqueue(MoveTemp(DataBuffer));
+	}
+
 	void FRelayDataStream::EnqueueRelayedData(TArray<uint8>&& InTraceDataBuffer)
 	{
 		FSerializedDataBuffer DataBuffer(MoveTemp(InTraceDataBuffer));
diff --git a/Engine/Plugins/ChaosVD/Source/ChaosVD/Private/Widgets/SChaosVDRecordingControls.cpp b/Engine/Plugins/ChaosVD/Source/ChaosVD/Private/Widgets/SChaosVDRecordingControls.cpp
--- a/Engine/Plugins/ChaosVD/Source/ChaosVD/Private/Widgets/SChaosVDRecordingControls.cpp
+++ b/Engine/Plugins/ChaosVD/Source/ChaosVD/Private/Widgets/SChaosVDRecordingControls.cpp
@@ -20,6 +20,7 @@
 #include "ToolMenu.h"
 #include "ToolMenus.h"
 #include "Chaos/ChaosVDEngineEditorBridge.h"
+#include "Chaos/ChaosVDTraceRelayTransport.h"
 #include "Framework/MultiBox/MultiBoxBuilder.h"
 #include "Misc/ScopedSlowTask.h"
 #include "Settings/ChaosVDGeneralSettings.h"
@@ -553,11 +554,11 @@ void SChaosVDRecordingControls::HandleRecordingStop(TWeakPtr<FChaosVDSessionInfo
 
 	FString TraceTarget;
 
-	switch (SessionInfoPtr->LastKnownRecordingState.TraceDetails.TransportMode)
+	switch (SessionInfoPtr->LastKnownConnectionDetails.TransportMode)
 	{
 		case EChaosVDTransportMode::FileSystem:
 			{
-				TraceTarget = SessionInfoPtr->LastKnownRecordingState.TraceDetails.TraceTarget;
+				TraceTarget = SessionInfoPtr->LastKnownConnectionDetails.TraceTarget;
 				break;
 			}
 		case EChaosVDTransportMode::TraceServer:
@@ -566,11 +567,11 @@ void SChaosVDRecordingControls::HandleRecordingStop(TWeakPtr<FChaosVDSessionInfo
 				TraceTarget = FChaosVDModule::Get().GetTraceManager()->GetTraceFileNameFromStoreForSession(ConnectionDetails.SessionAddress, ConnectionDetails.TraceID);
 				break;
 			}
-		case EChaosVDTransportMode::Memory:
+		case EChaosVDTransportMode::Direct:
 		case EChaosVDTransportMode::Relay:
 		case EChaosVDTransportMode::Invalid:
 		default:
-		UE_LOG(LogChaosVDEditor, Error, TEXT("[%hs] Transport mode is not supported | [%s]"), __func__, *UEnum::GetValueAsString(SessionInfoPtr->LastKnownRecordingState.TraceDetails.TransportMode));
+		UE_LOG(LogChaosVDEditor, Error, TEXT("[%hs] Transport mode is not supported | [%s]"), __func__, *UEnum::GetValueAsString(SessionInfoPtr->LastKnownConnectionDetails.TransportMode));
 			break;
 	}
 
@@ -699,7 +700,7 @@ SChaosVDRecordingControls::FTraceStoreConnectionSettings SChaosVDRecordingContro
 		return ConnectionSettings;
 	}
 
-	FChaosVDTraceDetails& RecordingSessionDetails = InSessionInfoRef->LastKnownRecordingState.TraceDetails;
+	FChaosVDTraceDetails& RecordingSessionDetails = InSessionInfoRef->LastKnownConnectionDetails;
 
 	if (const UE::Trace::FStoreClient::FSessionInfo* TraceSessionInfo = FChaosVDTraceManager::GetTraceSessionInfo(RecordingSessionDetails.TraceTarget, RecordingSessionDetails.TraceGuid))
 	{
@@ -813,7 +814,7 @@ void SChaosVDRecordingControls::SetSessionRecordingState(bool bIsRecording, ECha
 			});
 #endif
 
-			if (RecordingParams.TransportMode == EChaosVDTransportMode::Memory)
+			if (RecordingParams.TransportMode == EChaosVDTransportMode::Direct)
 			{
 				// Memory Tracing mode requires opening a trace session in the editor side first before starting the recording
 				// so targets have something to connect to (Usually targets connect to a running trace server)
@@ -843,6 +844,55 @@ void SChaosVDRecordingControls::SetSessionRecordingState(bool bIsRecording, ECha
 				}
 
 				RemoteSessionManager->SendStartRecordingCommand(SessionInfoRef->Address, RecordingParams);
+
+				// Once the start recording command is issued, we can try to connect to the created session (if everything went well).
+				// If it didn't go well, this will take care of update the UI to notify the user
+				ExecuteAsyncConnectionAttemptTaskWithRetry(RemainingRetries, [WeakThis = AsWeak(), SessionInstanceId = SessionInfoRef->InstanceId]()
+				{
+					TSharedPtr<FChaosVDRemoteSessionsManager> RemoteSessionManagerPtr = FChaosVDEngineEditorBridge::Get().GetRemoteSessionsManager();
+					TSharedPtr<FChaosVDSessionInfo> SessionInfoPtr = RemoteSessionManagerPtr ? RemoteSessionManagerPtr->GetSessionInfo(SessionInstanceId).Pin() : nullptr;
+					if (!SessionInfoPtr)
+					{
+						return false;
+					}
+    
+					FChaosVDTraceDetails& RecordingSessionDetails = SessionInfoPtr->LastKnownConnectionDetails;
+
+					// We didn't receive the connection details yet. Trying again later...
+					if (!RecordingSessionDetails.IsValid())
+					{
+						return false;
+					}
+
+					// Check if there is a connection attempt in progress and its results
+					const Chaos::VD::EConnectionAttemptResult ConnectionAttemptStatus = FChaosVDEngineEditorBridge::Get().GetRelayTransportInstance()->GetConnectionAttemptResult(SessionInstanceId);
+					switch (ConnectionAttemptStatus)
+					{
+					case Chaos::VD::InProgress:
+						{
+							// A connection attempt started but we don't have the results yet. Trying again later...
+							return false;
+						}
+					case Chaos::VD::NotStarted:
+						{
+							Chaos::VD::FRelayConnectionInfo ConnectionInfo;
+							ConnectionInfo.Port = RecordingSessionDetails.Port;
+							ConnectionInfo.Address = RecordingSessionDetails.TraceTarget;
+							ConnectionInfo.CertificateAuthority = RecordingSessionDetails.CertAuth;
+							Chaos::VD::EConnectionAttemptResult Result = FChaosVDEngineEditorBridge::Get().GetRelayTransportInstance()->ConnectToRelay(SessionInstanceId, ConnectionInfo);
+
+							return Result == Chaos::VD::EConnectionAttemptResult::Success;
+						}
+					case Chaos::VD::Success:
+						// The current connection attempt succeed, we can stop the async retry loop
+						return true;
+					case Chaos::VD::Failed:
+					default:
+						// There was a connection attempt, but it failed. Abandoning the retry attempt as they will not succeed
+						return false;
+					}
+				},
+				RecordingAttemptFailedCallback);
 			}
 			else
 			{
diff --git a/Engine/Plugins/ChaosVD/Source/ChaosVD/Public/Trace/ChaosVDTraceManager.h b/Engine/Plugins/ChaosVD/Source/ChaosVD/Public/Trace/ChaosVDTraceManager.h
--- a/Engine/Plugins/ChaosVD/Source/ChaosVD/Public/Trace/ChaosVDTraceManager.h
+++ b/Engine/Plugins/ChaosVD/Source/ChaosVD/Public/Trace/ChaosVDTraceManager.h
@@ -54,6 +54,7 @@ namespace Chaos::VD
 
 		virtual ~FRelayDataStream() override;
 
+		void EnqueueRelayedData(const TConstArrayView<uint8> InTraceDataBuffer);
 		void EnqueueRelayedData(TArray<uint8>&& InTraceDataBuffer);
 		void EnqueueRelayedData(const FChaosVDRelayTraceDataMessage& InTraceDataMessage);
 
diff --git a/Engine/Source/Runtime/Experimental/Chaos/Public/Chaos/Serialization/SerializedDataBuffer.h b/Engine/Source/Runtime/Experimental/Chaos/Public/Chaos/Serialization/SerializedDataBuffer.h
--- a/Engine/Source/Runtime/Experimental/Chaos/Public/Chaos/Serialization/SerializedDataBuffer.h
+++ b/Engine/Source/Runtime/Experimental/Chaos/Public/Chaos/Serialization/SerializedDataBuffer.h
@@ -10,6 +10,10 @@ namespace Chaos
 	/** Wrapper around a data buffer used for serialization of chaos data */
 	struct FSerializedDataBuffer
 	{
+		FSerializedDataBuffer(const TConstArrayView<uint8> ConstView) : DataBuffer(ConstView.GetData(), ConstView.Num())
+		{	
+		}
+
 		FSerializedDataBuffer(TArray<uint8>&& InByteArray) : DataBuffer(MoveTemp(InByteArray))
 		{	
 		}
diff --git a/Engine/Source/Runtime/Experimental/ChaosSolverEngine/ChaosSolverEngine.Build.cs b/Engine/Source/Runtime/Experimental/ChaosSolverEngine/ChaosSolverEngine.Build.cs
--- a/Engine/Source/Runtime/Experimental/ChaosSolverEngine/ChaosSolverEngine.Build.cs
+++ b/Engine/Source/Runtime/Experimental/ChaosSolverEngine/ChaosSolverEngine.Build.cs
@@ -18,13 +18,16 @@ namespace UnrealBuildTool.Rules
 					"DataflowCore",
 					"DataflowEngine",
 					"DataflowSimulation",
-					"ChaosVDRuntime" // If CVD is compiled out, this only provides access to its structs data types
+					"ChaosVDRuntime", // If CVD is compiled out, this only provides access to its structs data types
 				}
 				);
+			
 
 			SetupModulePhysicsSupport(Target);
 			SetupModuleChaosVisualDebuggerSupport(Target);
 			PrivateDefinitions.Add("CHAOS_INCLUDE_LEVEL_1=1");
+
+			bAllowUETypesInNamespaces = true;
 		}
 	}
 }
diff --git a/Engine/Source/Runtime/Experimental/ChaosSolverEngine/Private/Chaos/ChaosVDBusTraceRelayTransport.cpp b/Engine/Source/Runtime/Experimental/ChaosSolverEngine/Private/Chaos/ChaosVDBusTraceRelayTransport.cpp
deleted file mode 100644
diff --git a/Engine/Source/Runtime/Experimental/ChaosSolverEngine/Private/Chaos/ChaosVDEngineEditorBridge.cpp b/Engine/Source/Runtime/Experimental/ChaosSolverEngine/Private/Chaos/ChaosVDEngineEditorBridge.cpp
--- a/Engine/Source/Runtime/Experimental/ChaosSolverEngine/Private/Chaos/ChaosVDEngineEditorBridge.cpp
+++ b/Engine/Source/Runtime/Experimental/ChaosSolverEngine/Private/Chaos/ChaosVDEngineEditorBridge.cpp
@@ -7,7 +7,7 @@
 #if WITH_CHAOS_VISUAL_DEBUGGER
 
 #include "ChaosVDRuntimeModule.h"
-#include "Chaos/ChaosVDBusTraceRelayTransport.h"
+#include "Chaos/ChaosVDTraceRelayTransport.h"
 #include "ChaosVisualDebugger/ChaosVDTraceMacros.h"
 #include "ChaosVisualDebugger/ChaosVisualDebuggerTrace.h"
 #include "Components/PrimitiveComponent.h"
@@ -21,6 +21,7 @@
 #include "Serialization/MemoryWriter.h"
 #include "UObject/Package.h"
 
+
 bool FChaosVDEngineEditorBridge::AddOnScreenRecordingMessage(float DummyDeltaTime)
 {
 	constexpr bool bContinueLooping = false;
@@ -175,7 +176,6 @@ bool FChaosVDEngineEditorBridge::BroadcastSessionStatus(float DeltaTime)
 	StatusMessage.InstanceId = FApp::GetInstanceId();
 	StatusMessage.bIsRecording = RuntimeModule.IsRecording();
 	StatusMessage.ElapsedTime = RuntimeModule.GetAccumulatedRecordingTime();
-	StatusMessage.TraceDetails = RuntimeModule.GetCurrentTraceSessionDetails();
 
 	RemoteSessionsManager->PublishRecordingStatusUpdate(StatusMessage);
 
@@ -260,39 +260,41 @@ FChaosVDParticleMetadata FChaosVDEngineEditorBridge::GenerateParticleMetadata(co
 	return Metadata;
 }
 
-void FChaosVDEngineEditorBridge::InitializeTraceRelayInstance()
+void FChaosVDEngineEditorBridge::SetExternalTraceRelayInstance(const TSharedPtr<Chaos::VD::ITraceDataRelayTransport>& InExternalTraceRelayInstance)
 {
-	// TODO: Add a way for a CVD extension define the relay instance type
-	TSharedPtr<Chaos::VD::FBusTraceRelayTransport> BusRelayDataTransport = MakeShared<Chaos::VD::FBusTraceRelayTransport>();
-	BusRelayDataTransport->SetMessageBus(RemoteSessionsManager->GetMessageBusInstance());
-	RemoteSessionsManager->OnMessagingInitialized().AddLambda([WeakRelayInstance = BusRelayDataTransport.ToWeakPtr()] (TSharedPtr<IMessageBus> NewMessageBus, TSharedPtr<FMessageEndpoint> NewEndpoint)
-		{
-			if (TSharedPtr<Chaos::VD::FBusTraceRelayTransport> RelayInstance = WeakRelayInstance.Pin())
-			{
-				RelayInstance->SetMessageBus(NewMessageBus);
-			}
-		});
+	if (!ensure(!RelayTraceDataTransportInstance))
+	{
+		UE_LOG(LogChaos, Error, TEXT("An external relay instance was already set. This cannot be changed afterwards. Ignoring set request..."));
+		return;
+	}
 
-	BusRelayDataTransport->Initialize();
-	
-	RelayTraceDataTransportInstance = BusRelayDataTransport;
+	RelayTraceDataTransportInstance = InExternalTraceRelayInstance;
+}
+
+void FChaosVDEngineEditorBridge::HandleTraceConnectionDetailsUpdated()
+{
+	FChaosVDRuntimeModule& RuntimeModule = FChaosVDRuntimeModule::Get();
+
+	FChaosVDTraceConnectionDetailsMessage ConnectionDetailsMessage;
+	ConnectionDetailsMessage.InstanceId = FApp::GetInstanceId();
+	ConnectionDetailsMessage.TraceDetails = RuntimeModule.GetCurrentTraceSessionDetails();
+
+	RemoteSessionsManager->PublishTraceConnectionDetailsUpdate(ConnectionDetailsMessage);
 }
 
 void FChaosVDEngineEditorBridge::Initialize()
 {
 	FChaosVDRuntimeModule& CVDRuntimeModule = FChaosVDRuntimeModule::Get();
 
-	FChaosVDTraceDetails TraceDetails = CVDRuntimeModule.GetCurrentTraceSessionDetails();
-
 	RemoteSessionsManager->Initialize();
 
-	InitializeTraceRelayInstance();
-
 	RecordingStartedHandle = CVDRuntimeModule.RegisterRecordingStartedCallback(FChaosVDRecordingStateChangedDelegate::FDelegate::CreateRaw(this, &FChaosVDEngineEditorBridge::HandleCVDRecordingStarted));
 	PostRecordingStartedHandle = CVDRuntimeModule.RegisterPostRecordingStartedCallback(FChaosVDRecordingStateChangedDelegate::FDelegate::CreateRaw(this, &FChaosVDEngineEditorBridge::HandleCVDPostRecordingStarted));
 	RecordingStoppedHandle = CVDRuntimeModule.RegisterRecordingStopCallback(FChaosVDRecordingStateChangedDelegate::FDelegate::CreateRaw(this, &FChaosVDEngineEditorBridge::HandleCVDRecordingStopped));
 	RecordingStartFailedHandle = CVDRuntimeModule.RegisterRecordingStartFailedCallback(FChaosVDRecordingStartFailedDelegate::FDelegate::CreateRaw(this, &FChaosVDEngineEditorBridge::HandleCVDRecordingStartFailed));
 
+	CVDRuntimeModule.OnTraceConnectionDetailsUpdated().AddRaw(this, &FChaosVDEngineEditorBridge::HandleTraceConnectionDetailsUpdated);
+
 	FChaosVisualDebuggerTrace::RegisterExternalParticleDebugNameGenerator(Chaos::VD::FRecordingSessionState::FParticleMetaDataGeneratorDelegate::CreateStatic(&FChaosVDEngineEditorBridge::GenerateParticleMetadata));
 
 	Chaos::VisualDebugger::FChaosVDDataChannelsManager::Get().OnChannelStateChanged().AddRaw(this, &FChaosVDEngineEditorBridge::HandleDataChannelChanged);
@@ -347,6 +349,8 @@ void FChaosVDEngineEditorBridge::TearDown()
 		CVDRuntimeModule.RemoveRecordingStartFailedCallback(RecordingStartFailedHandle);
 		RecordingStartFailedHandle.Reset();
 
+		CVDRuntimeModule.OnTraceConnectionDetailsUpdated().RemoveAll(this);
+
 		// Make sure of removing the message from the screen in case the recording didn't quite stopped yet
 		if (CVDRuntimeModule.IsRecording())
 		{
diff --git a/Engine/Source/Runtime/Experimental/ChaosSolverEngine/Private/Chaos/ChaosVDRemoteSessionsManager.cpp b/Engine/Source/Runtime/Experimental/ChaosSolverEngine/Private/Chaos/ChaosVDRemoteSessionsManager.cpp
--- a/Engine/Source/Runtime/Experimental/ChaosSolverEngine/Private/Chaos/ChaosVDRemoteSessionsManager.cpp
+++ b/Engine/Source/Runtime/Experimental/ChaosSolverEngine/Private/Chaos/ChaosVDRemoteSessionsManager.cpp
@@ -39,12 +39,12 @@ bool FChaosVDSessionInfo::IsRecording() const
 
 EChaosVDRecordingMode FChaosVDSessionInfo::GetRecordingMode() const
 {
-	return LastKnownRecordingState.TraceDetails.Mode;
+	return LastKnownConnectionDetails.Mode;
 }
 
 bool FChaosVDSessionInfo::IsConnected() const
 {
-	return LastKnownRecordingState.TraceDetails.bIsConnected;
+	return false;
 }
 
 bool FChaosVDMultiSessionInfo::IsRecording() const
@@ -91,8 +91,23 @@ bool FChaosVDMultiSessionInfo::IsConnected() const
 	return bIsAnyConnected;
 }
 
+void FChaosVDRemoteSessionsManager::RegisterBuiltInMessageTypes()
+{
+	SupportedMessageTypes.Emplace(FChaosVDSessionPong::StaticStruct());
+	SupportedMessageTypes.Emplace(FChaosVDRecordingStatusMessage::StaticStruct());
+	SupportedMessageTypes.Emplace(FChaosVDSessionPing::StaticStruct());
+	SupportedMessageTypes.Emplace(FChaosVDStartRecordingCommandMessage::StaticStruct());
+	SupportedMessageTypes.Emplace(FChaosVDStopRecordingCommandMessage::StaticStruct());
+	SupportedMessageTypes.Emplace(FChaosVDChannelStateChangeCommandMessage::StaticStruct());
+	SupportedMessageTypes.Emplace(FChaosVDFullSessionInfoRequestMessage::StaticStruct());
+	SupportedMessageTypes.Emplace(FChaosVDFullSessionInfoResponseMessage::StaticStruct());
+	SupportedMessageTypes.Emplace(FChaosVDChannelStateChangeResponseMessage::StaticStruct());
+	SupportedMessageTypes.Emplace(FChaosVDTraceConnectionDetailsMessage::StaticStruct());
+}
+
 FChaosVDRemoteSessionsManager::FChaosVDRemoteSessionsManager()
 {
+	RegisterBuiltInMessageTypes();
 }
 
 void FChaosVDRemoteSessionsManager::Initialize(const TSharedPtr<IMessageBus>& InMessageBus)
@@ -145,7 +160,8 @@ TSharedPtr<FMessageEndpoint> FChaosVDRemoteSessionsManager::CreateEndPoint(const
 		EndPointBuilder.Handling<FChaosVDSessionPong>(this, &FChaosVDRemoteSessionsManager::HandleSessionPongMessage)
 						.Handling<FChaosVDRecordingStatusMessage>(this, &FChaosVDRemoteSessionsManager::HandleRecordingStatusUpdateMessage)
 						.Handling<FChaosVDFullSessionInfoResponseMessage>(this, &FChaosVDRemoteSessionsManager::HandleFullSessionStateResponseMessage)
-						.Handling<FChaosVDChannelStateChangeResponseMessage>(this, &FChaosVDRemoteSessionsManager::HandleChangeDataChannelStateResponseMessage);
+						.Handling<FChaosVDChannelStateChangeResponseMessage>(this, &FChaosVDRemoteSessionsManager::HandleChangeDataChannelStateResponseMessage)
+						.Handling<FChaosVDTraceConnectionDetailsMessage>(this, &FChaosVDRemoteSessionsManager::HandleConnectionDetailsUpdateMessage);
 	}
 
 	return EndPointBuilder.Build();
@@ -157,6 +173,14 @@ void FChaosVDRemoteSessionsManager::ReInitializeMessagingSystem(const TSharedPtr
 	InitializeMessagingSystem(InMessageBus);
 }
 
+void FChaosVDRemoteSessionsManager::RegisterExternalSupportedMessageType(const UScriptStruct* ScriptStruct)
+{
+	if (ensure(ScriptStruct))
+	{
+		SupportedMessageTypes.Emplace(ScriptStruct);
+	}
+}
+
 void FChaosVDRemoteSessionsManager::InitializeMessagingSystem(const TSharedPtr<IMessageBus>& InMessageBus)
 {
 	if (!ensure(InMessageBus))
@@ -178,6 +202,7 @@ void FChaosVDRemoteSessionsManager::InitializeMessagingSystem(const TSharedPtr<I
 		MessageEndpoint->Subscribe<FChaosVDRecordingStatusMessage>();
 		MessageEndpoint->Subscribe<FChaosVDFullSessionInfoResponseMessage>();
 		MessageEndpoint->Subscribe<FChaosVDChannelStateChangeResponseMessage>();
+		MessageEndpoint->Subscribe<FChaosVDTraceConnectionDetailsMessage>();
 	}
 
 	MessageEndpoint->Subscribe<FChaosVDSessionPing>();
@@ -245,17 +270,13 @@ void FChaosVDRemoteSessionsManager::Shutdown()
 
 void FChaosVDRemoteSessionsManager::EnumerateMessageTypes(const FVisitorFunction& InVisitor)
 {
-	InVisitor(FChaosVDSessionPong::StaticStruct());
-	InVisitor(FChaosVDRecordingStatusMessage::StaticStruct());
-	InVisitor(FChaosVDSessionPing::StaticStruct());
-	InVisitor(FChaosVDStartRecordingCommandMessage::StaticStruct());
-	InVisitor(FChaosVDStopRecordingCommandMessage::StaticStruct());
-	InVisitor(FChaosVDChannelStateChangeCommandMessage::StaticStruct());
-	InVisitor(FChaosVDFullSessionInfoRequestMessage::StaticStruct());
-	InVisitor(FChaosVDFullSessionInfoResponseMessage::StaticStruct());
-	InVisitor(FChaosVDChannelStateChangeResponseMessage::StaticStruct());
-	InVisitor(FChaosVDRelayTraceDataMessage::StaticStruct());
-	InVisitor(FChaosVDRelayTraceStatusMessage::StaticStruct());
+	for (const UScriptStruct* MessageType : SupportedMessageTypes)
+	{
+		if (MessageType)
+		{
+			InVisitor(MessageType);
+		}
+	}
 }
 
 void FChaosVDRemoteSessionsManager::StartSessionDiscovery()
@@ -280,11 +301,19 @@ void FChaosVDRemoteSessionsManager::StopSessionDiscovery()
 	}
 }
 
-void FChaosVDRemoteSessionsManager::PublishRecordingStatusUpdate(const FChaosVDRecordingStatusMessage& UpdateMessage)
+void FChaosVDRemoteSessionsManager::PublishRecordingStatusUpdate(const FChaosVDRecordingStatusMessage& InUpdateMessage)
+{
+	if (ensure(MessageEndpoint))
+	{
+		MessageEndpoint->Publish(FMessageEndpoint::MakeMessage<FChaosVDRecordingStatusMessage>(InUpdateMessage), EMessageScope::Network);
+	}
+}
+
+void FChaosVDRemoteSessionsManager::PublishTraceConnectionDetailsUpdate(const FChaosVDTraceConnectionDetailsMessage& InUpdateMessage)
 {
 	if (ensure(MessageEndpoint))
 	{
-		MessageEndpoint->Publish(FMessageEndpoint::MakeMessage<FChaosVDRecordingStatusMessage>(UpdateMessage), EMessageScope::Network);
+		MessageEndpoint->Publish(FMessageEndpoint::MakeMessage<FChaosVDTraceConnectionDetailsMessage>(InUpdateMessage), EMessageScope::Network);
 	}
 }
 
@@ -445,6 +474,12 @@ void FChaosVDRemoteSessionsManager::DeRegisterSessionInMultiSessionWrapper(const
 	StaticCastSharedPtr<FChaosVDMultiSessionInfo>(ActiveSessionsByInstanceId.FindChecked(AllRemoteClientsWrapperGUID))->InnerSessionsByInstanceID.Remove(InSessionInfoRef->InstanceId);
 }
 
+void FChaosVDRemoteSessionsManager::ProcessPendingMessagesForSession(const FChaosVDSessionPong& InMessage, const TSharedRef<FChaosVDSessionInfo>& InSessionInfoPtr)
+{
+	PendingRecordingStatusMessages.RemoveAndCopyValue(InMessage.InstanceId, InSessionInfoPtr->LastKnownRecordingState);
+	PendingRecordingConnectionDetailsMessages.RemoveAndCopyValue(InMessage.InstanceId, InSessionInfoPtr->LastKnownConnectionDetails);
+}
+
 void FChaosVDRemoteSessionsManager::HandleSessionPongMessage(const FChaosVDSessionPong& InMessage, const TSharedRef<IMessageContext>& InContext)
 {
 	TSharedPtr<FChaosVDSessionInfo>& SessionInfoPtr = ActiveSessionsByInstanceId.FindOrAdd(InMessage.InstanceId);
@@ -465,7 +500,7 @@ void FChaosVDRemoteSessionsManager::HandleSessionPongMessage(const FChaosVDSessi
 
 	SessionInfoPtr->LastPingTime = FDateTime::UtcNow();
 
-	PendingRecordingStatusMessages.RemoveAndCopyValue(InMessage.InstanceId, SessionInfoPtr->LastKnownRecordingState);
+	ProcessPendingMessagesForSession(InMessage, SessionInfoPtr.ToSharedRef());
 	
 	SessionsUpdatedDelegate.Broadcast();
 }
@@ -505,13 +540,28 @@ void FChaosVDRemoteSessionsManager::HandleRecordingStatusUpdateMessage(const FCh
 	}
 }
 
+void FChaosVDRemoteSessionsManager::HandleConnectionDetailsUpdateMessage(const FChaosVDTraceConnectionDetailsMessage& InMessage, const TSharedRef<IMessageContext>& InContext)
+{
+	if (TSharedPtr<FChaosVDSessionInfo>* SessionInfoPtrPtr = ActiveSessionsByInstanceId.Find(InMessage.InstanceId))
+	{
+		TSharedPtr<FChaosVDSessionInfo>& SessionInfoPtr = *SessionInfoPtrPtr;
+		check(SessionInfoPtr);
+
+		SessionInfoPtr->LastKnownConnectionDetails = InMessage.TraceDetails;
+	}
+	else
+	{
+		PendingRecordingConnectionDetailsMessages.FindOrAdd(InMessage.InstanceId) = InMessage.TraceDetails;
+	}
+}
+
 void FChaosVDRemoteSessionsManager::HandleRecordingStartCommandMessage(const FChaosVDStartRecordingCommandMessage& InMessage, const TSharedRef<IMessageContext>& InContext)
 {
 #if WITH_CHAOS_VISUAL_DEBUGGER
 	UE_AUTORTFM_ONCOMMIT(InMessage, InContext)
 	{
 		FChaosVisualDebuggerTrace::OverrideDefaultEnabledDataChannels(InMessage.DataChannelsEnabledOverrideList);
-		FChaosVDRuntimeModule::Get().StartRecording(InMessage, InContext->GetSender().ToString());
+		FChaosVDRuntimeModule::Get().StartRecording(InMessage);
 	};
 
 #endif
@@ -575,7 +625,7 @@ void FChaosVDRemoteSessionsManager::HandleFullSessionStateRequestMessage(const F
 			FChaosVDChannelStateChangeCommandMessage DataChannelStateMessage = {Channel->GetId().ToString(), Channel->IsChannelEnabled(), Channel->CanChangeEnabledState() };
 			FullSessionStateResponse->DataChannelsStates.Emplace(DataChannelStateMessage.NewState);
 			return true;
-		});	
+		});
 #endif
 
 		MessageEndpoint->Send(
diff --git a/Engine/Source/Runtime/Experimental/ChaosSolverEngine/Public/Chaos/ChaosVDBusTraceRelayTransport.h b/Engine/Source/Runtime/Experimental/ChaosSolverEngine/Public/Chaos/ChaosVDBusTraceRelayTransport.h
deleted file mode 100644
diff --git a/Engine/Source/Runtime/Experimental/ChaosSolverEngine/Public/Chaos/ChaosVDEngineEditorBridge.h b/Engine/Source/Runtime/Experimental/ChaosSolverEngine/Public/Chaos/ChaosVDEngineEditorBridge.h
--- a/Engine/Source/Runtime/Experimental/ChaosSolverEngine/Public/Chaos/ChaosVDEngineEditorBridge.h
+++ b/Engine/Source/Runtime/Experimental/ChaosSolverEngine/Public/Chaos/ChaosVDEngineEditorBridge.h
@@ -72,7 +72,10 @@ public:
 		return EngineEditorBridgeInitializedDelegate;
 	}
 
+	CHAOSSOLVERENGINE_API void SetExternalTraceRelayInstance(const TSharedPtr<Chaos::VD::ITraceDataRelayTransport>& InExternalTraceRelayInstance);
+
 private:
+	void HandleTraceConnectionDetailsUpdated();
 	bool AddOnScreenRecordingMessage(float DummyDeltaTime = 0.1f);
 	void RemoveOnScreenRecordingMessage();
 	void HandleCVDRecordingStarted();
@@ -88,7 +91,6 @@ private:
 	bool BroadcastSessionStatus(float DeltaTime);
 
 	static FChaosVDParticleMetadata GenerateParticleMetadata(const IPhysicsProxyBase* ParticleProxy, const Chaos::FGeometryParticleHandle* ParticleHandle);
-	void InitializeTraceRelayInstance();
 
 	FDelegateHandle RecordingStartedHandle;
 	FDelegateHandle PostRecordingStartedHandle;
diff --git a/Engine/Source/Runtime/Experimental/ChaosSolverEngine/Public/Chaos/ChaosVDRemoteSessionsManager.h b/Engine/Source/Runtime/Experimental/ChaosSolverEngine/Public/Chaos/ChaosVDRemoteSessionsManager.h
--- a/Engine/Source/Runtime/Experimental/ChaosSolverEngine/Public/Chaos/ChaosVDRemoteSessionsManager.h
+++ b/Engine/Source/Runtime/Experimental/ChaosSolverEngine/Public/Chaos/ChaosVDRemoteSessionsManager.h
@@ -52,12 +52,20 @@ struct FChaosVDStopRecordingCommandMessage
 USTRUCT()
 struct FChaosVDRecordingStatusMessage
 {
+
 	GENERATED_BODY()
-	
+
 	FChaosVDRecordingStatusMessage()
 	{
 	}
 
+	PRAGMA_DISABLE_DEPRECATION_WARNINGS
+	FChaosVDRecordingStatusMessage(const FChaosVDRecordingStatusMessage& Other) = default;
+	FChaosVDRecordingStatusMessage(FChaosVDRecordingStatusMessage&& Other) noexcept = default;
+	FChaosVDRecordingStatusMessage& operator=(const FChaosVDRecordingStatusMessage& Other) = default;
+	FChaosVDRecordingStatusMessage& operator=(FChaosVDRecordingStatusMessage&& Other) noexcept = default;
+	PRAGMA_ENABLE_DEPRECATION_WARNINGS
+
 	UPROPERTY()
 	FGuid InstanceId = FGuid();
 	
@@ -66,7 +74,23 @@ struct FChaosVDRecordingStatusMessage
 
 	UPROPERTY()
 	float ElapsedTime = false;
-	
+
+	UE_DEPRECATED(5.7, "Please get the trace details directly from the session info object")
+	FChaosVDTraceDetails TraceDetails;
+};
+
+USTRUCT()
+struct FChaosVDTraceConnectionDetailsMessage
+{
+	GENERATED_BODY()
+
+	FChaosVDTraceConnectionDetailsMessage()
+	{
+	}
+
+	UPROPERTY()
+	FGuid InstanceId = FGuid();
+
 	UPROPERTY()
 	FChaosVDTraceDetails TraceDetails;
 };
@@ -218,12 +242,15 @@ public:
 	EChaosVDRemoteSessionReadyState ReadyState = EChaosVDRemoteSessionReadyState::Ready;
 
 	FChaosVDRecordingStatusMessage LastKnownRecordingState;
+	FChaosVDTraceDetails LastKnownConnectionDetails;
+
 	TMap<FString, FChaosVDDataChannelState> DataChannelsStatesByName;
 
 	virtual bool IsRecording() const;
 	
 	virtual EChaosVDRecordingMode GetRecordingMode() const;
-	
+
+	UE_DEPRECATED(5.7, "This method will no longer be used and there is not a replacement planned.")
 	virtual bool IsConnected() const;
 
 	EChaosVDRemoteSessionAttributes GetSessionTypeAttributes() const
@@ -269,7 +296,7 @@ struct FChaosVDMultiSessionInfo : public FChaosVDSessionInfo
 	TMap<FGuid, TWeakPtr<FChaosVDSessionInfo>> InnerSessionsByInstanceID;
 };
 
-DECLARE_LOG_CATEGORY_EXTERN(LogChaosVDRemoteSession, Log, Log);
+DECLARE_LOG_CATEGORY_EXTERN(LogChaosVDRemoteSession, Log, VeryVerbose);
 
 DECLARE_MULTICAST_DELEGATE_OneParam(FChaosVDRecordingStateChangeDelegate, TWeakPtr<FChaosVDSessionInfo> Session)
 
@@ -277,7 +304,6 @@ DECLARE_MULTICAST_DELEGATE_OneParam(FChaosVDRecordingStateChangeDelegate, TWeakP
 class FChaosVDRemoteSessionsManager
 {
 public:
-
 	FChaosVDRemoteSessionsManager();
 
 	UE_DEPRECATED(5.7, "Use Initialize()")
@@ -298,8 +324,7 @@ public:
 	 * Enumerate messages types allowed by this module.
 	 * @param InVisitor
 	 */
-	CHAOSSOLVERENGINE_API static void EnumerateMessageTypes(const FVisitorFunction& InVisitor);
-
+	CHAOSSOLVERENGINE_API void EnumerateMessageTypes(const FVisitorFunction& InVisitor);
 
 	/**
 	 * Starts sending ping requests that potential targets (client, servers or other editors) can answer to and report themselves as available
@@ -338,7 +363,7 @@ public:
 	CHAOSSOLVERENGINE_API TWeakPtr<FChaosVDSessionInfo> GetSessionInfo(FGuid Id);
 
 	/**
-	 * Iterates trough all active and valid cvd sessions, and executes the provided callback to it.
+	 * Iterates through all active and valid cvd sessions, and executes the provided callback to it.
 	 * if the callbacks returns false, the iteration will stop
 	* */
 	template<typename CallbackType>
@@ -358,19 +383,24 @@ public:
 		}
 	}
 
-
 	/**
 	 * Broadcast to the network a recording state update
-	 * @param UpdateMessage latest recording state of the issuing instance
+	 * @param InUpdateMessage latest recording state of the issuing instance
+	 */
+	void PublishRecordingStatusUpdate(const FChaosVDRecordingStatusMessage& InUpdateMessage);
+
+	/**
+	 * Broadcast to the network new trace connection details as they become available
+	 * @param InUpdateMessage latest connection details
 	 */
-	void PublishRecordingStatusUpdate(const FChaosVDRecordingStatusMessage& UpdateMessage);
+	void PublishTraceConnectionDetailsUpdate(const FChaosVDTraceConnectionDetailsMessage& InUpdateMessage);
+	
 	/**
 	 * Broadcast to the network a data channel state update
 	 * @param InNewStateData latest data channel state of the issuing instance
 	 */
 	void PublishDataChannelStateChangeUpdate(const FChaosVDChannelStateChangeResponseMessage& InNewStateData);
 
-
 	/**
 	 * Delegate that broadcast when a recording was started in a session (either local or remote)
 	 */
@@ -424,8 +454,16 @@ public:
 	 */
 	CHAOSSOLVERENGINE_API static constexpr bool IsController();
 
+	/**
+	 * Adds a message type to the list of supported message types for this session system.
+	 * @param ScriptStruct Message type info
+	 */
+	CHAOSSOLVERENGINE_API void RegisterExternalSupportedMessageType(const UScriptStruct* ScriptStruct);
+
 private:
 
+	void RegisterBuiltInMessageTypes();
+
 	void InitializeMessagingSystem(const TSharedPtr<IMessageBus>& InMessageBus);
 	void ShutdownMessagingSystem();
 
@@ -476,6 +514,7 @@ private:
 	void HandleSessionPongMessage(const FChaosVDSessionPong& InMessage, const TSharedRef<IMessageContext>& InContext);
 	void HandleSessionPingMessage(const FChaosVDSessionPing& InMessage, const TSharedRef<IMessageContext>& InContext);
 	void HandleRecordingStatusUpdateMessage(const FChaosVDRecordingStatusMessage& InMessage, const TSharedRef<IMessageContext>& InContext);
+	void HandleConnectionDetailsUpdateMessage(const FChaosVDTraceConnectionDetailsMessage& InMessage, const TSharedRef<IMessageContext>& InContext);
 	void HandleRecordingStartCommandMessage(const FChaosVDStartRecordingCommandMessage& InMessage, const TSharedRef<IMessageContext>& InContext);
 	void HandleRecordingStopCommandMessage(const FChaosVDStopRecordingCommandMessage& InMessage, const TSharedRef<IMessageContext>& InContext);
 	void HandleChangeDataChannelStateCommandMessage(const FChaosVDChannelStateChangeCommandMessage& InMessage, const TSharedRef<IMessageContext>& InContext);
@@ -492,6 +531,8 @@ private:
 
 	void RemoveExpiredSessions(ERemoveSessionOptions Options = ERemoveSessionOptions::None);
 
+	void ProcessPendingMessagesForSession(const FChaosVDSessionPong& InMessage, const TSharedRef<FChaosVDSessionInfo>& InSessionInfoPtr);
+
 	/** Holds the time at which the last ping was sent. */
     FDateTime LastPingTime;
 
@@ -503,6 +544,7 @@ private:
 
 	TMap<FGuid, TSharedPtr<FChaosVDSessionInfo>> ActiveSessionsByInstanceId;
 	TMap<FGuid, FChaosVDRecordingStatusMessage> PendingRecordingStatusMessages;
+	TMap<FGuid, FChaosVDTraceDetails> PendingRecordingConnectionDetailsMessages;
 
 	FSimpleMulticastDelegate SessionsUpdatedDelegate;
 
@@ -513,5 +555,7 @@ private:
 
 	FMessagingInitialized MessagingInitializedDelegate;
 
+	TArray<const UScriptStruct*> SupportedMessageTypes;
+
 	bool bInitialized = false;
 };
diff --git a/Engine/Source/Runtime/Experimental/ChaosSolverEngine/Public/Chaos/ChaosVDTraceRelayTransport.h b/Engine/Source/Runtime/Experimental/ChaosSolverEngine/Public/Chaos/ChaosVDTraceRelayTransport.h
--- a/Engine/Source/Runtime/Experimental/ChaosSolverEngine/Public/Chaos/ChaosVDTraceRelayTransport.h
+++ b/Engine/Source/Runtime/Experimental/ChaosSolverEngine/Public/Chaos/ChaosVDTraceRelayTransport.h
@@ -2,6 +2,32 @@
 
 #pragma once
 
+#include "Containers/Array.h"
+#include "ChaosVDTraceRelayTransport.generated.h"
+
+namespace Chaos::VD
+{
+
+/** Structure containing the details necessary to connect to a relay instance */
+USTRUCT()
+struct FRelayConnectionInfo
+{
+	GENERATED_BODY()
+
+	/** Relay IP address */
+	UPROPERTY()
+	FString Address;
+
+	/** Relay port */
+	UPROPERTY()
+	uint16 Port;
+
+	/** Encoded SSL certificate */
+	UPROPERTY()
+	TArray<uint8> CertificateAuthority;
+};
+}
+
 #if WITH_CHAOS_VISUAL_DEBUGGER
 #include "Delegates/Delegate.h"
 
@@ -12,6 +38,15 @@ namespace Chaos::VD
 {
 class FRelayTraceWriter;
 
+/** Possible connection attempt results */
+enum EConnectionAttemptResult
+{
+	NotStarted,
+	InProgress,
+	Success,
+	Failed
+};
+
 /**
  * Interface for any CVD trace relay transport implementation
  */
@@ -23,19 +58,15 @@ public:
 	virtual void Initialize() = 0;
 	virtual void Shutdown() = 0;
 
-	DECLARE_DELEGATE_OneParam(FProcessReceivedRelayDataDelegate, const FChaosVDRelayTraceDataMessage& DataPacket);
+	DECLARE_DELEGATE_OneParam(FProcessReceivedRelayDataDelegate, const TConstArrayView<uint8> DataPacket);
 	virtual void RegisterRelayDataReceiverForSessionID(FGuid RemoteSessionID, const FProcessReceivedRelayDataDelegate& InProcessDataDelegate) = 0;
 	virtual void UnregisterRelayDataReceiverForSessionID(FGuid RemoteSessionID) = 0;
 
 	virtual void RelayTraceDataFromWriter(FRelayTraceWriter& InRelayTraceWriter) = 0;
 
-	DECLARE_DELEGATE_RetVal_OneParam(bool, FCanSendDataDelegate, int32 BytesSent);
-	virtual void RegisterExternalDataThrottleHandler(const FCanSendDataDelegate& InExternalThrottleHandler) = 0;
-	virtual void UnregisterExternalDataThrottleHandler() = 0;
-
-protected:
-	virtual void SendRelayedTraceData(const FChaosVDRelayTraceDataMessage& InRelayData) = 0;
-	virtual void ForwardReceivedTraceData(const FChaosVDRelayTraceDataMessage& InTraceDataPacket) = 0;
+	virtual FRelayConnectionInfo GetConnectionInfo() = 0;
+	virtual EConnectionAttemptResult ConnectToRelay(FGuid RemoteSessionID, const FRelayConnectionInfo ConnectionInfo) = 0;
+	virtual EConnectionAttemptResult GetConnectionAttemptResult(FGuid RemoteSessionID) = 0;
 };
 
 }
diff --git a/Engine/Source/Runtime/Experimental/ChaosVisualDebugger/Private/ChaosVDRuntimeModule.cpp b/Engine/Source/Runtime/Experimental/ChaosVisualDebugger/Private/ChaosVDRuntimeModule.cpp
--- a/Engine/Source/Runtime/Experimental/ChaosVisualDebugger/Private/ChaosVDRuntimeModule.cpp
+++ b/Engine/Source/Runtime/Experimental/ChaosVisualDebugger/Private/ChaosVDRuntimeModule.cpp
@@ -67,7 +67,7 @@ static FAutoConsoleVariable CVarChaosVDMaxTimeToWaitForDisconnect(
 
 static FAutoConsoleVariable CVarChaosVDMaxBytesPerRelayChunk(
 	TEXT("p.Chaos.VD.MaxBytesPerRelayChunk"),
-	8192,
+	65536,
 	TEXT("How many bytes per relay data message we should send"));
 
 FChaosVDRecordingStateChangedDelegate FChaosVDRuntimeModule::RecordingStartedDelegate = FChaosVDRecordingStateChangedDelegate();
@@ -161,23 +161,23 @@ FString FChaosVDRuntimeModule::GetLastRecordingFileNamePath() const
 FChaosVDTraceDetails FChaosVDRuntimeModule::GetCurrentTraceSessionDetails() const
 {
 	FChaosVDTraceDetails Details;
-	if (bool bIsConnected = FTraceAuxiliary::IsConnected(Details.SessionGuid, Details.TraceGuid))
+	if (ExternalTraceStatusDelegate.IsBound())
+	{
+		return ExternalTraceStatusDelegate.Execute();
+	}
+	
+	if (FTraceAuxiliary::IsConnected(Details.SessionGuid, Details.TraceGuid))
 	{
 		Details.TraceTarget = FTraceAuxiliary::GetTraceDestinationString();
 		Details.Mode = GetCurrentRecordingMode();
 		Details.TransportMode = CurrentTransportMode;
-		Details.bIsConnected = bIsConnected;
+		Details.MarkAsValid();
 		return Details;
 	}
 
 	return Details;
 }
 
-const FString& FChaosVDRuntimeModule::GetConnectedToAddress() const
-{
-	return CurrentlyConnectedToAddress;
-}
-
 void FChaosVDRuntimeModule::StopTrace()
 {
 	bRequestedStop = true;
@@ -206,14 +206,40 @@ void FChaosVDRuntimeModule::UnregisterCurrentExternalRelayExecutor()
 	ExternalRelayExecutorDelegate.Unbind();
 }
 
-bool FChaosVDRuntimeModule::RelayTraceData(float DeltaTime)
+void FChaosVDRuntimeModule::RegisterExternalTraceStatusProvider(const FChaosVDExternalTraceStatusDelegate& InExternalTraceStatusCallback)
+{
+	ExternalTraceStatusDelegate = InExternalTraceStatusCallback;
+}
+
+void FChaosVDRuntimeModule::UnregisterExternalTraceStatusProvider()
+{
+	ExternalTraceStatusDelegate.Unbind();
+}
+
+void FChaosVDRuntimeModule::RelayTraceData()
 {
 	if (ensure(ExternalRelayExecutorDelegate.IsBound() && RelayWriter.IsValid()))
 	{
 		ExternalRelayExecutorDelegate.Execute(*RelayWriter);
 	}
+}
 
-	return true;
+void FChaosVDRuntimeModule::SetupRelayDataPumpDelegates()
+{
+	if (!ensure(RelayWriter))
+	{
+		return;
+	}
+
+	RelayWriter->OnNewDataAvailable().BindRaw(this, &FChaosVDRuntimeModule::RelayTraceData);
+}
+
+void FChaosVDRuntimeModule::ClearRelayDataPumpDelegates()
+{
+	if (RelayWriter)
+	{
+		RelayWriter->OnNewDataAvailable().Unbind();
+	}
 }
 
 bool FChaosVDRuntimeModule::RequestFullCapture(float DeltaTime)
@@ -253,6 +279,8 @@ void FChaosVDRuntimeModule::HandleTraceConnectionEstablished()
 		UE::TReadScopeLock ReadLock(DelegatesRWLock);
 		RecordingStartedDelegate.Broadcast();
 	}
+
+	TraceConnectionDetailsUpdatedDelegate.Broadcast();
 	
 	constexpr int32 MinAllowedTimeInSecondsBetweenCaptures = 1;
 	int32 ConfiguredTimeBetweenCaptures = CVarChaosVDGTimeBetweenFullCaptures->GetInt();
@@ -273,7 +301,7 @@ void FChaosVDRuntimeModule::HandleTraceConnectionEstablished()
 }
 
 
-void FChaosVDRuntimeModule::StartRecording(const FChaosVDStartRecordingCommandMessage& InRecordingStartCommand, const FString& InConnectedToAddress)
+void FChaosVDRuntimeModule::StartRecording(const FChaosVDStartRecordingCommandMessage& InRecordingStartCommand)
 {
 	if (!ensureAlwaysMsgf(!bIsRecording, TEXT("Received a recording start command while a recording is already active")))
 	{
@@ -325,7 +353,7 @@ void FChaosVDRuntimeModule::StartRecording(const FChaosVDStartRecordingCommandMe
 		case EChaosVDTransportMode::TraceServer:
 			bIsRecording = FTraceAuxiliary::Start(FTraceAuxiliary::EConnectionType::Network, *InRecordingStartCommand.Target, nullptr, &TracingOptions);
 			break;
-		case EChaosVDTransportMode::Memory:
+		case EChaosVDTransportMode::Direct:
 		case EChaosVDTransportMode::Relay:
 		case EChaosVDTransportMode::Invalid:
 			UE_LOG(LogChaosVDRuntime, Error, TEXT("[%hs] Unsupported Transport mode [%s]"), __func__, *UEnum::GetValueAsString(InRecordingStartCommand.TransportMode))
@@ -338,7 +366,7 @@ void FChaosVDRuntimeModule::StartRecording(const FChaosVDStartRecordingCommandMe
 	{
 		switch (InRecordingStartCommand.TransportMode)
 		{
-		case EChaosVDTransportMode::Memory:
+		case EChaosVDTransportMode::Direct:
 		case EChaosVDTransportMode::TraceServer:
 			bIsRecording = FTraceAuxiliary::Start(FTraceAuxiliary::EConnectionType::Network, *InRecordingStartCommand.Target, nullptr, &TracingOptions);
 			break;
@@ -363,10 +391,7 @@ void FChaosVDRuntimeModule::StartRecording(const FChaosVDStartRecordingCommandMe
 
 				if (bIsRecording)
 				{
-					if (!RelayExecutorTickerHandle.IsValid())
-					{
-						RelayExecutorTickerHandle = FTSTicker::GetCoreTicker().AddTicker(FTickerDelegate::CreateRaw(this, &FChaosVDRuntimeModule::RelayTraceData));
-					}
+					SetupRelayDataPumpDelegates();
 				}
 
 				break;
@@ -401,8 +426,6 @@ void FChaosVDRuntimeModule::StartRecording(const FChaosVDStartRecordingCommandMe
 		{
 			OnTraceStoppedDelegateHandle = FTraceAuxiliary::OnTraceStopped.AddRaw(this, &FChaosVDRuntimeModule::HandleTraceStopRequest);
 		}
-
-		CurrentlyConnectedToAddress = InConnectedToAddress;
 	}
 	else
 	{
@@ -436,8 +459,6 @@ void FChaosVDRuntimeModule::StopRecording()
 
 	if (!bIsRecording)
 	{
-		CurrentlyConnectedToAddress = FString();
-
 		UE_LOG(LogChaosVDRuntime, Warning, TEXT("[%s] Attempted to stop recorded when there is no CVD recording active."), ANSI_TO_TCHAR(__FUNCTION__));
 		return;
 	}
@@ -461,11 +482,7 @@ void FChaosVDRuntimeModule::StopRecording()
 		RecordingTimerHandle.Reset();
 	}
 
-	if (RelayExecutorTickerHandle.IsValid())
-	{
-		FTSTicker::RemoveTicker(RelayExecutorTickerHandle);
-		RelayExecutorTickerHandle.Reset();
-	}
+	ClearRelayDataPumpDelegates();
 
 	if (RelayWriter)
 	{
@@ -479,8 +496,6 @@ void FChaosVDRuntimeModule::StopRecording()
 		UE::TReadScopeLock ReadLock(DelegatesRWLock);
 		RecordingStopDelegate.Broadcast();
 	}
-
-	CurrentlyConnectedToAddress = FString();
 }
 
 void FChaosVDRuntimeModule::HandleTraceStopRequest(FTraceAuxiliary::EConnectionType TraceType, const FString& TraceDestination)
diff --git a/Engine/Source/Runtime/Experimental/ChaosVisualDebugger/Private/Trace/RelayTraceWriter.cpp b/Engine/Source/Runtime/Experimental/ChaosVisualDebugger/Private/Trace/RelayTraceWriter.cpp
--- a/Engine/Source/Runtime/Experimental/ChaosVisualDebugger/Private/Trace/RelayTraceWriter.cpp
+++ b/Engine/Source/Runtime/Experimental/ChaosVisualDebugger/Private/Trace/RelayTraceWriter.cpp
@@ -105,6 +105,8 @@ namespace Chaos::VD
 				++PendingBunchesCount;
 				CurrentPendingBunch = CreateBunch();
 				AvailableSpaceInCurrentBunch = MaxPendingBytesPerBunch;
+
+				NewDataAvailableDelegate.ExecuteIfBound();
 			}
 
 			const uint32 BytesToCopyNum = static_cast<uint32>(AvailableSpaceInCurrentBunch) > RemainingSizeToCopy ? RemainingSizeToCopy : static_cast<uint32>(AvailableSpaceInCurrentBunch);
diff --git a/Engine/Source/Runtime/Experimental/ChaosVisualDebugger/Public/ChaosVDRecordingDetails.h b/Engine/Source/Runtime/Experimental/ChaosVisualDebugger/Public/ChaosVDRecordingDetails.h
--- a/Engine/Source/Runtime/Experimental/ChaosVisualDebugger/Public/ChaosVDRecordingDetails.h
+++ b/Engine/Source/Runtime/Experimental/ChaosVisualDebugger/Public/ChaosVDRecordingDetails.h
@@ -23,9 +23,9 @@ enum class EChaosVDTransportMode : uint8
 	FileSystem,
 	/** Data will be traced to the selected trace store server */
 	TraceServer,
-	/** NOT IMPLEMENTED - Data will be traced directly to memory, accessible via sockets */
-	Memory,
-	/** NOT IMPLEMENTED - Data will be traced via the trace relay system using a custom transport, like normal UE Networking */
+	/** Data will be traced directly to the editor, accessible via sockets */
+	Direct,
+	/** Data will be traced via the trace relay system using a custom transport, like normal UE Networking or a direct socket connection */
 	Relay
 };
 
@@ -47,29 +47,62 @@ struct FChaosVDStartRecordingCommandMessage
 	TArray<FString> DataChannelsEnabledOverrideList;
 };
 
+/** Structure containing the info necessary to connect to a trace session and start receiving data */
 USTRUCT()
 struct FChaosVDTraceDetails
 {
 	GENERATED_BODY()
 
+	FChaosVDTraceDetails() = default;
+
+	PRAGMA_DISABLE_DEPRECATION_WARNINGS
+	FChaosVDTraceDetails(const FChaosVDTraceDetails& Other) = default;
+	FChaosVDTraceDetails(FChaosVDTraceDetails&& Other) noexcept = default;
+	FChaosVDTraceDetails& operator=(const FChaosVDTraceDetails& Other) = default;
+	FChaosVDTraceDetails& operator=(FChaosVDTraceDetails&& Other) noexcept = default;
+	PRAGMA_ENABLE_DEPRECATION_WARNINGS
+
+	void MarkAsValid()
+	{
+		bIsValid = true;
+	}
+
+	bool IsValid() const
+	{
+		return bIsValid;
+	}
+
+	/** ID used to find the active trace session in the trace system */
 	UPROPERTY()
 	FGuid TraceGuid;
-	
+
+	/** ID from the CVD remote session from which this connection details come from */
 	UPROPERTY()
 	FGuid SessionGuid;
 
+	/** Address of the trace server or editor in direct trace mode that will receive the trace data */
 	UPROPERTY()
 	FString TraceTarget;
 
+	/*** Port number used for the trace connection, if any*/
 	UPROPERTY()
+	uint16 Port = 0;
+
+	UE_DEPRECATED(5.7, "This property is no longer used and it will removed in the future")
 	bool bIsConnected = false;
 
+	/** Which recording mode is being used by this session (Live or File)*/
 	UPROPERTY()
 	EChaosVDRecordingMode Mode = EChaosVDRecordingMode::Invalid;
 
+	/** How the data is being transported to the editor (Trace Server, Custom Relay, Direct, etc.)*/
 	UPROPERTY()
 	EChaosVDTransportMode TransportMode = EChaosVDTransportMode::Invalid;
 
+	/** Encoded SSL certificate for the trace connection, if any */
+	UPROPERTY()
+	TArray<uint8> CertAuth;
+
 	UPROPERTY()
-	FString ConnectedToAddress;
+	bool bIsValid = false;
 };
\ No newline at end of file
diff --git a/Engine/Source/Runtime/Experimental/ChaosVisualDebugger/Public/ChaosVDRuntimeModule.h b/Engine/Source/Runtime/Experimental/ChaosVisualDebugger/Public/ChaosVDRuntimeModule.h
--- a/Engine/Source/Runtime/Experimental/ChaosVisualDebugger/Public/ChaosVDRuntimeModule.h
+++ b/Engine/Source/Runtime/Experimental/ChaosVisualDebugger/Public/ChaosVDRuntimeModule.h
@@ -30,11 +30,11 @@ enum class EChaosVDFullCaptureFlags : int32
 };
 ENUM_CLASS_FLAGS(EChaosVDFullCaptureFlags)
 
-
 DECLARE_MULTICAST_DELEGATE(FChaosVDRecordingStateChangedDelegate)
 DECLARE_MULTICAST_DELEGATE_OneParam(FChaosVDCaptureRequestDelegate, EChaosVDFullCaptureFlags)
 DECLARE_MULTICAST_DELEGATE_OneParam(FChaosVDRecordingStartFailedDelegate, const FText&)
 DECLARE_DELEGATE_OneParam(FChaosVDRelayTraceDataDelegate, Chaos::VD::FRelayTraceWriter&)
+DECLARE_DELEGATE_RetVal(FChaosVDTraceDetails, FChaosVDExternalTraceStatusDelegate)
 
 DECLARE_LOG_CATEGORY_EXTERN(LogChaosVDRuntime, Log, All)
 
@@ -56,7 +56,7 @@ public:
 	/** Starts a CVD recording by starting a Trace session. It will stop any existing trace session
 	 * @param InRecordingStartCommand : Used to determine if we want to record to file or a local trace server among other settings
 	 */
-	UE_API void StartRecording(const FChaosVDStartRecordingCommandMessage& InRecordingStartCommand, const FString& InConnectedToAddress = FString());
+	UE_API void StartRecording(const FChaosVDStartRecordingCommandMessage& InRecordingStartCommand);
 	
 	/* Stops an active recording */
 	UE_API void StopRecording();
@@ -140,8 +140,8 @@ public:
 	UE_DEPRECATED(5.7, "This method is no longer used. Use GetCurrentTraceSessionDetails")
 	UE_API FString GetLastRecordingFileNamePath() const;
 
+	/** Return the current connection details for the underlying trace session */
 	UE_API FChaosVDTraceDetails GetCurrentTraceSessionDetails() const;
-	UE_API const FString& GetConnectedToAddress() const;
 
 	EChaosVDRecordingMode GetCurrentRecordingMode() const
 	{
@@ -162,9 +162,28 @@ public:
 	 */
 	UE_API void UnregisterCurrentExternalRelayExecutor();
 
+	/**
+	 * Registers an external handler that will be in charge of providing the trace status if the recording mode is set to Relay
+	 * @param InExternalTraceStatusCallback Delegate that will be executed to get the trace status, periodically
+	 */
+	UE_API void RegisterExternalTraceStatusProvider(const FChaosVDExternalTraceStatusDelegate& InExternalTraceStatusCallback);
+
+	/**
+	 * Unbinds the delegate to the current external trace status provider
+	 */
+	UE_API void UnregisterExternalTraceStatusProvider();
+
+	FSimpleMulticastDelegate& OnTraceConnectionDetailsUpdated()
+	{
+		return TraceConnectionDetailsUpdatedDelegate;
+	}
+
 private:
 
-	bool RelayTraceData(float DeltaTime);
+	void RelayTraceData();
+
+	void SetupRelayDataPumpDelegates();
+	void ClearRelayDataPumpDelegates();
 
 	/** Stops the current Trace session */
 	UE_API void StopTrace();
@@ -216,9 +235,11 @@ private:
 	TUniquePtr<Chaos::VD::FRelayTraceWriter> RelayWriter;
 	FChaosVDRelayTraceDataDelegate ExternalRelayExecutorDelegate;
 
-	FString CurrentlyConnectedToAddress;
-
 	static UE_API FTransactionallySafeRWLock DelegatesRWLock;
+
+	FChaosVDExternalTraceStatusDelegate ExternalTraceStatusDelegate;
+
+	FSimpleMulticastDelegate TraceConnectionDetailsUpdatedDelegate;
 };
 #endif
 
diff --git a/Engine/Source/Runtime/Experimental/ChaosVisualDebugger/Public/Trace/RelayTraceWriter.h b/Engine/Source/Runtime/Experimental/ChaosVisualDebugger/Public/Trace/RelayTraceWriter.h
--- a/Engine/Source/Runtime/Experimental/ChaosVisualDebugger/Public/Trace/RelayTraceWriter.h
+++ b/Engine/Source/Runtime/Experimental/ChaosVisualDebugger/Public/Trace/RelayTraceWriter.h
@@ -6,6 +6,7 @@
 #include <atomic>
 #include "Containers/Array.h"
 #include "Containers/SpscQueue.h"
+#include "Delegates/Delegate.h"
 #include "Serialization/BitWriter.h"
 #include "Templates/UniquePtr.h"
 
@@ -87,18 +88,26 @@ public:
 		return bClosed;
 	}
 
+	DECLARE_DELEGATE(FNewDataAvailableDelegate)
+	FNewDataAvailableDelegate& OnNewDataAvailable()
+	{
+		return NewDataAvailableDelegate;
+	}
+
 	private:
 
-		TUniquePtr<FBufferArchive> CreateBunch();
+	FNewDataAvailableDelegate NewDataAvailableDelegate;
+
+	TUniquePtr<FBufferArchive> CreateBunch();
 
-		bool Write_Internal(const void* Data, uint32 Size);
+	bool Write_Internal(const void* Data, uint32 Size);
 
-		TSpscQueue<TUniquePtr<FBufferArchive>> PendingBunchesToRelay;
-		TUniquePtr<FBufferArchive> CurrentPendingBunch;
-		std::atomic<int32> PendingBunchesCount = 0;
-		std::atomic<bool> bClosed = false;
-		uint32 MaxPendingBytesPerBunch = 0;
-		std::atomic<int64> QueuedBytesNum = 0;
+	TSpscQueue<TUniquePtr<FBufferArchive>> PendingBunchesToRelay;
+	TUniquePtr<FBufferArchive> CurrentPendingBunch;
+	std::atomic<int32> PendingBunchesCount = 0;
+	std::atomic<bool> bClosed = false;
+	uint32 MaxPendingBytesPerBunch = 0;
+	std::atomic<int64> QueuedBytesNum = 0;
 };
 }
 #endif // WITH_CHAOS_VISUAL_DEBUGGER
