IAD Installer Concurrent Request Support
- First Pass
- Disabled by default
- Does not support purging or defrag yet
- Does not support canceling yet (although this is not currently properly implement for serial request processesing either)
- Removed HttpMutex b/c all access was happening from the installer pipe

diff --git a/Engine/Source/Runtime/Experimental/IoStore/OnDemand/Private/OnDemandContentInstaller.cpp b/Engine/Source/Runtime/Experimental/IoStore/OnDemand/Private/OnDemandContentInstaller.cpp
--- a/Engine/Source/Runtime/Experimental/IoStore/OnDemand/Private/OnDemandContentInstaller.cpp
+++ b/Engine/Source/Runtime/Experimental/IoStore/OnDemand/Private/OnDemandContentInstaller.cpp
@@ -48,7 +48,22 @@ namespace UE::IoStore
 			bContentInstallerFlushOnIdle,
 			TEXT("When the content installer runs out of tasks, it will flush any pending iochunks to disk.")
 			);
+
+		// Can't safely change this on the fly. Don't call directly
+		static TAutoConsoleVariable<bool> CVar_EnableConcurrentRequests(
+			TEXT("iostore.ContentInstallerEnableConcurrentInstallRequests"),
+			false,
+			TEXT("Enable the content installer to install chunks for multiple requests concurrently."),
+			ECVF_SaveForNextBoot
+			);
+
+		static bool EnableConcurentInstallRequests()
+		{
+			static bool bEnableConcurentInstallRequests = CVar_EnableConcurrentRequests.GetValueOnAnyThread();
+			return bEnableConcurentInstallRequests;
+		}
 	}
+
 namespace Private
 {
 
@@ -240,6 +255,14 @@ void ResolveChunksToInstall(
 
 } // namespace Private
 
+////////////////////////////////////////////////////////////////////////////////
+// Lambda visitor for FChunkDownloadRequest::HttpHandle
+template<class... TFunc>
+struct THttpHandleVisitor : TFunc...
+{
+	using TFunc::operator()...;
+};
+
 ////////////////////////////////////////////////////////////////////////////////
 uint32 FOnDemandContentInstaller::FRequest::NextSeqNo = 0;
 
@@ -336,6 +359,8 @@ void FOnDemandContentInstaller::EnqueueFlushLastAccessRequest(FOnDemandFlushLast
 
 void FOnDemandContentInstaller::CancelInstallRequest(FSharedInternalInstallRequest InstallRequest)
 {
+	// TODO: this is bit rotted a bit.  It needs to acutally cancel HTTP requests now but it doesn't
+
 	InstallerPipe.Launch(
 		TEXT("CancelIoStoreOnDemandInstallRequest"),
 		[this, InstallRequest]
@@ -373,7 +398,7 @@ void FOnDemandContentInstaller::CancelInstallRequest(FSharedInternalInstallReque
 		});
 }
 
-void FOnDemandContentInstaller::UpdateInstallRequestPriority(FSharedInternalInstallRequest InstallRequest, int32 NewPriority)
+void FOnDemandContentInstaller::UpdateInstallRequestPriority(FSharedInternalInstallRequest InstallRequest, const int32 NewPriority)
 {
 	InstallerPipe.Launch(
 		TEXT("UpdateIoStoreOnDemandInstallRequestPriority"),
@@ -398,19 +423,48 @@ void FOnDemandContentInstaller::UpdateInstallRequestPriority(FSharedInternalInst
 			{
 				for (FChunkHttpRequestHandle& PendingHttpRequest : Install.HttpRequestHandles)
 				{
-					if (PendingHttpRequest.IoHandle.IsValid())
-					{
-						PendingHttpRequest.IoHandle.UpdatePriorty(Request.Priority);
-					}
-					else 
-					{
-						check(HttpClient != nullptr);
+					const FSharedOnDemandContainer& Container = Install.ResolvedChunks[PendingHttpRequest.ContainerIndex].Container;
+					const FOnDemandChunkEntry& ChunkEntry = Container->ChunkEntries[PendingHttpRequest.EntryIndex];
+					const FIoHash& ChunkHash = ChunkEntry.Hash;
 
-						TUniqueLock HttpLock(HttpMutex);
-						if (PendingHttpRequest.Handle != nullptr)
+					if (FChunkDownloadRequest* ChunkDownload = PendingChunkDownloads.Find(ChunkHash))
+					{
+						// Use the maximum priority of all requests waiting on this chunk
+						int32 MaxPriority = std::numeric_limits<int32>::min();
+						if (ChunkDownload->ChunkRequestHttpHandles.IsEmpty() == false &&
+							ChunkDownload->ChunkRequestHttpHandles.GetHead() == ChunkDownload->ChunkRequestHttpHandles.GetTail())
+						{
+							// Common case that this is the only request waiting on the chunk
+							MaxPriority = NewPriority;
+						}
+						else
 						{
-							HttpClient->ReprioritizeRequest(PendingHttpRequest.Handle, Request.Priority);
+							// Multiple requests are waiting on the chunk
+							for (FChunkHttpRequestHandle& Handle : ChunkDownload->ChunkRequestHttpHandles)
+							{
+								const int32 RequestPriority = Handle.OwnerRequest->Priority;
+								if (RequestPriority > MaxPriority)
+								{
+									MaxPriority = RequestPriority;
+								}
+							}
 						}
+
+						Visit(THttpHandleVisitor
+						{
+							[MaxPriority](FIoHttpRequest& Handle)
+							{
+								Handle.UpdatePriorty(MaxPriority);
+							},
+							[this, MaxPriority](void* Handle)
+							{
+								check(HttpClient != nullptr);
+								if (Handle != nullptr)
+								{
+									HttpClient->ReprioritizeRequest(Handle, MaxPriority);
+								}
+							}
+						}, ChunkDownload->HttpHandle);
 					}
 				}
 			}
@@ -426,6 +480,27 @@ void FOnDemandContentInstaller::ReportAnalytics(TArray<FAnalyticsEventAttribute>
 	FOnDemandContentInstallerStats::ReportAnalytics(OutAnalyticsArray);
 }
 
+bool FOnDemandContentInstaller::CanExecuteRequest(FRequest& Request) const
+{
+	bool bExecute = false;
+
+	if (CVars::EnableConcurentInstallRequests())
+	{
+		// Allow multiple install requests XOR any other request
+		bExecute = RunningRequests.IsEmpty() ||
+			(
+				Request.IsInstall() &&
+					(RunningRequests.Num() > 1 || RunningRequests[0]->IsInstall())
+			);
+	}
+	else
+	{
+		bExecute = RunningRequests.IsEmpty();
+	}
+
+	return bExecute;
+}
+
 void FOnDemandContentInstaller::TryExecuteNextRequest()
 {
 	if (bShuttingDown.load(std::memory_order_relaxed))
@@ -433,48 +508,74 @@ void FOnDemandContentInstaller::TryExecuteNextRequest()
 		return;
 	}
 
-	FRequest* NextRequest = nullptr;
 	bool bRequestQueueIsEmpty = false;
+	while (bRequestQueueIsEmpty == false)
 	{
-		TUniqueLock Lock(Mutex);
-		bRequestQueueIsEmpty = RequestQueue.IsEmpty();
-		if (CurrentRequest == nullptr && bRequestQueueIsEmpty == false)
+		FRequest* NextRequest = nullptr;
 		{
-			RequestQueue.HeapPop(NextRequest, RequestSortPredicate, EAllowShrinking::No);
-			CurrentRequest = NextRequest;
+			TUniqueLock Lock(Mutex);
+			bRequestQueueIsEmpty = RequestQueue.IsEmpty();
+			if (bRequestQueueIsEmpty == false)
+			{
+				if (CanExecuteRequest(*RequestQueue.HeapTop()))
+				{
+					RequestQueue.HeapPop(NextRequest, RequestSortPredicate, EAllowShrinking::No);
+					RunningRequests.Add(NextRequest);
+				}
+			}
+		}
+
+		if (NextRequest == nullptr)
+		{
+			break;
 		}
-	}
 
-	if (NextRequest != nullptr)
-	{
 		InstallerPipe.Launch(
 			TEXT("OnDemandContentInstaller::ExecuteRequest"),
 			[this, NextRequest] { ExecuteRequest(*NextRequest); },
 			UE::Tasks::ETaskPriority::BackgroundLow);
+
+		// If we aren't doing an install request, we only allow one running request so 
+		// stop trying to dequeue requests.
+		if (NextRequest->IsInstall() == false)
+		{
+			break;
+		}
 	}
-	else if (bRequestQueueIsEmpty && UE::IoStore::CVars::bContentInstallerFlushOnIdle)
-	{
-		// Flush any pending chunks, but only if the installer is otherwise idle
-		InstallerPipe.Launch(
-			TEXT("OnDemandContentInstaller::Flush"),
-			[this] 
-			{
-				bool bExecuteFlush = false;
-				{
-					TUniqueLock Lock(Mutex);
-					bExecuteFlush = RequestQueue.IsEmpty() && CurrentRequest == nullptr;
-				}
 
-				if (bExecuteFlush)
+
+	if (bRequestQueueIsEmpty && UE::IoStore::CVars::bContentInstallerFlushOnIdle)
+	{		
+		bool bRunningRequestsIsEmpty = false;
+		{
+			TUniqueLock Lock(Mutex);
+			bRunningRequestsIsEmpty = RunningRequests.IsEmpty();
+		}
+
+		if (bRunningRequestsIsEmpty)
+		{
+			// Flush any pending chunks, but only if the installer is otherwise idle
+			InstallerPipe.Launch(
+				TEXT("OnDemandContentInstaller::Flush"),
+				[this]
 				{
-					FResult Result = IoStore.InstallCache->Flush();
-					if (Result.HasError())
+					bool bExecuteFlush = false;
 					{
-						UE_LOG(LogIoStoreOnDemand, Error, TEXT("Failed to flush pending chunks, Reason %s"), *LexToString(Result.GetError()));
+						TUniqueLock Lock(Mutex);
+						bExecuteFlush = RequestQueue.IsEmpty() && RunningRequests.IsEmpty();
 					}
-				}
-			},
-			UE::Tasks::ETaskPriority::BackgroundLow);
+
+					if (bExecuteFlush)
+					{
+						FResult Result = IoStore.InstallCache->Flush();
+						if (Result.HasError())
+						{
+							UE_LOG(LogIoStoreOnDemand, Error, TEXT("Failed to flush pending chunks, Reason %s"), *LexToString(Result.GetError()));
+						}
+					}
+				},
+				UE::Tasks::ETaskPriority::BackgroundLow);
+		}
 	}
 }
 
@@ -672,9 +773,9 @@ void FOnDemandContentInstaller::ProcessInstallRequest(FRequest& Request)
 	bool bExecuteRequest = false;
 	{
 		TUniqueLock Lock(Mutex);
-		if (CurrentRequest == nullptr)
+		if (CanExecuteRequest(Request))
 		{
-			CurrentRequest = &Request;
+			RunningRequests.Add(&Request);
 			bExecuteRequest = true;
 		}
 	}
@@ -682,8 +783,6 @@ void FOnDemandContentInstaller::ProcessInstallRequest(FRequest& Request)
 	if (bExecuteRequest)
 	{
 		// Execute immediately, let ExecuteInstallRequest handle chunk pinning
-
-		check(CurrentRequest == &Request);
 		ExecuteInstallRequest(Request);
 		return;
 	}
@@ -728,7 +827,7 @@ void FOnDemandContentInstaller::ProcessInstallRequest(FRequest& Request)
 void FOnDemandContentInstaller::ExecuteInstallRequest(FRequest& Request)
 {
 	check(Request.IsInstall());
-	check(&Request == CurrentRequest);
+	check(RunningRequests.Contains(&Request));
 
 	FRequest::FInstall& InstallRequest = Request.AsInstall();
 	check(InstallRequest.HttpRequestHandles.IsEmpty());
@@ -744,8 +843,9 @@ void FOnDemandContentInstaller::ExecuteInstallRequest(FRequest& Request)
 	uint64 TotalContentSize = 0;
 	uint64 TotalInstallSize	= 0;
 
-	PinCachedChunks(InstallRequest, [&InstallRequest, &TotalContentSize, &TotalInstallSize](const int32 ContainerIndex, const int32 EntryIndex, const bool bCached)
+	PinCachedChunks(InstallRequest, [&Request, &TotalContentSize, &TotalInstallSize](const int32 ContainerIndex, const int32 EntryIndex, const bool bCached)
 	{
+		FRequest::FInstall& InstallRequest = Request.AsInstall();
 		const FOnDemandChunkEntry& Entry = InstallRequest.ResolvedChunks[ContainerIndex].Container->ChunkEntries[EntryIndex];
 		TotalContentSize += Entry.GetDiskSize();
 
@@ -756,7 +856,7 @@ void FOnDemandContentInstaller::ExecuteInstallRequest(FRequest& Request)
 
 		InstallRequest.HttpRequestHandles.Add(FChunkHttpRequestHandle
 		{
-			.Handle			= nullptr,
+			.OwnerRequest	= &Request,
 			.ContainerIndex = ContainerIndex,
 			.EntryIndex		= EntryIndex
 		});
@@ -779,8 +879,11 @@ void FOnDemandContentInstaller::ExecuteInstallRequest(FRequest& Request)
 		return CompleteInstallRequest(Request);
 	}
 
-	// Make sure we have enough space in the cache
+	
+	if (CVars::EnableConcurentInstallRequests() == false)
 	{
+		// Make sure we have enough space in the cache
+
 		TSet<FIoHash> ChunksToInstall;
 		for (FChunkHttpRequestHandle& HttpRequest : InstallRequest.HttpRequestHandles)
 		{
@@ -803,81 +906,87 @@ void FOnDemandContentInstaller::ExecuteInstallRequest(FRequest& Request)
 
 	NotifyInstallProgress(Request);
 
-	if (FHttpIoDispatcher::IsInitialized())
+	const bool bUseHttpIoDispatcher = FHttpIoDispatcher::IsInitialized();
+	check(bUseHttpIoDispatcher || HttpClient != nullptr);
+
+	TOptional<FIoHttpBatch> Batch;
+	if (bUseHttpIoDispatcher)
 	{
-		const int32 HttpRetryCount = 2;
-		FIoHttpBatch Batch = FHttpIoDispatcher::NewBatch();
-		for (FChunkHttpRequestHandle& HttpRequest : InstallRequest.HttpRequestHandles)
+		Batch.Emplace(FHttpIoDispatcher::NewBatch());
+	}
+
+	constexpr const int32 HttpRetryCount = 2;
+	for (FChunkHttpRequestHandle& HttpRequest : InstallRequest.HttpRequestHandles)
+	{
+		TRACE_CPUPROFILER_EVENT_SCOPE(FOnDemandContentInstaller::IssueRequest);
+		FSharedOnDemandContainer& Container		= InstallRequest.ResolvedChunks[HttpRequest.ContainerIndex].Container;
+		FOnDemandChunkInfo ChunkInfo			= FOnDemandChunkInfo(Container, Container->ChunkEntries[HttpRequest.EntryIndex]);
+		const uint8 HttpCategory				= 1; // see FOnDemandHttpThread
+
+		if (FChunkDownloadRequest* ChunkRequest = PendingChunkDownloads.Find(ChunkInfo.ChunkEntry().Hash))
 		{
-			TRACE_CPUPROFILER_EVENT_SCOPE(FOnDemandContentInstaller::IssueRequest);
-			FSharedOnDemandContainer& Container		= InstallRequest.ResolvedChunks[HttpRequest.ContainerIndex].Container;
-			FOnDemandChunkInfo ChunkInfo			= FOnDemandChunkInfo(Container, Container->ChunkEntries[HttpRequest.EntryIndex]);
-			const uint8 HttpCategory				= 1; // see FOnDemandHttpThread
+			ChunkRequest->ChunkRequestHttpHandles.AddHead(&HttpRequest);
+		}
+		else
+		{
+			ChunkRequest = &PendingChunkDownloads.Add(ChunkInfo.ChunkEntry().Hash);
+			ChunkRequest->ChunkRequestHttpHandles.AddHead(&HttpRequest);
 
 			// Setting HTTP flags to None to not use the HTTP cache
 			FIoHttpOptions HttpOptions = FIoHttpOptions(Request.Priority, HttpRetryCount, EIoHttpFlags::None);
 			HttpOptions.SetCategory(HttpCategory);
 
-			HttpRequest.IoHandle = Batch.Get(
-				ChunkInfo.HostGroupName(),
-				ChunkInfo.RelativeUrl(),
-				FIoHttpHeaders(),
-				HttpOptions,
-				ChunkInfo.ChunkEntry().Hash,
-			[this, &Request, &HttpRequest](FIoHttpResponse&& HttpResponse)
-			{ 
-				InstallerPipe.Launch(
-					TEXT("ProcessIoStoreOnDemandDownloadedChunk"),
-					[this, &Request, &HttpRequest, HttpResponse = MoveTemp(HttpResponse)]() mutable
+			const FIoHash& ChunkHash = ChunkInfo.ChunkEntry().Hash;
+
+			if (bUseHttpIoDispatcher)
+			{
+				ChunkRequest->HttpHandle.Emplace<FIoHttpRequest>(Batch->Get(
+					ChunkInfo.HostGroupName(),
+					ChunkInfo.RelativeUrl(),
+					FIoHttpHeaders(),
+					HttpOptions,
+					ChunkHash,
+					[this, ChunkHash](FIoHttpResponse&& HttpResponse)
+					{
+						InstallerPipe.Launch(
+							TEXT("ProcessIoStoreOnDemandDownloadedChunk"),
+							[this, ChunkHash, HttpResponse = MoveTemp(HttpResponse)]() mutable
+							{
+								FIoBuffer Chunk = HttpResponse.GetBody();
+								ProcessDownloadedChunk(ChunkHash, HttpResponse.GetStatusCode(), MoveTemp(Chunk));
+							},
+							UE::Tasks::ETaskPriority::BackgroundLow);
+					})
+				);
+			}
+			else
+			{
+				ChunkRequest->HttpHandle.Emplace<void*>(HttpClient->IssueRequest(
+					MoveTemp(ChunkInfo),
+					FIoOffsetAndLength(),
+					Request.Priority,
+					[this, ChunkHash](uint32 HttpStatusCode, FStringView ErrorReason, FIoBuffer&& Chunk)
 					{
-						FIoBuffer Chunk = HttpResponse.GetBody();
-						ProcessDownloadedChunk(
-							Request,
-							HttpRequest,
-							HttpResponse.GetStatusCode(),
-							FString(),
-							MoveTemp(Chunk));
+						UE_CLOG(ErrorReason.IsEmpty() == false, LogIoStoreOnDemand, Error, TEXT("Failed to download chunk: %.*s"), ErrorReason.Len(), ErrorReason.GetData());
+
+						InstallerPipe.Launch(
+							TEXT("ProcessIoStoreOnDemandDownloadedChunk"),
+							[this, ChunkHash, HttpStatusCode, Chunk = MoveTemp(Chunk)]() mutable
+							{
+								ProcessDownloadedChunk(ChunkHash, HttpStatusCode, MoveTemp(Chunk));
+							},
+							UE::Tasks::ETaskPriority::BackgroundLow);
 					},
-					UE::Tasks::ETaskPriority::BackgroundLow);
-			});
+					EHttpRequestType::Installed)
+				);
+			}
 		}
-		Batch.Issue();
-	}
-	else
-	{
-		check(HttpClient != nullptr);
 
-		for (FChunkHttpRequestHandle& HttpRequest : InstallRequest.HttpRequestHandles)
-		{
-			TRACE_CPUPROFILER_EVENT_SCOPE(FOnDemandContentInstaller::IssueRequest);
-			FSharedOnDemandContainer& Container		= InstallRequest.ResolvedChunks[HttpRequest.ContainerIndex].Container;
-			const FOnDemandChunkEntry& ChunkEntry	= Container->ChunkEntries[HttpRequest.EntryIndex];
+	}
 
-			HttpRequest.Handle = HttpClient->IssueRequest(
-				FOnDemandChunkInfo(Container, ChunkEntry),
-				FIoOffsetAndLength(),
-				Request.Priority,
-				[this, &Request, &HttpRequest](uint32 HttpStatusCode, FStringView ErrorReason, FIoBuffer&& Chunk)
-				{
-					TRACE_CPUPROFILER_EVENT_SCOPE(FOnDemandContentInstaller::Callback);
-					{
-						TUniqueLock Lock(HttpMutex);
-						HttpRequest.Handle = nullptr;
-					}
-					InstallerPipe.Launch(
-						TEXT("ProcessIoStoreOnDemandDownloadedChunk"),
-						[this, &Request, &HttpRequest, HttpStatusCode, ErrorReason = FString(ErrorReason), Chunk = MoveTemp(Chunk)]() mutable
-						{
-							ProcessDownloadedChunk(
-								Request,
-								HttpRequest,
-								HttpStatusCode,
-								MoveTemp(ErrorReason),
-								MoveTemp(Chunk));
-						},
-						UE::Tasks::ETaskPriority::BackgroundLow);
-				}, EHttpRequestType::Installed);
-		}
+	if (Batch)
+	{
+		Batch->Issue();
 	}
 	
 	InstallRequest.bHttpRequestsIssued = true;
@@ -886,7 +995,7 @@ void FOnDemandContentInstaller::ExecuteInstallRequest(FRequest& Request)
 void FOnDemandContentInstaller::ExecutePurgeRequest(FRequest& Request)
 {
 	check(Request.IsPurge());
-	check(&Request == CurrentRequest);
+	check(RunningRequests.Num() == 1 && RunningRequests.Contains(&Request));
 
 	FRequest::FPurge& PurgeRequest	= Request.AsPurge();
 	const bool bDefrag				= EnumHasAnyFlags(PurgeRequest.Args.Options, EOnDemandPurgeOptions::Defrag);
@@ -906,7 +1015,7 @@ void FOnDemandContentInstaller::ExecutePurgeRequest(FRequest& Request)
 void FOnDemandContentInstaller::ExecuteDefragRequest(FRequest& Request)
 {
 	check(Request.IsDefrag());
-	check(&Request == CurrentRequest);
+	check(RunningRequests.Num() == 1 && RunningRequests.Contains(&Request));
 
 	FRequest::FDefrag& DefragRequest	= Request.AsDefrag();
 	const uint64* BytesToFree			= DefragRequest.Args.BytesToFree.GetPtrOrNull();
@@ -925,7 +1034,7 @@ void FOnDemandContentInstaller::ExecuteDefragRequest(FRequest& Request)
 void FOnDemandContentInstaller::ExecuteVerifyRequest(FRequest& Request)
 {
 	check(Request.IsVerify());
-	check(&Request == CurrentRequest);
+	check(RunningRequests.Num() == 1 && RunningRequests.Contains(&Request));
 
 	UE_LOG(LogIoStoreOnDemand, Log, TEXT("Executing verify cache request"));
 
@@ -940,7 +1049,7 @@ void FOnDemandContentInstaller::ExecuteVerifyRequest(FRequest& Request)
 void FOnDemandContentInstaller::ExectuteFlushLastAccessRequest(FRequest& Request)
 {
 	check(Request.IsFlushLastAccess());
-	check(&Request == CurrentRequest);
+	check(RunningRequests.Num() == 1 && RunningRequests.Contains(&Request));
 
 	UE_LOG(LogIoStoreOnDemand, Log, TEXT("Executing flush last access request"));
 
@@ -952,90 +1061,139 @@ void FOnDemandContentInstaller::ExectuteFlushLastAccessRequest(FRequest& Request
 	CompleteFlushLastAccessRequest(Request);
 }
 
-void FOnDemandContentInstaller::ProcessDownloadedChunk(
-	FRequest& Request,
-	FChunkHttpRequestHandle& HttpRequest,
-	uint32 HttpStatusCode,
-	FString&& ErrorReason,
-	FIoBuffer&& Chunk)
+void FOnDemandContentInstaller::ProcessDownloadedChunk(const FIoHash& ChunkHash, uint32 HttpStatusCode, FIoBuffer&& Chunk)
 {
 	TRACE_CPUPROFILER_EVENT_SCOPE(FOnDemandContentInstaller::ProcessDownloadedChunk);
 
-	FRequest::FInstall& InstallRequest			= Request.AsInstall();
-	FSharedOnDemandContainer& Container			= InstallRequest.ResolvedChunks[HttpRequest.ContainerIndex].Container;
-	const FOnDemandChunkEntry& ChunkEntry		= Container->ChunkEntries[HttpRequest.EntryIndex];
-	const FIoChunkId& ChunkId					= Container->ChunkIds[HttpRequest.EntryIndex];
+	FResult Result = MakeValue();
 
-	InstallRequest.Progress.CurrentInstallSize	+= ChunkEntry.GetDiskSize();
+	const bool bHttpOk = HttpStatusCode > 199 && HttpStatusCode < 300 && Chunk.GetSize() > 0;
+	if (bHttpOk)
+	{
+		const FIoHash VerifyChunkHash = FIoHash::HashBuffer(Chunk.GetView());
+		if (ChunkHash == VerifyChunkHash)
+		{
+			Result = IoStore.InstallCache->PutChunk(MoveTemp(Chunk), ChunkHash);
+			if (CVars::EnableConcurentInstallRequests())
+			{
+				// TODO: Cache needs to track size of blocks in memory so we can quickly compute cache usage
+				// TODO: if cache is full, Defrag here based on size of the pending http chunks
+			}
+		}
+		else
+		{
+			Result = MakeError(UnifiedError::IoStoreOnDemand::ChunkHashError::MakeError(
+				UnifiedError::IoStoreOnDemand::FChunkHashMismatchErrorContext
+				{
+					.ChunkId = FIoChunkId::InvalidChunkId, // Filled in later by OnProcessDownloadedChunkNotifyRequest
+					.ExpectedHash = ChunkHash,
+					.ActualHash = VerifyChunkHash
+				}));
+		}
+	}
+	else
+	{
+		Result = MakeError(UnifiedError::IoStoreOnDemand::HttpError::MakeError(HttpStatusCode));
+	}
 
-	UE_LOG(LogIoStoreOnDemand, Verbose, TEXT("Install progress %.2lf/%.2lf MiB, SeqNo=%u, Priority=%d, ContentHandle=(%s), ChunkId='%s', ChunkSize=%.2lf KiB, HttpStatus=%u"),
-		double(InstallRequest.Progress.CurrentInstallSize) / 1024.0 / 1024.0, double(InstallRequest.Progress.TotalInstallSize) / 1024.0 / 1024.0,
-		Request.SeqNo, Request.Priority, *LexToString(InstallRequest.Args.ContentHandle), *LexToString(ChunkId), double(ChunkEntry.GetDiskSize()) / 1024.0, HttpStatusCode);
+	FChunkDownloadRequest ChunkRequest = PendingChunkDownloads.FindAndRemoveChecked(ChunkHash);
+	for (FChunkHttpRequestHandle& ChunkHttpRequestHandle : ChunkRequest.ChunkRequestHttpHandles)
+	{
+		OnProcessDownloadedChunkNotifyRequest(ChunkHttpRequestHandle, Result);
+	}
+}
+
+void FOnDemandContentInstaller::OnProcessDownloadedChunkNotifyRequest(FChunkHttpRequestHandle& HttpRequest, const FResult& ChunkDownloadResult)
+{
+	TRACE_CPUPROFILER_EVENT_SCOPE(FOnDemandContentInstaller::OnProcessDownloadedChunkNotifyRequest);
+
+	FRequest& Request = *HttpRequest.OwnerRequest;
+	FRequest::FInstall& InstallRequest = Request.AsInstall();
+	FSharedOnDemandContainer& Container = InstallRequest.ResolvedChunks[HttpRequest.ContainerIndex].Container;
+	const FOnDemandChunkEntry& ChunkEntry = Container->ChunkEntries[HttpRequest.EntryIndex];
+	const FIoChunkId& ChunkId = Container->ChunkIds[HttpRequest.EntryIndex];
+
+	InstallRequest.Progress.CurrentInstallSize += ChunkEntry.GetDiskSize();
 
 	if (Request.IsOk())
 	{
-		const bool bHttpOk = HttpStatusCode > 199 && HttpStatusCode < 300 && Chunk.GetSize() > 0;
-		if (bHttpOk)
+		if (const UnifiedError::FError* Error = ChunkDownloadResult.TryGetError())
 		{
-			const FIoHash ChunkHash = FIoHash::HashBuffer(Chunk.GetView());
-
-			if (ChunkEntry.Hash == ChunkHash)
-			{
-				check(ChunkEntry.GetDiskSize() == Chunk.GetSize());
-				Request.Result = IoStore.InstallCache->PutChunk(MoveTemp(Chunk), ChunkHash);
-			}
-			else
+			// Set the chunk Id from this container if needed
+			const UnifiedError::IoStoreOnDemand::FChunkHashMismatchErrorContext* ErrorContext =
+				Error->GetErrorContext<UnifiedError::IoStoreOnDemand::FChunkHashMismatchErrorContext>();
+			if (ErrorContext)
 			{
-				Request.Result = MakeError(UE::UnifiedError::IoStoreOnDemand::ChunkHashError::MakeError(
-					UE::UnifiedError::IoStoreOnDemand::FChunkHashMismatchErrorContext
+				// FError does not deep copy by design.
+				// We need a different error per-request so we need to re-instance it.
+				Request.Result = MakeError(UnifiedError::IoStoreOnDemand::ChunkHashError::MakeError(
+					UnifiedError::IoStoreOnDemand::FChunkHashMismatchErrorContext
 					{
-						.ChunkId		= ChunkId,
-						.ExpectedHash	= ChunkEntry.Hash,
-						.ActualHash		= ChunkHash
+						.ChunkId = ChunkId,
+						.ExpectedHash = ErrorContext->ExpectedHash,
+						.ActualHash = ErrorContext->ActualHash
 					}));
 			}
 		}
-		else
+
+		if (Request.IsOk())
 		{
-			Request.Result = MakeError(UE::UnifiedError::IoStoreOnDemand::HttpError::MakeError(HttpStatusCode));
+			Request.Result = ChunkDownloadResult;
 		}
 	}
 
+	UE_LOG(LogIoStoreOnDemand, Verbose, TEXT("Install progress %.2lf/%.2lf MiB, SeqNo=%u, Priority=%d, ContentHandle=(%s), ChunkId='%s', ChunkSize=%.2lf KiB, DownloadResult=%s"),
+		double(InstallRequest.Progress.CurrentInstallSize) / 1024.0 / 1024.0, double(InstallRequest.Progress.TotalInstallSize) / 1024.0 / 1024.0,
+		Request.SeqNo, Request.Priority, *LexToString(InstallRequest.Args.ContentHandle), *LexToString(ChunkId), double(ChunkEntry.GetDiskSize()) / 1024.0, 
+		Request.Result.HasError() ? *Request.Result.GetError().GetErrorMessage(true).ToString() : TEXT("OK"));
+
 	const bool bCompleted = ++InstallRequest.DownloadedChunkCount >= InstallRequest.HttpRequestHandles.Num();
 	if (bCompleted)
 	{
 		CompleteInstallRequest(Request);
 	}
-	else 
+	else
 	{
 		NotifyInstallProgress(Request);
 
 		if (Request.IsOk() == false)
 		{
 			int32 NumCancelled = 0;
-			if (FHttpIoDispatcher::IsInitialized())
+			for (FChunkHttpRequestHandle& PendingHttpRequest : InstallRequest.HttpRequestHandles)
 			{
-				for (FChunkHttpRequestHandle& PendingHttpRequest : InstallRequest.HttpRequestHandles)
-				{
-					if (PendingHttpRequest.IoHandle.IsValid())
-					{
-						PendingHttpRequest.IoHandle.Cancel();
-						++NumCancelled;
-					}
-				}
-			}
-			else
-			{
-				check(HttpClient != nullptr);
+				const FSharedOnDemandContainer& PendingContainer = InstallRequest.ResolvedChunks[PendingHttpRequest.ContainerIndex].Container;
+				const FOnDemandChunkEntry& PendingChunkEntry = Container->ChunkEntries[PendingHttpRequest.EntryIndex];
+				const FIoHash& PendingChunkHash = ChunkEntry.Hash;
 
-				TUniqueLock Lock(HttpMutex);
-				for (FChunkHttpRequestHandle& PendingHttpRequest : InstallRequest.HttpRequestHandles)
+				if (FChunkDownloadRequest* ChunkDownload = PendingChunkDownloads.Find(PendingChunkHash))
 				{
-					if (PendingHttpRequest.Handle != nullptr)
+					// Remove PendingHttpRequest request from PendingChunkDownloads
+					ChunkDownload->ChunkRequestHttpHandles.Remove(&PendingHttpRequest);
+					if (ChunkDownload->ChunkRequestHttpHandles.IsEmpty())
 					{
-						HttpClient->CancelRequest(PendingHttpRequest.Handle);
-						PendingHttpRequest.Handle = nullptr;
-						++NumCancelled;
+						// No requests are waiting on this chunk, cancel it
+						Visit(THttpHandleVisitor
+						{
+							[&NumCancelled](FIoHttpRequest& Handle)
+							{
+								if (Handle.IsValid())
+								{
+									Handle.Cancel();
+									++NumCancelled;
+								}
+							},
+							[this, &NumCancelled](void* Handle)
+							{
+								check(HttpClient != nullptr);
+								if (Handle != nullptr)
+								{
+									HttpClient->CancelRequest(Handle);
+									++NumCancelled;
+								}
+							}
+						}, ChunkDownload->HttpHandle);
+
+						// ProcessDownloadedChunk will remove the chunk request
 					}
 				}
 			}
@@ -1171,10 +1329,7 @@ void FOnDemandContentInstaller::CompleteInstallRequest(FRequest& Request)
 		UE::TUniqueLock Lock(Mutex);
 
 		InstallRequest.Request->InstallerRequest = 0;
-		if (&Request == CurrentRequest)
-		{
-			CurrentRequest = nullptr;
-		}
+		RunningRequests.RemoveSingleSwap(&Request);
 	}
 
 	TryExecuteNextRequest();
@@ -1222,7 +1377,7 @@ void FOnDemandContentInstaller::CompleteInstallRequest(FRequest& Request)
 
 void FOnDemandContentInstaller::CompletePurgeRequest(FRequest& Request)
 {
-	check(&Request == CurrentRequest);
+	check(RunningRequests.Num() == 1 && RunningRequests.Contains(&Request));
 
 	const double DurationInSeconds = FPlatformTime::ToSeconds64(FPlatformTime::Cycles64() - Request.StartTimeCycles);
 	
@@ -1237,10 +1392,7 @@ void FOnDemandContentInstaller::CompletePurgeRequest(FRequest& Request)
 	FOnDemandPurgeCompleted OnCompleted = MoveTemp(Request.AsPurge().OnCompleted);
 	{
 		UE::TUniqueLock Lock(Mutex);
-		if (&Request == CurrentRequest)
-		{
-			CurrentRequest = nullptr;
-		}
+		RunningRequests.RemoveSingleSwap(&Request);
 	}
 	DestroyRequest(Request);
 
@@ -1268,7 +1420,7 @@ void FOnDemandContentInstaller::CompletePurgeRequest(FRequest& Request)
 
 void FOnDemandContentInstaller::CompleteDefragRequest(FRequest& Request)
 {
-	check(&Request == CurrentRequest);
+	check(RunningRequests.Num() == 1 && RunningRequests.Contains(&Request));
 
 	const double DurationInSeconds = FPlatformTime::ToSeconds64(FPlatformTime::Cycles64() - Request.StartTimeCycles);
 
@@ -1283,10 +1435,7 @@ void FOnDemandContentInstaller::CompleteDefragRequest(FRequest& Request)
 	FOnDemandDefragCompleted OnCompleted	= MoveTemp(Request.AsDefrag().OnCompleted);
 	{
 		UE::TUniqueLock Lock(Mutex);
-		if (&Request == CurrentRequest)
-		{
-			CurrentRequest = nullptr;
-		}
+		RunningRequests.RemoveSingleSwap(&Request);
 	}
 	DestroyRequest(Request);
 
@@ -1326,10 +1475,7 @@ void FOnDemandContentInstaller::CompleteVerifyRequest(FRequest& Request)
 	FOnDemandVerifyCacheCompleted OnCompleted = MoveTemp(Request.AsVerify().OnCompleted);
 	{
 		UE::TUniqueLock Lock(Mutex);
-		if (&Request == CurrentRequest)
-		{
-			CurrentRequest = nullptr;
-		}
+		RunningRequests.RemoveSingleSwap(&Request);
 	}
 	DestroyRequest(Request);
 
@@ -1357,10 +1503,7 @@ void FOnDemandContentInstaller::CompleteFlushLastAccessRequest(FRequest& Request
 	FOnDemandFlushLastAccessCompleted OnCompleted = MoveTemp(Request.AsFlushLastAccess().OnCompleted);
 	{
 		UE::TUniqueLock Lock(Mutex);
-		if (&Request == CurrentRequest)
-		{
-			CurrentRequest = nullptr;
-		}
+		RunningRequests.RemoveSingleSwap(&Request);
 	}
 	DestroyRequest(Request);
 
@@ -1380,7 +1523,16 @@ void FOnDemandContentInstaller::Shutdown()
 	const double WaitTimeoutSeconds	= 5.0;
 	const uint64 StartTimeCycles	= FPlatformTime::Cycles64();
 
-	// Wait for the current request to finish
+	// Cancel current requests
+	{
+		TUniqueLock Lock(Mutex);
+		for (FRequest* Request : RunningRequests)
+		{
+			Request->TryCancel();
+		}
+	}
+
+	// Wait for the current request(s) to finish
 	for (;;)
 	{
 		double WaitTimeSeconds = FPlatformTime::ToSeconds64(FPlatformTime::Cycles64() - StartTimeCycles);
@@ -1393,7 +1545,7 @@ void FOnDemandContentInstaller::Shutdown()
 		InstallerPipe.WaitUntilEmpty(FTimespan::FromSeconds(1.0));
 		{
 			TUniqueLock Lock(Mutex);
-			if (CurrentRequest == nullptr)
+			if (RunningRequests.IsEmpty())
 			{
 				break;
 			}
@@ -1402,7 +1554,8 @@ void FOnDemandContentInstaller::Shutdown()
 
 	{
 		TUniqueLock Lock(Mutex);
-		UE_CLOG(CurrentRequest != nullptr, LogIoStoreOnDemand, Error, TEXT("Content installer has still inflight request(s) while shutting down"));
+		UE_CLOG(RunningRequests.IsEmpty() == false, LogIoStoreOnDemand, Error, TEXT("Content installer has still inflight request(s) while shutting down"));
+		RunningRequests.Reset();
 	}
 
 	// Cancel all remaining request(s)
@@ -1414,7 +1567,7 @@ void FOnDemandContentInstaller::Shutdown()
 			if (RequestQueue.IsEmpty() == false)
 			{
 				RequestQueue.HeapPop(NextRequest, RequestSortPredicate, EAllowShrinking::No);
-				CurrentRequest = NextRequest;
+				RunningRequests.Add(NextRequest);
 			}
 		}
 
@@ -1425,6 +1578,7 @@ void FOnDemandContentInstaller::Shutdown()
 
 		NextRequest->TryCancel();
 		ExecuteRequest(*NextRequest);
+		check(RunningRequests.IsEmpty());
 	}
 }
 
diff --git a/Engine/Source/Runtime/Experimental/IoStore/OnDemand/Private/OnDemandContentInstaller.h b/Engine/Source/Runtime/Experimental/IoStore/OnDemand/Private/OnDemandContentInstaller.h
--- a/Engine/Source/Runtime/Experimental/IoStore/OnDemand/Private/OnDemandContentInstaller.h
+++ b/Engine/Source/Runtime/Experimental/IoStore/OnDemand/Private/OnDemandContentInstaller.h
@@ -10,6 +10,7 @@
 #include "IO/PackageId.h"
 #include "IO/IoStoreOnDemandInternals.h"
 #include "IO/IoAllocators.h"
+#include "IO/IoContainers.h"
 #include "IO/HttpIoDispatcher.h"
 #include "Misc/TVariant.h"
 #include "OnDemandIoStore.h"
@@ -60,12 +61,15 @@ void ResolveChunksToInstall(
 ////////////////////////////////////////////////////////////////////////////////
 class FOnDemandContentInstaller
 {
-	struct FChunkHttpRequestHandle
+	struct FRequest;
+
+	struct FChunkHttpRequestHandle : public TIntrusiveListElement<FChunkHttpRequestHandle>
 	{
-		void*			Handle = nullptr;
-		FIoHttpRequest	IoHandle;
-		int32			ContainerIndex = INDEX_NONE;
-		int32			EntryIndex = INDEX_NONE;
+		FChunkHttpRequestHandle*	Next = nullptr;
+		FRequest*					OwnerRequest = nullptr;
+
+		int32						ContainerIndex = INDEX_NONE;
+		int32						EntryIndex = INDEX_NONE;
 	};
 
 	struct FRequest
@@ -89,7 +93,7 @@ class FOnDemandContentInstaller
 			FOnDemandInstallProgress					Progress;
 			uint64										DownloadedChunkCount = 0;
 			uint64										LastProgressCycles = 0;
-			std::atomic_bool							bHttpRequestsIssued{false};
+			std::atomic_bool							bHttpRequestsIssued{false}; // TODO: does this need to be atomic? All access is from the installer pipe
 			std::atomic_bool							bNotifyingProgressOnGameThread{false};
 		};
 
@@ -185,6 +189,12 @@ class FOnDemandContentInstaller
 		FRequestVariant				Variant;
 	};
 
+	struct FChunkDownloadRequest
+	{
+		TVariant<void*, FIoHttpRequest> HttpHandle;
+		TIntrusiveList<FChunkHttpRequestHandle> ChunkRequestHttpHandles;
+	};
+
 	static bool RequestSortPredicate(const FRequest& LHS, const FRequest& RHS)
 	{
 		if (LHS.Variant.GetIndex() == RHS.Variant.GetIndex())
@@ -219,6 +229,7 @@ public:
 	void							ReportAnalytics(TArray<FAnalyticsEventAttribute>& OutAnalyticsArray) const;
 
 private:
+	bool							CanExecuteRequest(FRequest& Request) const;
 	void							TryExecuteNextRequest();
 	void							ExecuteRequest(FRequest& Request);
 	void							PinCachedChunks(FRequest::FInstall& InstallRequest, TFunctionRef<void(int32, int32, bool)> OnChunkFound) const;
@@ -228,12 +239,8 @@ private:
 	void							ExecuteDefragRequest(FRequest& Request);
 	void							ExecuteVerifyRequest(FRequest& Request);
 	void							ExectuteFlushLastAccessRequest(FRequest& Request);
-	void							ProcessDownloadedChunk(
-										FRequest& Request,
-										FChunkHttpRequestHandle& HttpRequest,
-										uint32 HttpStatusCode,
-										FString&& ErrorReason,
-										FIoBuffer&& Chunk);
+	void							ProcessDownloadedChunk(const FIoHash& ChunkHash, uint32 HttpStatusCode, FIoBuffer&& Chunk);
+	void							OnProcessDownloadedChunkNotifyRequest(FChunkHttpRequestHandle& HttpRequest, const FResult& ChunkDownloadResult);
 	void							NotifyInstallProgress(FRequest& Request);
 	void							DestroyRequest(FRequest& Request);
 	void							CompleteInstallRequest(FRequest& Request);
@@ -248,10 +255,10 @@ private:
 	UE::Tasks::FPipe				InstallerPipe;
 
 	FMutex							Mutex;
-	FMutex							HttpMutex;
 	FRequestAllocator				RequestAllocator;
 	TArray<FRequest*>				RequestQueue;
-	FRequest*						CurrentRequest = nullptr;
+	TArray<FRequest*>				RunningRequests;
+	TMap<FIoHash, FChunkDownloadRequest> PendingChunkDownloads;
 	std::atomic_bool				bShuttingDown{false};
 };
 
diff --git a/Engine/Source/Runtime/Experimental/IoStore/OnDemandCore/Public/IO/HttpIoDispatcher.h b/Engine/Source/Runtime/Experimental/IoStore/OnDemandCore/Public/IO/HttpIoDispatcher.h
--- a/Engine/Source/Runtime/Experimental/IoStore/OnDemandCore/Public/IO/HttpIoDispatcher.h
+++ b/Engine/Source/Runtime/Experimental/IoStore/OnDemandCore/Public/IO/HttpIoDispatcher.h
@@ -281,7 +281,13 @@ class FIoHttpBatch
 {
 public:
 							FIoHttpBatch(const FIoHttpRequest&) = delete;
-							FIoHttpBatch(FIoHttpBatch&& Other) = delete;
+							FIoHttpBatch(FIoHttpBatch&& Other)
+								: First(Other.First), Last(Other.Last)
+							{
+								Other.First = 0;
+								Other.Last = 0;
+							}
+
 	UE_API					~FIoHttpBatch();
 
 	UE_API FIoHttpRequest	Get(const FName& HostGroup,
