Fix using the wrong container in OnProcessDownloadedChunkNotifyRequest when canceling HTTP requests for an install request with an error.

#FailedValidations Preflight

diff --git a/Engine/Source/Runtime/Experimental/IoStore/OnDemand/Private/OnDemandContentInstaller.cpp b/Engine/Source/Runtime/Experimental/IoStore/OnDemand/Private/OnDemandContentInstaller.cpp
--- a/Engine/Source/Runtime/Experimental/IoStore/OnDemand/Private/OnDemandContentInstaller.cpp
+++ b/Engine/Source/Runtime/Experimental/IoStore/OnDemand/Private/OnDemandContentInstaller.cpp
@@ -1162,7 +1162,7 @@ void FOnDemandContentInstaller::OnProcessDownloadedChunkNotifyRequest(FChunkHttp
 			for (FChunkHttpRequestHandle& PendingHttpRequest : InstallRequest.HttpRequestHandles)
 			{
 				const FSharedOnDemandContainer& PendingContainer = InstallRequest.ResolvedChunks[PendingHttpRequest.ContainerIndex].Container;
-				const FOnDemandChunkEntry& PendingChunkEntry = Container->ChunkEntries[PendingHttpRequest.EntryIndex];
+				const FOnDemandChunkEntry& PendingChunkEntry = PendingContainer->ChunkEntries[PendingHttpRequest.EntryIndex];
 				const FIoHash& PendingChunkHash = ChunkEntry.Hash;
 
 				if (FChunkDownloadRequest* ChunkDownload = PendingChunkDownloads.Find(PendingChunkHash))
