Fix using the wrong chunk entry in OnProcessDownloadedChunkNotifyRequest when canceling HTTP requests for an install request with an error.

diff --git a/Engine/Source/Runtime/Experimental/IoStore/OnDemand/Private/OnDemandContentInstaller.cpp b/Engine/Source/Runtime/Experimental/IoStore/OnDemand/Private/OnDemandContentInstaller.cpp
--- a/Engine/Source/Runtime/Experimental/IoStore/OnDemand/Private/OnDemandContentInstaller.cpp
+++ b/Engine/Source/Runtime/Experimental/IoStore/OnDemand/Private/OnDemandContentInstaller.cpp
@@ -1163,7 +1163,7 @@ void FOnDemandContentInstaller::OnProcessDownloadedChunkNotifyRequest(FChunkHttp
 			{
 				const FSharedOnDemandContainer& PendingContainer = InstallRequest.ResolvedChunks[PendingHttpRequest.ContainerIndex].Container;
 				const FOnDemandChunkEntry& PendingChunkEntry = PendingContainer->ChunkEntries[PendingHttpRequest.EntryIndex];
-				const FIoHash& PendingChunkHash = ChunkEntry.Hash;
+				const FIoHash& PendingChunkHash = PendingChunkEntry.Hash;
 
 				if (FChunkDownloadRequest* ChunkDownload = PendingChunkDownloads.Find(PendingChunkHash))
 				{
