[Core] Cleanup some uses of compact and compactstable
* Removed uses of compact that had no effect since only adds were used
* Replace CompactStable with RemoveStable or just Compact depending

#tests Preflight, ran pie

diff --git a/Engine/Source/Runtime/Core/Private/Internationalization/CultureFilter.cpp b/Engine/Source/Runtime/Core/Private/Internationalization/CultureFilter.cpp
--- a/Engine/Source/Runtime/Core/Private/Internationalization/CultureFilter.cpp
+++ b/Engine/Source/Runtime/Core/Private/Internationalization/CultureFilter.cpp
@@ -115,7 +115,6 @@ void FCultureFilter::Init(const EBuildConfiguration BuildConfig, const ELocaliza
 				UE_LOG(LogLocalization, Warning, TEXT("Culture '%s' is unknown and has been ignored when parsing the enabled/disabled culture list."), *CultureName);
 			}
 		}
-		OutCulturesSet.Compact();
 	};
 
 	const TArray<FString> EnabledCulturesArray = LoadInternationalizationConfigArray(TEXT("EnabledCultures"));
diff --git a/Engine/Source/Runtime/CoreUObject/Private/UObject/UObjectHash.cpp b/Engine/Source/Runtime/CoreUObject/Private/UObject/UObjectHash.cpp
--- a/Engine/Source/Runtime/CoreUObject/Private/UObject/UObjectHash.cpp
+++ b/Engine/Source/Runtime/CoreUObject/Private/UObject/UObjectHash.cpp
@@ -227,7 +227,7 @@ struct FSetHashBucket
 		TSet<UObjectBase*>* Items = GetItems();
 		if (Items)
 		{
-			Items->Compact();
+			Items->Shrink();
 		}
 	}
 
diff --git a/Engine/Source/Runtime/D3D12RHI/Private/D3D12DescriptorCache.cpp b/Engine/Source/Runtime/D3D12RHI/Private/D3D12DescriptorCache.cpp
--- a/Engine/Source/Runtime/D3D12RHI/Private/D3D12DescriptorCache.cpp
+++ b/Engine/Source/Runtime/D3D12RHI/Private/D3D12DescriptorCache.cpp
@@ -1092,8 +1092,6 @@ void FD3D12GlobalOnlineSamplerHeap::ConsolidateUniqueSamplerTables(TArrayView<FD
 
 	if (bModified)
 	{
-		// Rearrange the set for better look-up performance
-		UniqueDescriptorTables->Compact();
 		SET_DWORD_STAT(STAT_NumReuseableSamplerOnlineDescriptorTables, UniqueDescriptorTables->Num());
 		SET_DWORD_STAT(STAT_NumReuseableSamplerOnlineDescriptors, GetNextSlotIndex());
 	}
diff --git a/Engine/Source/Runtime/Datasmith/DirectLink/Private/DirectLinkSceneSnapshot.cpp b/Engine/Source/Runtime/Datasmith/DirectLink/Private/DirectLinkSceneSnapshot.cpp
--- a/Engine/Source/Runtime/Datasmith/DirectLink/Private/DirectLinkSceneSnapshot.cpp
+++ b/Engine/Source/Runtime/Datasmith/DirectLink/Private/DirectLinkSceneSnapshot.cpp
@@ -75,10 +75,7 @@ TSharedPtr<FSceneSnapshot> SnapshotScene(ISceneGraphNode* RootElement)
 	}
 
 	TSet<ISceneGraphNode*> Nodes = BuildIndexForScene(RootElement);
-	{
-		TRACE_CPUPROFILER_EVENT_SCOPE(DirectLink::SnapshotScene/Compact);
-		Nodes.Compact(); // expectation: No-op because the set was just built
-	}
+	check(Nodes.Num() == Nodes.GetMaxIndex()); // We use direct indexing later so no holes are allowed
 
 	if (!ensure(RootElement->GetSharedState().IsValid()))
 	{
diff --git a/Engine/Source/Runtime/Experimental/Dataflow/Core/Private/Dataflow/DataflowNode.cpp b/Engine/Source/Runtime/Experimental/Dataflow/Core/Private/Dataflow/DataflowNode.cpp
--- a/Engine/Source/Runtime/Experimental/Dataflow/Core/Private/Dataflow/DataflowNode.cpp
+++ b/Engine/Source/Runtime/Experimental/Dataflow/Core/Private/Dataflow/DataflowNode.cpp
@@ -1047,10 +1047,8 @@ FDataflowOutput& FDataflowNode::RegisterOutputArrayConnectionInternal(const UE::
 
 void FDataflowNode::UnregisterInputConnection(const UE::Dataflow::FConnectionKey& Key)
 {
-	if (ExpandedInputs.Remove(Key))
+	if (ExpandedInputs.RemoveStable(Key))
 	{
-		// make sure we compact so that the next added element will be in proper order 
-		ExpandedInputs.CompactStable();
 		// Invalidate graph as this input might have had connections
 		Invalidate();
 	}
@@ -1058,10 +1056,8 @@ void FDataflowNode::UnregisterInputConnection(const UE::Dataflow::FConnectionKey
 
 void FDataflowNode::UnregisterOutputConnection(const UE::Dataflow::FConnectionKey& Key)
 {
-	if (ExpandedOutputs.Remove(Key))
+	if (ExpandedOutputs.RemoveStable(Key))
 	{
-		// make sure we compact so that the next added element will be in proper order 
-		ExpandedOutputs.CompactStable();
 		// Invalidate graph as this output might have had connections
 		Invalidate();
 	}
diff --git a/Engine/Source/Runtime/Interchange/Core/Private/Types/AttributeStorage.cpp b/Engine/Source/Runtime/Interchange/Core/Private/Types/AttributeStorage.cpp
--- a/Engine/Source/Runtime/Interchange/Core/Private/Types/AttributeStorage.cpp
+++ b/Engine/Source/Runtime/Interchange/Core/Private/Types/AttributeStorage.cpp
@@ -1948,7 +1948,7 @@ namespace UE
 			FragmentedMemoryCost = 0;
 
 			//Remove holes in the TMap
-			AttributeAllocationTable.CompactStable();
+			AttributeAllocationTable.Compact();
 
 			//Update our allocation table with the defragged offsets / sizes
 			for (const FLinearAllocation& Alloc : LinearAllocations)
