Fix TicketId not being initialized.

diff --git a/Engine/Source/Runtime/Experimental/IoStore/OnDemand/Private/OnDemandHttpClient.h b/Engine/Source/Runtime/Experimental/IoStore/OnDemand/Private/OnDemandHttpClient.h
--- a/Engine/Source/Runtime/Experimental/IoStore/OnDemand/Private/OnDemandHttpClient.h
+++ b/Engine/Source/Runtime/Experimental/IoStore/OnDemand/Private/OnDemandHttpClient.h
@@ -159,7 +159,7 @@ private:
 		TArray<FAnsiString>	Headers;
 		FAnsiString			RelativeUrl;
 		FIoOffsetAndLength	Range;
-		FHttpTicketId		TicketId;
+		FHttpTicketId		TicketId = 0;
 		FConnection&		Connection;
 		/** Stored the response body */
 		FIoBuffer			Body;
