Implement HTTP request canceling for IAD installer.  Can be disabled with iostore.ContentInstallerEnableHttpCancel

#tests Tested with HTTP dispatcher on and off and with concurrent requests on and off

diff --git a/Engine/Source/Runtime/Experimental/IoStore/OnDemand/Private/OnDemandContentInstaller.cpp b/Engine/Source/Runtime/Experimental/IoStore/OnDemand/Private/OnDemandContentInstaller.cpp
--- a/Engine/Source/Runtime/Experimental/IoStore/OnDemand/Private/OnDemandContentInstaller.cpp
+++ b/Engine/Source/Runtime/Experimental/IoStore/OnDemand/Private/OnDemandContentInstaller.cpp
@@ -49,6 +49,13 @@ namespace UE::IoStore
 			TEXT("When the content installer runs out of tasks, it will flush any pending iochunks to disk.")
 			);
 
+		static bool bContentInstallerEnableHttpCancel = true;
+		static FAutoConsoleVariableRef CVar_ContentInstallerEnableHttpCancel(
+			TEXT("iostore.ContentInstallerEnableHttpCancel"),
+			bContentInstallerEnableHttpCancel,
+			TEXT("Enable canceling HTTP requests when an install request is canceled")
+			);
+
 		// Can't safely change this on the fly. Don't call directly
 		static TAutoConsoleVariable<bool> CVar_EnableConcurrentRequests(
 			TEXT("iostore.ContentInstallerEnableConcurrentInstallRequests"),
@@ -359,35 +366,51 @@ void FOnDemandContentInstaller::EnqueueFlushLastAccessRequest(FOnDemandFlushLast
 
 void FOnDemandContentInstaller::CancelInstallRequest(FSharedInternalInstallRequest InstallRequest)
 {
-	// TODO: this is bit rotted a bit.  It needs to acutally cancel HTTP requests now but it doesn't
-
 	InstallerPipe.Launch(
 		TEXT("CancelIoStoreOnDemandInstallRequest"),
 		[this, InstallRequest]
 		{
 			FRequest* ToComplete = nullptr;
 			{
-				UE::TUniqueLock Lock(Mutex);
-
 				if (InstallRequest->InstallerRequest == 0)
 				{
 					return;
 				}
 
 				FRequest* Request = reinterpret_cast<FRequest*>(InstallRequest->InstallerRequest);
+				FRequest::FInstall& Install = Request->AsInstall();
 
-				if (Request->TryCancel() == false)
+				if (Install.bHttpRequestsIssued)
 				{
-					return;
-				}
+					if (Request->TryCancel() == false)
+					{
+						return;
+					}
 
-				UE_LOG(LogIoStoreOnDemand, Log, TEXT("Cancelling install request, ContentHandle=(%s)"),
-					*LexToString(Request->AsInstall().Args.ContentHandle));
+					UE_LOG(LogIoStoreOnDemand, Log, TEXT("Cancelling install request, ContentHandle=(%s)"),
+						*LexToString(Install.Args.ContentHandle));
 
-				if (RequestQueue.Remove(Request) > 0)
+					int32 NumCancelled = 0;
+					TryCancelHttpRequestsForInstallRequest(Install, NumCancelled);
+					UE_LOG(LogIoStoreOnDemand, Verbose, TEXT("Cancelled %d HTTP request(s) due to install request cancellation"), NumCancelled);
+				}
+				else
 				{
-					ToComplete = Request;
-					RequestQueue.Heapify(RequestSortPredicate);
+					if (Request->TryCancel() == false)
+					{
+						return;
+					}
+
+					UE_LOG(LogIoStoreOnDemand, Log, TEXT("Cancelling install request, ContentHandle=(%s)"),
+						*LexToString(Install.Args.ContentHandle));
+
+					TUniqueLock Lock(Mutex);
+
+					if (RequestQueue.Remove(Request) > 0)
+					{
+						ToComplete = Request;
+						RequestQueue.Heapify(RequestSortPredicate);
+					}
 				}
 			}
 
@@ -404,8 +427,6 @@ void FOnDemandContentInstaller::UpdateInstallRequestPriority(FSharedInternalInst
 		TEXT("UpdateIoStoreOnDemandInstallRequestPriority"),
 		[this, InstallRequest, NewPriority]
 		{
-			UE::TUniqueLock Lock(Mutex);
-
 			if (InstallRequest->InstallerRequest == 0)
 			{
 				return;
@@ -417,36 +438,34 @@ void FOnDemandContentInstaller::UpdateInstallRequestPriority(FSharedInternalInst
 			UE_LOG(LogIoStoreOnDemand, Log, TEXT("Updating install request priority, SeqNo=%u, Priority=%d, NewPriority=%d, ContentHandle=(%s)"),
 				Request.SeqNo, Request.Priority, NewPriority, *LexToString(Install.Args.ContentHandle));
 
-			Request.Priority = NewPriority;
-
-			if (Install.bHttpRequestsIssued.load(std::memory_order_seq_cst))
+			if (Install.bHttpRequestsIssued)
 			{
+				// The request has definitely left the RequestQueue 
+
+				Request.Priority = NewPriority;
+
 				for (FChunkHttpRequestHandle& PendingHttpRequest : Install.HttpRequestHandles)
 				{
 					const FSharedOnDemandContainer& Container = Install.ResolvedChunks[PendingHttpRequest.ContainerIndex].Container;
 					const FOnDemandChunkEntry& ChunkEntry = Container->ChunkEntries[PendingHttpRequest.EntryIndex];
 					const FIoHash& ChunkHash = ChunkEntry.Hash;
 
-					if (FChunkDownloadRequest* ChunkDownload = PendingChunkDownloads.Find(ChunkHash))
+					for (auto It = PendingChunkDownloads.CreateKeyIterator(ChunkHash); It; ++It)
 					{
-						// Use the maximum priority of all requests waiting on this chunk
-						int32 MaxPriority = std::numeric_limits<int32>::min();
-						if (ChunkDownload->ChunkRequestHttpHandles.IsEmpty() == false &&
-							ChunkDownload->ChunkRequestHttpHandles.GetHead() == ChunkDownload->ChunkRequestHttpHandles.GetTail())
+						FChunkDownloadRequest& ChunkDownload = It.Value();
+						if (ChunkDownload.bChunkCanceled)
 						{
-							// Common case that this is the only request waiting on the chunk
-							MaxPriority = NewPriority;
+							continue;
 						}
-						else
+
+						// Use the maximum priority of all requests waiting on this chunk
+						int32 MaxPriority = std::numeric_limits<int32>::min();
+						for (FChunkHttpRequestHandle& Handle : ChunkDownload.ChunkRequestHttpHandles)
 						{
-							// Multiple requests are waiting on the chunk
-							for (FChunkHttpRequestHandle& Handle : ChunkDownload->ChunkRequestHttpHandles)
+							const int32 RequestPriority = Handle.OwnerRequest->Priority;
+							if (RequestPriority > MaxPriority)
 							{
-								const int32 RequestPriority = Handle.OwnerRequest->Priority;
-								if (RequestPriority > MaxPriority)
-								{
-									MaxPriority = RequestPriority;
-								}
+								MaxPriority = RequestPriority;
 							}
 						}
 
@@ -464,12 +483,18 @@ void FOnDemandContentInstaller::UpdateInstallRequestPriority(FSharedInternalInst
 									HttpClient->ReprioritizeRequest(Handle, MaxPriority);
 								}
 							}
-						}, ChunkDownload->HttpHandle);
+						}, ChunkDownload.HttpHandle);
+
+						// Should be a max of one active request for any chunk
+						break;
 					}
 				}
 			}
 			else
 			{
+				// The request may or may not still be in the RequestQueue
+				TUniqueLock Lock(Mutex);
+				Request.Priority = NewPriority;
 				RequestQueue.Heapify(RequestSortPredicate);
 			}
 		});
@@ -923,37 +948,52 @@ void FOnDemandContentInstaller::ExecuteInstallRequest(FRequest& Request)
 		FOnDemandChunkInfo ChunkInfo			= FOnDemandChunkInfo(Container, Container->ChunkEntries[HttpRequest.EntryIndex]);
 		const uint8 HttpCategory				= 1; // see FOnDemandHttpThread
 
-		if (FChunkDownloadRequest* ChunkRequest = PendingChunkDownloads.Find(ChunkInfo.ChunkEntry().Hash))
+		FChunkDownloadRequest* PendingChunkRequest = nullptr;
+		for (auto It = PendingChunkDownloads.CreateKeyIterator(ChunkInfo.ChunkEntry().Hash); It; ++It)
+		{
+			FChunkDownloadRequest& ChunkDownload = It.Value();
+			if (ChunkDownload.bChunkCanceled == false)
+			{
+				PendingChunkRequest = &ChunkDownload;
+				break; // Should be a max of one active request for any chunk
+			}
+		}
+
+		if (PendingChunkRequest)
 		{
-			ChunkRequest->ChunkRequestHttpHandles.AddHead(&HttpRequest);
+			PendingChunkRequest->ChunkRequestHttpHandles.AddHead(&HttpRequest);
 		}
 		else
 		{
-			ChunkRequest = &PendingChunkDownloads.Add(ChunkInfo.ChunkEntry().Hash);
-			ChunkRequest->ChunkRequestHttpHandles.AddHead(&HttpRequest);
+			const FIoHash& ChunkHash = ChunkInfo.ChunkEntry().Hash;
+
+			PendingChunkRequest = &PendingChunkDownloads.Add(ChunkHash);
+			PendingChunkRequest->RequestId = Request.SeqNo;
+			PendingChunkRequest->ChunkRequestHttpHandles.AddHead(&HttpRequest);
 
 			// Setting HTTP flags to None to not use the HTTP cache
 			FIoHttpOptions HttpOptions = FIoHttpOptions(Request.Priority, HttpRetryCount, EIoHttpFlags::None);
 			HttpOptions.SetCategory(HttpCategory);
 
-			const FIoHash& ChunkHash = ChunkInfo.ChunkEntry().Hash;
+			UE_LOG(LogIoStoreOnDemand, VeryVerbose, TEXT("Created request for chunk: ChunkHash=%s, RequestId=%u"),
+				*TStringBuilder<64>(InPlace, ChunkHash), PendingChunkRequest->RequestId);
 
 			if (bUseHttpIoDispatcher)
 			{
-				ChunkRequest->HttpHandle.Emplace<FIoHttpRequest>(Batch->Get(
+				PendingChunkRequest->HttpHandle.Emplace<FIoHttpRequest>(Batch->Get(
 					ChunkInfo.HostGroupName(),
 					ChunkInfo.RelativeUrl(),
 					FIoHttpHeaders(),
 					HttpOptions,
 					ChunkHash,
-					[this, ChunkHash](FIoHttpResponse&& HttpResponse)
+					[this, ChunkHash, RequestId = PendingChunkRequest->RequestId](FIoHttpResponse&& HttpResponse)
 					{
 						InstallerPipe.Launch(
 							TEXT("ProcessIoStoreOnDemandDownloadedChunk"),
-							[this, ChunkHash, HttpResponse = MoveTemp(HttpResponse)]() mutable
+							[this, ChunkHash, RequestId, HttpResponse = MoveTemp(HttpResponse)]() mutable
 							{
 								FIoBuffer Chunk = HttpResponse.GetBody();
-								ProcessDownloadedChunk(ChunkHash, HttpResponse.GetStatusCode(), MoveTemp(Chunk));
+								ProcessDownloadedChunk(ChunkHash, RequestId, HttpResponse.GetErrorCode(), HttpResponse.GetStatusCode(), MoveTemp(Chunk));
 							},
 							UE::Tasks::ETaskPriority::BackgroundLow);
 					})
@@ -961,19 +1001,28 @@ void FOnDemandContentInstaller::ExecuteInstallRequest(FRequest& Request)
 			}
 			else
 			{
-				ChunkRequest->HttpHandle.Emplace<void*>(HttpClient->IssueRequest(
+				PendingChunkRequest->HttpHandle.Emplace<void*>(HttpClient->IssueRequest(
 					MoveTemp(ChunkInfo),
 					FIoOffsetAndLength(),
 					Request.Priority,
-					[this, ChunkHash](uint32 HttpStatusCode, FStringView ErrorReason, FIoBuffer&& Chunk)
+					[this, ChunkHash, RequestId = PendingChunkRequest->RequestId](uint32 HttpStatusCode, FStringView ErrorReason, FIoBuffer&& Chunk)
 					{
-						UE_CLOG(ErrorReason.IsEmpty() == false, LogIoStoreOnDemand, Error, TEXT("Failed to download chunk: %.*s"), ErrorReason.Len(), ErrorReason.GetData());
+						EIoErrorCode CompletionStatus = (HttpStatusCode > 199 && HttpStatusCode < 300) ? EIoErrorCode::Ok : EIoErrorCode::ReadError;
+						if (ErrorReason.Contains(TEXTVIEW("cancelled"), ESearchCase::IgnoreCase) ||
+							ErrorReason.Contains(TEXTVIEW("canceled"),  ESearchCase::IgnoreCase)) // Yes we spell it both ways (╯°□°）╯︵ ┻━┻
+						{
+							CompletionStatus = EIoErrorCode::Cancelled;
+						}
+						else
+						{
+							UE_CLOG(ErrorReason.IsEmpty() == false, LogIoStoreOnDemand, Warning, TEXT("Failed to download chunk: %.*s"), ErrorReason.Len(), ErrorReason.GetData());
+						}
 
 						InstallerPipe.Launch(
 							TEXT("ProcessIoStoreOnDemandDownloadedChunk"),
-							[this, ChunkHash, HttpStatusCode, Chunk = MoveTemp(Chunk)]() mutable
+							[this, ChunkHash, RequestId, CompletionStatus, HttpStatusCode, Chunk = MoveTemp(Chunk)]() mutable
 							{
-								ProcessDownloadedChunk(ChunkHash, HttpStatusCode, MoveTemp(Chunk));
+								ProcessDownloadedChunk(ChunkHash, RequestId, CompletionStatus, HttpStatusCode, MoveTemp(Chunk));
 							},
 							UE::Tasks::ETaskPriority::BackgroundLow);
 					},
@@ -1061,7 +1110,8 @@ void FOnDemandContentInstaller::ExectuteFlushLastAccessRequest(FRequest& Request
 	CompleteFlushLastAccessRequest(Request);
 }
 
-void FOnDemandContentInstaller::ProcessDownloadedChunk(const FIoHash& ChunkHash, uint32 HttpStatusCode, FIoBuffer&& Chunk)
+void FOnDemandContentInstaller::ProcessDownloadedChunk(
+	const FIoHash& ChunkHash, uint32 RequestId, EIoErrorCode InErrorCode, uint32 HttpStatusCode, FIoBuffer&& Chunk)
 {
 	TRACE_CPUPROFILER_EVENT_SCOPE(FOnDemandContentInstaller::ProcessDownloadedChunk);
 
@@ -1096,14 +1146,46 @@ void FOnDemandContentInstaller::ProcessDownloadedChunk(const FIoHash& ChunkHash,
 		Result = MakeError(UnifiedError::IoStoreOnDemand::HttpError::MakeError(HttpStatusCode));
 	}
 
-	FChunkDownloadRequest ChunkRequest = PendingChunkDownloads.FindAndRemoveChecked(ChunkHash);
-	for (FChunkHttpRequestHandle& ChunkHttpRequestHandle : ChunkRequest.ChunkRequestHttpHandles)
+	const bool bCancelled = (InErrorCode == EIoErrorCode::Cancelled);
+
+	bool bFoundRequest = false;
+	for (auto It = PendingChunkDownloads.CreateKeyIterator(ChunkHash); It; ++It)
 	{
-		OnProcessDownloadedChunkNotifyRequest(ChunkHttpRequestHandle, Result);
+		const FChunkDownloadRequest& ChunkRequest = It.Value();
+		if (ChunkRequest.RequestId != RequestId)
+		{
+			continue;
+		}
+
+		static constexpr const TCHAR LogFmt[] = TEXT("ProcessDownloadedChunk: ChunkHash=%s, RequestId=%u, DownloadResult=%s");
+		if (bCancelled)
+		{
+			UE_LOG(LogIoStoreOnDemand, Verbose, LogFmt, *TStringBuilder<64>(InPlace, ChunkHash), ChunkRequest.RequestId, TEXT("Cancelled"));
+		}
+		else if (Result.HasError())
+		{
+			UE_LOG(LogIoStoreOnDemand, Warning, LogFmt, *TStringBuilder<64>(InPlace, ChunkHash), ChunkRequest.RequestId, *Result.GetError().GetErrorMessage(true).ToString());
+		}
+		else
+		{
+			UE_LOG(LogIoStoreOnDemand, VeryVerbose, LogFmt, *TStringBuilder<64>(InPlace, ChunkHash), ChunkRequest.RequestId, TEXT("OK"));
+		}
+
+		bFoundRequest = true;
+		for (FChunkHttpRequestHandle& ChunkHttpRequestHandle : ChunkRequest.ChunkRequestHttpHandles)
+		{
+			OnProcessDownloadedChunkNotifyRequest(ChunkHttpRequestHandle, Result, bCancelled);
+		}
+
+		It.RemoveCurrent();
+		break;
 	}
+
+	check(bFoundRequest);
 }
 
-void FOnDemandContentInstaller::OnProcessDownloadedChunkNotifyRequest(FChunkHttpRequestHandle& HttpRequest, const FResult& ChunkDownloadResult)
+void FOnDemandContentInstaller::OnProcessDownloadedChunkNotifyRequest(
+	const FChunkHttpRequestHandle& HttpRequest, const FResult& ChunkDownloadResult, bool bChunkCancelled)
 {
 	TRACE_CPUPROFILER_EVENT_SCOPE(FOnDemandContentInstaller::OnProcessDownloadedChunkNotifyRequest);
 
@@ -1113,12 +1195,13 @@ void FOnDemandContentInstaller::OnProcessDownloadedChunkNotifyRequest(FChunkHttp
 	const FOnDemandChunkEntry& ChunkEntry = Container->ChunkEntries[HttpRequest.EntryIndex];
 	const FIoChunkId& ChunkId = Container->ChunkIds[HttpRequest.EntryIndex];
 
-	InstallRequest.Progress.CurrentInstallSize += ChunkEntry.GetDiskSize();
-
-	if (Request.IsOk())
+	bool bSetErrorOnRequest = false;
+	if (Request.IsOk() || (Request.IsCancelled() == false && bChunkCancelled))
 	{
 		if (const UnifiedError::FError* Error = ChunkDownloadResult.TryGetError())
 		{
+			bSetErrorOnRequest = true;
+
 			// Set the chunk Id from this container if needed
 			const UnifiedError::IoStoreOnDemand::FChunkHashMismatchErrorContext* ErrorContext =
 				Error->GetErrorContext<UnifiedError::IoStoreOnDemand::FChunkHashMismatchErrorContext>();
@@ -1134,18 +1217,37 @@ void FOnDemandContentInstaller::OnProcessDownloadedChunkNotifyRequest(FChunkHttp
 						.ActualHash = ErrorContext->ActualHash
 					}));
 			}
+			else
+			{
+				Request.Result = ChunkDownloadResult;
+			}
 		}
+	}
 
-		if (Request.IsOk())
-		{
-			Request.Result = ChunkDownloadResult;
-		}
+	static constexpr const TCHAR LogFmt[] = TEXT("Install progress %.2lf/%.2lf MiB, SeqNo=%u, Priority=%d, ContentHandle=(%s), ChunkId='%s', ChunkSize=%.2lf KiB, DownloadResult=%s");
+	if (bChunkCancelled)
+	{
+		UE_LOG(LogIoStoreOnDemand, Verbose, LogFmt,
+			double(InstallRequest.Progress.CurrentInstallSize) / 1024.0 / 1024.0, double(InstallRequest.Progress.TotalInstallSize) / 1024.0 / 1024.0,
+			Request.SeqNo, Request.Priority, *LexToString(InstallRequest.Args.ContentHandle), *LexToString(ChunkId), double(ChunkEntry.GetDiskSize()) / 1024.0,
+			TEXT("Cancelled"));
+	}
+	else if (ChunkDownloadResult.HasError())
+	{
+		UE_LOG(LogIoStoreOnDemand, Warning, LogFmt,
+			double(InstallRequest.Progress.CurrentInstallSize) / 1024.0 / 1024.0, double(InstallRequest.Progress.TotalInstallSize) / 1024.0 / 1024.0,
+			Request.SeqNo, Request.Priority, *LexToString(InstallRequest.Args.ContentHandle), *LexToString(ChunkId), double(ChunkEntry.GetDiskSize()) / 1024.0,
+			*ChunkDownloadResult.GetError().GetErrorMessage(true).ToString());
 	}
+	else
+	{
+		InstallRequest.Progress.CurrentInstallSize += ChunkEntry.GetDiskSize();
 
-	UE_LOG(LogIoStoreOnDemand, Verbose, TEXT("Install progress %.2lf/%.2lf MiB, SeqNo=%u, Priority=%d, ContentHandle=(%s), ChunkId='%s', ChunkSize=%.2lf KiB, DownloadResult=%s"),
-		double(InstallRequest.Progress.CurrentInstallSize) / 1024.0 / 1024.0, double(InstallRequest.Progress.TotalInstallSize) / 1024.0 / 1024.0,
-		Request.SeqNo, Request.Priority, *LexToString(InstallRequest.Args.ContentHandle), *LexToString(ChunkId), double(ChunkEntry.GetDiskSize()) / 1024.0, 
-		Request.Result.HasError() ? *Request.Result.GetError().GetErrorMessage(true).ToString() : TEXT("OK"));
+		UE_LOG(LogIoStoreOnDemand, Verbose, LogFmt,
+			double(InstallRequest.Progress.CurrentInstallSize) / 1024.0 / 1024.0, double(InstallRequest.Progress.TotalInstallSize) / 1024.0 / 1024.0,
+			Request.SeqNo, Request.Priority, *LexToString(InstallRequest.Args.ContentHandle), *LexToString(ChunkId), double(ChunkEntry.GetDiskSize()) / 1024.0,
+			TEXT("OK"));
+	}
 
 	const bool bCompleted = ++InstallRequest.DownloadedChunkCount >= InstallRequest.HttpRequestHandles.Num();
 	if (bCompleted)
@@ -1156,51 +1258,88 @@ void FOnDemandContentInstaller::OnProcessDownloadedChunkNotifyRequest(FChunkHttp
 	{
 		NotifyInstallProgress(Request);
 
-		if (Request.IsOk() == false)
+		// Only try to cancel requests the first time an error is set
+		if (bSetErrorOnRequest) 
 		{
 			int32 NumCancelled = 0;
-			for (FChunkHttpRequestHandle& PendingHttpRequest : InstallRequest.HttpRequestHandles)
+			TryCancelHttpRequestsForInstallRequest(InstallRequest, NumCancelled);
+			UE_CLOG(NumCancelled == 0, LogIoStoreOnDemand, Warning, TEXT("Cancelled %d HTTP request(s) due to install error"), NumCancelled);
+		}
+	}
+}
+
+void FOnDemandContentInstaller::TryCancelHttpRequestsForInstallRequest(FRequest::FInstall& InstallRequest, int32& OutNumCancelled)
+{
+	int32 NumCancelled = 0;
+
+	if (CVars::bContentInstallerEnableHttpCancel == false)
+	{
+		OutNumCancelled = NumCancelled;
+		return;
+	}
+
+	for (FChunkHttpRequestHandle& PendingHttpRequest : InstallRequest.HttpRequestHandles)
+	{
+		const FSharedOnDemandContainer& PendingContainer = InstallRequest.ResolvedChunks[PendingHttpRequest.ContainerIndex].Container;
+		const FOnDemandChunkEntry& PendingChunkEntry = PendingContainer->ChunkEntries[PendingHttpRequest.EntryIndex];
+		const FIoHash& PendingChunkHash = PendingChunkEntry.Hash;
+
+		FChunkDownloadRequest* ChunkDownload = nullptr;
+		for (auto It = PendingChunkDownloads.CreateKeyIterator(PendingChunkHash); It; ++It)
+		{
+			FChunkDownloadRequest& FoundChunkDownload = It.Value();
+			if (FoundChunkDownload.bChunkCanceled == false)
 			{
-				const FSharedOnDemandContainer& PendingContainer = InstallRequest.ResolvedChunks[PendingHttpRequest.ContainerIndex].Container;
-				const FOnDemandChunkEntry& PendingChunkEntry = PendingContainer->ChunkEntries[PendingHttpRequest.EntryIndex];
-				const FIoHash& PendingChunkHash = PendingChunkEntry.Hash;
+				ChunkDownload = &FoundChunkDownload;
+				break; // Should be a max of one active request for any chunk
+			}
+		}
 
-				if (FChunkDownloadRequest* ChunkDownload = PendingChunkDownloads.Find(PendingChunkHash))
+		if (ChunkDownload)
+		{
+			bool bCancelHttpRequest = true;
+			for (FChunkHttpRequestHandle& Handle : ChunkDownload->ChunkRequestHttpHandles)
+			{
+				if (Handle.OwnerRequest->IsOk())
 				{
-					// Remove PendingHttpRequest request from PendingChunkDownloads
-					ChunkDownload->ChunkRequestHttpHandles.Remove(&PendingHttpRequest);
-					if (ChunkDownload->ChunkRequestHttpHandles.IsEmpty())
+					bCancelHttpRequest = false;
+					break;
+				}
+			}
+
+			if (bCancelHttpRequest)
+			{
+				// No requests are waiting on this chunk, cancel it
+				Visit(THttpHandleVisitor
+				{
+					[&NumCancelled](FIoHttpRequest& Handle)
 					{
-						// No requests are waiting on this chunk, cancel it
-						Visit(THttpHandleVisitor
+						if (Handle.IsValid())
 						{
-							[&NumCancelled](FIoHttpRequest& Handle)
-							{
-								if (Handle.IsValid())
-								{
-									Handle.Cancel();
-									++NumCancelled;
-								}
-							},
-							[this, &NumCancelled](void* Handle)
-							{
-								check(HttpClient != nullptr);
-								if (Handle != nullptr)
-								{
-									HttpClient->CancelRequest(Handle);
-									++NumCancelled;
-								}
-							}
-						}, ChunkDownload->HttpHandle);
-
-						// ProcessDownloadedChunk will remove the chunk request
+							Handle.Cancel();
+							++NumCancelled;
+						}
+					},
+					[this, &NumCancelled](void* Handle)
+					{
+						check(HttpClient != nullptr);
+						if (Handle != nullptr)
+						{
+							HttpClient->CancelRequest(Handle);
+							++NumCancelled;
+						}
 					}
-				}
-			}
+				}, ChunkDownload->HttpHandle);
+
+				ChunkDownload->bChunkCanceled = true;
+				UE_LOG(LogIoStoreOnDemand, Verbose, TEXT("Cancelled chunk download: hash - %s, Id - %u"), *TStringBuilder<64>(InPlace, PendingChunkHash), ChunkDownload->RequestId);
 
-			UE_CLOG(NumCancelled == 0, LogIoStoreOnDemand, Log, TEXT("Cancelled %d HTTP request(s) due to install error"), NumCancelled);
+				// ProcessDownloadedChunk will remove the chunk request
+			}
 		}
 	}
+
+	OutNumCancelled = NumCancelled;
 }
 
 void FOnDemandContentInstaller::NotifyInstallProgress(FRequest& Request)
@@ -1246,7 +1385,7 @@ void FOnDemandContentInstaller::DestroyRequest(FRequest& Request)
 	// Call destructor outside the lock to avoid double locking
 	Request.~FRequest();
 
-	UE::TUniqueLock Lock(Mutex);
+	TUniqueLock Lock(Mutex);
 	RequestAllocator.Free(&Request);
 }
 
@@ -1326,7 +1465,7 @@ void FOnDemandContentInstaller::CompleteInstallRequest(FRequest& Request)
 	InstallResult.Error				= Request.ConsumeError();
 
 	{
-		UE::TUniqueLock Lock(Mutex);
+		TUniqueLock Lock(Mutex);
 
 		InstallRequest.Request->InstallerRequest = 0;
 		RunningRequests.RemoveSingleSwap(&Request);
@@ -1391,7 +1530,7 @@ void FOnDemandContentInstaller::CompletePurgeRequest(FRequest& Request)
 	const bool bCallbackOnGameThread	= EnumHasAnyFlags(Request.AsPurge().Args.Options, EOnDemandPurgeOptions::CallbackOnGameThread);
 	FOnDemandPurgeCompleted OnCompleted = MoveTemp(Request.AsPurge().OnCompleted);
 	{
-		UE::TUniqueLock Lock(Mutex);
+		TUniqueLock Lock(Mutex);
 		RunningRequests.RemoveSingleSwap(&Request);
 	}
 	DestroyRequest(Request);
@@ -1434,7 +1573,7 @@ void FOnDemandContentInstaller::CompleteDefragRequest(FRequest& Request)
 	const bool bCallbackOnGameThread		= EnumHasAnyFlags(Request.AsDefrag().Args.Options, EOnDemandDefragOptions::CallbackOnGameThread);
 	FOnDemandDefragCompleted OnCompleted	= MoveTemp(Request.AsDefrag().OnCompleted);
 	{
-		UE::TUniqueLock Lock(Mutex);
+		TUniqueLock Lock(Mutex);
 		RunningRequests.RemoveSingleSwap(&Request);
 	}
 	DestroyRequest(Request);
@@ -1474,7 +1613,7 @@ void FOnDemandContentInstaller::CompleteVerifyRequest(FRequest& Request)
 
 	FOnDemandVerifyCacheCompleted OnCompleted = MoveTemp(Request.AsVerify().OnCompleted);
 	{
-		UE::TUniqueLock Lock(Mutex);
+		TUniqueLock Lock(Mutex);
 		RunningRequests.RemoveSingleSwap(&Request);
 	}
 	DestroyRequest(Request);
@@ -1502,7 +1641,7 @@ void FOnDemandContentInstaller::CompleteFlushLastAccessRequest(FRequest& Request
 
 	FOnDemandFlushLastAccessCompleted OnCompleted = MoveTemp(Request.AsFlushLastAccess().OnCompleted);
 	{
-		UE::TUniqueLock Lock(Mutex);
+		TUniqueLock Lock(Mutex);
 		RunningRequests.RemoveSingleSwap(&Request);
 	}
 	DestroyRequest(Request);
diff --git a/Engine/Source/Runtime/Experimental/IoStore/OnDemand/Private/OnDemandContentInstaller.h b/Engine/Source/Runtime/Experimental/IoStore/OnDemand/Private/OnDemandContentInstaller.h
--- a/Engine/Source/Runtime/Experimental/IoStore/OnDemand/Private/OnDemandContentInstaller.h
+++ b/Engine/Source/Runtime/Experimental/IoStore/OnDemand/Private/OnDemandContentInstaller.h
@@ -93,7 +93,7 @@ class FOnDemandContentInstaller
 			FOnDemandInstallProgress					Progress;
 			uint64										DownloadedChunkCount = 0;
 			uint64										LastProgressCycles = 0;
-			std::atomic_bool							bHttpRequestsIssued{false}; // TODO: does this need to be atomic? All access is from the installer pipe
+			bool										bHttpRequestsIssued = false;
 			std::atomic_bool							bNotifyingProgressOnGameThread{false};
 		};
 
@@ -160,8 +160,13 @@ class FOnDemandContentInstaller
 			Variant.Emplace<FFlushLastAccess>(MoveTemp(OnCompleted));
 		}
 
-		bool				IsCancelled() const			{ return bCancelled.load(std::memory_order_relaxed); }
-		bool				TryCancel()					{ bool bExpected = false; return bCancelled.compare_exchange_strong(bExpected, true); }
+		bool				IsCancelled() const			{ return bCancelled; }
+		bool				TryCancel()
+		{ 
+			bool bPrev = bCancelled;
+			bCancelled = true;
+			return bPrev == false;
+		}
 		bool				IsOk() const				{ return !IsCancelled() && !Result.HasError(); }
 		FString				GetErrorString() const		{ return IsOk() ? TEXT("Ok") : IsCancelled() ? TEXT("Cancelled") : LexToString(Result.GetError()); }
 
@@ -185,7 +190,7 @@ class FOnDemandContentInstaller
 		int32						Priority = 0;
 		uint64						StartTimeCycles = FPlatformTime::Cycles64();
 		FResult						Result = MakeValue();
-		std::atomic_bool			bCancelled{false};
+		bool						bCancelled = false;
 		FRequestVariant				Variant;
 	};
 
@@ -193,6 +198,10 @@ class FOnDemandContentInstaller
 	{
 		TVariant<void*, FIoHttpRequest> HttpHandle;
 		TIntrusiveList<FChunkHttpRequestHandle> ChunkRequestHttpHandles;
+
+		// Should only be used to compare two requests for the same iochunk, otherwise RequestId is not unique
+		uint32				RequestId = 0;
+		bool				bChunkCanceled = false;
 	};
 
 	static bool RequestSortPredicate(const FRequest& LHS, const FRequest& RHS)
@@ -239,8 +248,14 @@ private:
 	void							ExecuteDefragRequest(FRequest& Request);
 	void							ExecuteVerifyRequest(FRequest& Request);
 	void							ExectuteFlushLastAccessRequest(FRequest& Request);
-	void							ProcessDownloadedChunk(const FIoHash& ChunkHash, uint32 HttpStatusCode, FIoBuffer&& Chunk);
-	void							OnProcessDownloadedChunkNotifyRequest(FChunkHttpRequestHandle& HttpRequest, const FResult& ChunkDownloadResult);
+	void							ProcessDownloadedChunk(
+										const FIoHash& ChunkHash, 
+										uint32 RequestId, 
+										EIoErrorCode InErrorCode, 
+										uint32 HttpStatusCode, 
+										FIoBuffer&& Chunk);
+	void							OnProcessDownloadedChunkNotifyRequest(const FChunkHttpRequestHandle& HttpRequest, const FResult& ChunkDownloadResult, bool bChunkCancelled);
+	void							TryCancelHttpRequestsForInstallRequest(FRequest::FInstall& InstallRequest, int32& OutNumCancelled);
 	void							NotifyInstallProgress(FRequest& Request);
 	void							DestroyRequest(FRequest& Request);
 	void							CompleteInstallRequest(FRequest& Request);
@@ -258,7 +273,10 @@ private:
 	FRequestAllocator				RequestAllocator;
 	TArray<FRequest*>				RequestQueue;
 	TArray<FRequest*>				RunningRequests;
-	TMap<FIoHash, FChunkDownloadRequest> PendingChunkDownloads;
+	// This needs to be a multimap to support cancelation. At any time, there may be multiple
+	// requests for the same iochunk but only one of them should be active with the remainder
+	// waiting on final callbacks after being cancelled. See FChunkDownloadRequest.
+	TMultiMap<FIoHash, FChunkDownloadRequest> PendingChunkDownloads;
 	std::atomic_bool				bShuttingDown{false};
 };
 
