IAD Installer
- Add ensure to catch an http request being unexpectedly cancelled for an active install request.
- Use new UE::Overload template for chunk handle visitor instead of a bespoke template.

diff --git a/Engine/Source/Runtime/Experimental/IoStore/OnDemand/Private/OnDemandContentInstaller.cpp b/Engine/Source/Runtime/Experimental/IoStore/OnDemand/Private/OnDemandContentInstaller.cpp
--- a/Engine/Source/Runtime/Experimental/IoStore/OnDemand/Private/OnDemandContentInstaller.cpp
+++ b/Engine/Source/Runtime/Experimental/IoStore/OnDemand/Private/OnDemandContentInstaller.cpp
@@ -20,6 +20,7 @@
 #include "OnDemandPackageStoreBackend.h"
 #include "Serialization/PackageStore.h"
 #include "Statistics.h"
+#include "Templates/Overload.h"
 #include <atomic>
 
 namespace UE::IoStore
@@ -262,14 +263,6 @@ void ResolveChunksToInstall(
 
 } // namespace Private
 
-////////////////////////////////////////////////////////////////////////////////
-// Lambda visitor for FChunkDownloadRequest::HttpHandle
-template<class... TFunc>
-struct THttpHandleVisitor : TFunc...
-{
-	using TFunc::operator()...;
-};
-
 ////////////////////////////////////////////////////////////////////////////////
 uint32 FOnDemandContentInstaller::FRequest::NextSeqNo = 0;
 
@@ -469,8 +462,7 @@ void FOnDemandContentInstaller::UpdateInstallRequestPriority(FSharedInternalInst
 							}
 						}
 
-						Visit(THttpHandleVisitor
-						{
+						Visit(UE::Overload(
 							[MaxPriority](FIoHttpRequest& Handle)
 							{
 								Handle.UpdatePriorty(MaxPriority);
@@ -482,8 +474,8 @@ void FOnDemandContentInstaller::UpdateInstallRequestPriority(FSharedInternalInst
 								{
 									HttpClient->ReprioritizeRequest(Handle, MaxPriority);
 								}
-							}
-						}, ChunkDownload.HttpHandle);
+							}), 
+							ChunkDownload.HttpHandle);
 
 						// Should be a max of one active request for any chunk
 						break;
@@ -1196,9 +1188,14 @@ void FOnDemandContentInstaller::OnProcessDownloadedChunkNotifyRequest(
 	const FIoChunkId& ChunkId = Container->ChunkIds[HttpRequest.EntryIndex];
 
 	bool bSetErrorOnRequest = false;
-	if (Request.IsOk() || (Request.IsCancelled() == false && bChunkCancelled))
+	if (Request.IsOk())
 	{
-		if (const UnifiedError::FError* Error = ChunkDownloadResult.TryGetError())
+		if (!ensureMsgf(!bChunkCancelled, TEXT("Recieved cancelled http callback for active Request %s ! This should never happen."), *InstallRequest.Args.DebugName))
+		{
+			bSetErrorOnRequest = true;
+			Request.Result = ChunkDownloadResult;
+		}
+		else if (const UnifiedError::FError* Error = ChunkDownloadResult.TryGetError())
 		{
 			bSetErrorOnRequest = true;
 
@@ -1310,8 +1307,7 @@ void FOnDemandContentInstaller::TryCancelHttpRequestsForInstallRequest(FRequest:
 			if (bCancelHttpRequest)
 			{
 				// No requests are waiting on this chunk, cancel it
-				Visit(THttpHandleVisitor
-				{
+				Visit(UE::Overload(
 					[&NumCancelled](FIoHttpRequest& Handle)
 					{
 						if (Handle.IsValid())
@@ -1328,8 +1324,8 @@ void FOnDemandContentInstaller::TryCancelHttpRequestsForInstallRequest(FRequest:
 							HttpClient->CancelRequest(Handle);
 							++NumCancelled;
 						}
-					}
-				}, ChunkDownload->HttpHandle);
+					}),
+					ChunkDownload->HttpHandle);
 
 				ChunkDownload->bChunkCanceled = true;
 				UE_LOG(LogIoStoreOnDemand, Verbose, TEXT("Cancelled chunk download: hash - %s, Id - %u"), *TStringBuilder<64>(InPlace, PendingChunkHash), ChunkDownload->RequestId);
