Add support for including response headers with option flag in the HTTP I/O dispatcher response message.

#rnx

diff --git a/Engine/Source/Runtime/Experimental/IoStore/OnDemand/Private/OnDemandHttpClient.cpp b/Engine/Source/Runtime/Experimental/IoStore/OnDemand/Private/OnDemandHttpClient.cpp
--- a/Engine/Source/Runtime/Experimental/IoStore/OnDemand/Private/OnDemandHttpClient.cpp
+++ b/Engine/Source/Runtime/Experimental/IoStore/OnDemand/Private/OnDemandHttpClient.cpp
@@ -254,22 +254,29 @@ FMultiEndpointHttpClient::FHttpTicketId FMultiEndpointHttpClient::Get(const FOnD
 
 FMultiEndpointHttpClient::FHttpTicketId FMultiEndpointHttpClient::Get(const FOnDemandHostGroup& HostGroup, FAnsiStringView RelativeUrl, const FIoOffsetAndLength& ChunkRange, FOnHttpResponse&& OnResponse)
 {
-	return Get(HostGroup, RelativeUrl, ChunkRange, TArray<FAnsiString>(), MoveTemp(OnResponse));
+	return Get(HostGroup, RelativeUrl, ChunkRange, TArray<FAnsiString>(), EMultiEndpointRequestFlags::None, MoveTemp(OnResponse));
 }
 
-FMultiEndpointHttpClient::FHttpTicketId FMultiEndpointHttpClient::Get(const FOnDemandHostGroup& HostGroup, FAnsiStringView RelativeUrl, const FIoOffsetAndLength& ChunkRange, TArray<FAnsiString>&& Headers, FOnHttpResponse&& OnResponse)
+FMultiEndpointHttpClient::FHttpTicketId FMultiEndpointHttpClient::Get(
+	const FOnDemandHostGroup& HostGroup,
+	FAnsiStringView RelativeUrl,
+	const FIoOffsetAndLength& ChunkRange,
+	TArray<FAnsiString>&& Headers,
+	EMultiEndpointRequestFlags Flags,
+	FOnHttpResponse&& OnResponse)
 {
 	FConnection& Connection = GetConnection(HostGroup);
 
 	FHttpTicketId TicketId = IssueRequest(FRequest
 	{
 		.OnResponse		= MoveTemp(OnResponse),
-		.Headers		= MoveTemp(Headers),
+		.RequestHeaders	= MoveTemp(Headers),
 		.RelativeUrl	= FAnsiString(RelativeUrl),
 		.Range			= ChunkRange,
 		.Connection		= Connection,
 		.StartTime		= FPlatformTime::Cycles64(),
-		.Host			= Connection.CurrentHost
+		.Host			= Connection.CurrentHost,
+		.Flags			= Flags
 	});
 
 	return TicketId;
@@ -373,12 +380,12 @@ FMultiEndpointHttpClient::FHttpTicketId FMultiEndpointHttpClient::IssueRequest(F
 	HttpRequest.Header("pragma", "akamai-x-cache-on");
 #endif
 
-	check(Request.Headers.IsEmpty() || ((Request.Headers.Num() % 2) == 0));
-	if (!Request.Headers.IsEmpty())
+	check(Request.RequestHeaders.IsEmpty() || ((Request.RequestHeaders.Num() % 2) == 0));
+	if (!Request.RequestHeaders.IsEmpty())
 	{
-		for (int32 Idx = 0; Idx < Request.Headers.Num(); Idx += 2)
+		for (int32 Idx = 0; Idx < Request.RequestHeaders.Num(); Idx += 2)
 		{
-			HttpRequest.Header(Request.Headers[Idx], Request.Headers[Idx + 1]);
+			HttpRequest.Header(Request.RequestHeaders[Idx], Request.RequestHeaders[Idx + 1]);
 		}
 	}
 
@@ -407,6 +414,18 @@ FMultiEndpointHttpClient::FHttpTicketId FMultiEndpointHttpClient::IssueRequest(F
 				Request.StatusCode = HttpResponse.GetStatusCode();
 				Request.bIsChunkedTransfer = HttpResponse.GetContentLength() == -1;
 
+				if (EnumHasAnyFlags(Request.Flags, EMultiEndpointRequestFlags::ResponseHeaders))
+				{
+					Request.ResponseHeaders.Empty(16);
+					HttpResponse.ReadHeaders([&Request] (FAnsiStringView Key, FAnsiStringView Value)
+					{
+						//TODO: Can we get all headers in one buffer?
+						Request.ResponseHeaders.Add(FAnsiString(Key));
+						Request.ResponseHeaders.Add(FAnsiString(Value));
+						return true;
+					});
+				}
+
 				if (IsHttpStatusOk(Request.StatusCode))
 				{
 					Request.CDNCacheStatus = TrackCdnCacheStats(HttpResponse);
@@ -454,6 +473,7 @@ FMultiEndpointHttpClient::FHttpTicketId FMultiEndpointHttpClient::IssueRequest(F
 				{
 					FMultiEndpointHttpClientResponse Response
 					{
+						.Headers				= MoveTemp(Request.ResponseHeaders),
 						.DurationMilliseconds	= uint64(FPlatformTime::ToMilliseconds64(FPlatformTime::Cycles64() - Request.StartTime)),
 						.StatusCode				= Request.StatusCode,
 						.RetryCount				= Request.RetryCount,
diff --git a/Engine/Source/Runtime/Experimental/IoStore/OnDemand/Private/OnDemandHttpClient.h b/Engine/Source/Runtime/Experimental/IoStore/OnDemand/Private/OnDemandHttpClient.h
--- a/Engine/Source/Runtime/Experimental/IoStore/OnDemand/Private/OnDemandHttpClient.h
+++ b/Engine/Source/Runtime/Experimental/IoStore/OnDemand/Private/OnDemandHttpClient.h
@@ -11,6 +11,7 @@
 #include "IO/IoOffsetLength.h"
 #include "IO/IoStatus.h"
 #include "IO/OnDemandHostGroup.h"
+#include "Misc/EnumClassFlags.h"
 #include "Templates/Function.h"
 #include "Templates/UniquePtr.h"
 
@@ -72,6 +73,13 @@ struct FMultiEndpointHttpClientConfig
 	ELogVerbosity::Type			LogVerbosity = ELogVerbosity::Log;
 };
 
+enum class EMultiEndpointRequestFlags : uint8
+{
+	None,
+	ResponseHeaders
+};
+ENUM_CLASS_FLAGS(EMultiEndpointRequestFlags);
+
 struct FMultiEndpointHttpClientResponse
 {
 	bool IsOk() const
@@ -85,6 +93,7 @@ struct FMultiEndpointHttpClientResponse
 	}
 
 	FIoBuffer					Body;
+	TArray<FAnsiString>			Headers;
 	FString						Reason;
 	HTTP::FTicketPerf::FSample	Sample;
 	uint64						DurationMilliseconds = 0;
@@ -122,7 +131,7 @@ public:
 
 	FHttpTicketId Get(const FOnDemandHostGroup& HostGroup, FAnsiStringView RelativeUrl, FOnHttpResponse&& OnResponse);
 	FHttpTicketId Get(const FOnDemandHostGroup& HostGroup, FAnsiStringView RelativeUrl, const FIoOffsetAndLength& ChunkRange, FOnHttpResponse&& OnResponse);
-	FHttpTicketId Get(const FOnDemandHostGroup& HostGroup, FAnsiStringView RelativeUrl, const FIoOffsetAndLength& ChunkRange, TArray<FAnsiString>&& Headers, FOnHttpResponse&& OnResponse);
+	FHttpTicketId Get(const FOnDemandHostGroup& HostGroup, FAnsiStringView RelativeUrl, const FIoOffsetAndLength& ChunkRange, TArray<FAnsiString>&& Headers, EMultiEndpointRequestFlags Flags, FOnHttpResponse&& OnResponse);
 
 	/**
 	 * Process the underlying http client.
@@ -155,27 +164,29 @@ private:
 
 	struct FRequest
 	{
-		FOnHttpResponse		OnResponse;
-		TArray<FAnsiString>	Headers;
-		FAnsiString			RelativeUrl;
-		FIoOffsetAndLength	Range;
-		FHttpTicketId		TicketId = 0;
-		FConnection&		Connection;
+		FOnHttpResponse				OnResponse;
+		TArray<FAnsiString>			RequestHeaders;
+		TArray<FAnsiString>			ResponseHeaders;
+		FAnsiString					RelativeUrl;
+		FIoOffsetAndLength			Range;
+		FHttpTicketId				TicketId = 0;
+		FConnection&				Connection;
 		/** Stored the response body */
-		FIoBuffer			Body;
+		FIoBuffer					Body;
 		/**
 		 * Stores each chunk being sent if the response uses chunked transfer, the chunks will be combined
 		 * together so that the final result will be in Body once all chunks are processed.
 		 */
-		FIoBuffer			Chunk;
+		FIoBuffer					Chunk;
 		/** Stores the response message if the StatusCode indicates a problem */
-		FString				ResponseMessage;
-		uint64				StartTime = 0;
-		uint32				RetryCount = 0;
-		uint32				StatusCode = 0;
-		int32				Host = INDEX_NONE;
-		bool				bIsChunkedTransfer = false;
-		int8				CDNCacheStatus = -1;
+		FString						ResponseMessage;
+		uint64						StartTime = 0;
+		uint32						RetryCount = 0;
+		uint32						StatusCode = 0;
+		int32						Host = INDEX_NONE;
+		EMultiEndpointRequestFlags	Flags = EMultiEndpointRequestFlags::None;
+		bool						bIsChunkedTransfer = false;
+		int8						CDNCacheStatus = -1;
 
 #if DO_CHECK
 		/** Used to make sure that the sink is only called with a completed state one per request attempt */
diff --git a/Engine/Source/Runtime/Experimental/IoStore/OnDemand/Private/OnDemandHttpIoDispacher.cpp b/Engine/Source/Runtime/Experimental/IoStore/OnDemand/Private/OnDemandHttpIoDispacher.cpp
--- a/Engine/Source/Runtime/Experimental/IoStore/OnDemand/Private/OnDemandHttpIoDispacher.cpp
+++ b/Engine/Source/Runtime/Experimental/IoStore/OnDemand/Private/OnDemandHttpIoDispacher.cpp
@@ -660,7 +660,11 @@ private:
 	bool							ProcessHttpRequests();
 	FHttpRequest*					GetNextRequest();
 	void							CompleteHttpRequest(FHttpRequest& Request, FMultiEndpointHttpClientResponse&& HttpResponse, FIASHostGroup& HostGroup, const FIoHttpRange& TotalRange);
-	void							CompleteRequest(FHttpRequest& Request, FIoBuffer& Body, const FIoHttpRange& TotalRange, uint32 StatusCode, bool bCached);
+	void							CompleteRequest(FHttpRequest& Request, TArray<FAnsiString>&& Headers, FIoBuffer& Body, const FIoHttpRange& TotalRange, uint32 StatusCode, bool bCached);
+	void							CompleteRequest(FHttpRequest& Request, FIoBuffer& Body, const FIoHttpRange& TotalRange, uint32 StatusCode, bool bCached)
+									{
+										CompleteRequest(Request, TArray<FAnsiString>(), Body, TotalRange, StatusCode, bCached);
+									}
 
 #if UE_ALLOW_HTTP_PAUSE
 	void OnPauseCommand(bool bShouldPause, const TArray<FString>& Args, FOutputDevice& Ar);
@@ -1254,8 +1258,12 @@ bool FHttpDispatcher::ProcessHttpRequests()
 			NumConcurrentRequests++;
 			FOnDemandIoBackendStats::Get()->OnHttpRequestStarted();
 
+			EMultiEndpointRequestFlags Flags = EnumHasAnyFlags(HttpRequest->Options.GetFlags(), EIoHttpFlags::ResponseHeaders)
+				? EMultiEndpointRequestFlags::ResponseHeaders
+				: EMultiEndpointRequestFlags::None;
+
 			// Note that the requests issued here must be completed before this method can return so the capturing of locals is safe.
-			HttpClient->Get(HostGroup.GetUnderlyingHostGroup(), Url, TotalRange.ToOffsetAndLength(), MoveTemp(Headers),
+			HttpClient->Get(HostGroup.GetUnderlyingHostGroup(), Url, TotalRange.ToOffsetAndLength(), MoveTemp(Headers), Flags,
 				[this, HttpRequest, TotalRange, &NumConcurrentRequests, HostGroup = MoveTemp(HostGroup)]
 				(FMultiEndpointHttpClientResponse&& HttpResponse) mutable
 				{
@@ -1339,10 +1347,10 @@ void FHttpDispatcher::CompleteHttpRequest(FHttpRequest& Request, FMultiEndpointH
 		}
 	}
 
-	CompleteRequest(Request, HttpResponse.Body, TotalRange, HttpResponse.StatusCode, false);
+	CompleteRequest(Request, MoveTemp(HttpResponse.Headers), HttpResponse.Body, TotalRange, HttpResponse.StatusCode, false);
 }
 
-void FHttpDispatcher::CompleteRequest(FHttpRequest& Request, FIoBuffer& Body, const FIoHttpRange& TotalRange, uint32 StatusCode, bool bCached)
+void FHttpDispatcher::CompleteRequest(FHttpRequest& Request, TArray<FAnsiString>&& Headers, FIoBuffer& Body, const FIoHttpRange& TotalRange, uint32 StatusCode, bool bCached)
 {
 #if UE_ALLOW_CACHE_POISONING
 	if (GIasPoisonCache && bCached)
@@ -1395,7 +1403,7 @@ void FHttpDispatcher::CompleteRequest(FHttpRequest& Request, FIoBuffer& Body, co
 
 		const EIoHttpResponseFlags Flags	= bCached ? EIoHttpResponseFlags::Cached : EIoHttpResponseFlags::None;
 		FIoHttpRequestCompleted OnCompleted = MoveTemp(ToComplete->OnCompleted);
-		OnCompleted(FIoHttpResponse(PartialBody, ToComplete->CacheKey, CompletionStatus, StatusCode, Flags));
+		OnCompleted(FIoHttpResponse(ToComplete->CacheKey, FIoHttpHeaders::Create(MoveTemp(Headers)), PartialBody, CompletionStatus, StatusCode, Flags));
 
 		ToComplete->Flags.Add(EHttpRequestFlags::Completed);
 		ToComplete->CompletionStatus.store(CompletionStatus, std::memory_order_seq_cst);
diff --git a/Engine/Source/Runtime/Experimental/IoStore/OnDemandCore/Private/HttpIoDispatcher.cpp b/Engine/Source/Runtime/Experimental/IoStore/OnDemandCore/Private/HttpIoDispatcher.cpp
--- a/Engine/Source/Runtime/Experimental/IoStore/OnDemandCore/Private/HttpIoDispatcher.cpp
+++ b/Engine/Source/Runtime/Experimental/IoStore/OnDemandCore/Private/HttpIoDispatcher.cpp
@@ -47,6 +47,12 @@ FIoHttpHeaders FIoHttpHeaders::Create(FAnsiStringView HeaderName, FAnsiStringVie
 	return Headers;
 }
 
+FIoHttpHeaders FIoHttpHeaders::Create(TArray<FAnsiString>&& Headers)
+{
+	check(Headers.IsEmpty() || ((Headers.Num() % 2) == 0));
+	return FIoHttpHeaders(MoveTemp(Headers));
+}
+
 ////////////////////////////////////////////////////////////////////////////////
 FIoRelativeUrl FIoRelativeUrl::From(FAnsiStringView Url)
 {
diff --git a/Engine/Source/Runtime/Experimental/IoStore/OnDemandCore/Public/IO/HttpIoDispatcher.h b/Engine/Source/Runtime/Experimental/IoStore/OnDemandCore/Public/IO/HttpIoDispatcher.h
--- a/Engine/Source/Runtime/Experimental/IoStore/OnDemandCore/Public/IO/HttpIoDispatcher.h
+++ b/Engine/Source/Runtime/Experimental/IoStore/OnDemandCore/Public/IO/HttpIoDispatcher.h
@@ -28,9 +28,11 @@ enum class EIoHttpFlags
 {
 	/** No additional flags. */
 	None			= 0,
-	/** Whether to read from the HTTP cache. */ 
+	/** Whether to read from the HTTP cache. */
 	ReadCache		= (1 << 0),
-	/** Default flags . */
+	/** Whether to include response headers. */
+	ResponseHeaders	= (1 << 1),
+	/** Default flags. */
 	Default			= ReadCache
 };
 ENUM_CLASS_FLAGS(EIoHttpFlags);
@@ -170,11 +172,15 @@ public:
 
 	UE_API static FIoHttpHeaders	Create(FAnsiString&& HeaderName, FAnsiString&& HeaderValue);
 	UE_API static FIoHttpHeaders	Create(FAnsiStringView HeaderName, FAnsiStringView HeaderValue);
+	UE_API static FIoHttpHeaders	Create(TArray<FAnsiString>&& Headers);
 	
 	FIoHttpHeaders&					operator=(const FIoHttpHeaders& Other)	{ Headers = Other.Headers; return *this; }
 	FIoHttpHeaders&					operator=(FIoHttpHeaders&& Other)		{ Headers = MoveTemp(Other.Headers); return *this; }
 
 private:
+									FIoHttpHeaders(TArray<FAnsiString>&& InHeaders)
+										: Headers(MoveTemp(InHeaders))
+									{ }
 	//TODO: Make this better
 	TArray<FAnsiString> Headers;
 };
@@ -242,11 +248,13 @@ class FIoHttpResponse
 {
 public:
 							FIoHttpResponse() = default;
-	inline					FIoHttpResponse(const FIoBuffer& InBody, const FIoHash& InCacheKey, EIoErrorCode InErrorCode, uint32 InStatusCode, EIoHttpResponseFlags InFlags);
+	inline					FIoHttpResponse(const FIoHash& InCacheKey, const FIoBuffer& InBody, EIoErrorCode InErrorCode, uint32 InStatusCode, EIoHttpResponseFlags InFlags);
+	inline					FIoHttpResponse(const FIoHash& InCacheKey, FIoHttpHeaders&& InHeaders, const FIoBuffer& InBody, EIoErrorCode InErrorCode, uint32 InStatusCode, EIoHttpResponseFlags InFlags);
 							FIoHttpResponse(EIoErrorCode InErrorCode, uint32 InStatusCode)
 								: ErrorCode(InErrorCode)
 								, StatusCode(InStatusCode) { }
 
+	const FIoHttpHeaders&	GetHeaders()	const { return Headers; }
 	const FIoBuffer&		GetBody()		const { return Body; }
 	const FIoHash&			GetCacheKey()	const { return CacheKey; }
 	EIoErrorCode			GetErrorCode()	const { return ErrorCode; }
@@ -257,16 +265,27 @@ public:
 	bool					IsCached()		const { return EnumHasAnyFlags(Flags, EIoHttpResponseFlags::Cached); } 
 
 private:
-	FIoBuffer				Body;
 	FIoHash					CacheKey;
+	FIoHttpHeaders			Headers;
+	FIoBuffer				Body;
 	EIoErrorCode			ErrorCode;
 	uint32					StatusCode = 0;
 	EIoHttpResponseFlags	Flags = EIoHttpResponseFlags::None;
 };
 
-FIoHttpResponse::FIoHttpResponse(const FIoBuffer& InBody, const FIoHash& InCacheKey, EIoErrorCode InErrorCode, uint32 InStatusCode, EIoHttpResponseFlags InFlags)
-	: Body(InBody)
-	, CacheKey(InCacheKey)
+FIoHttpResponse::FIoHttpResponse(const FIoHash& InCacheKey, const FIoBuffer& InBody, EIoErrorCode InErrorCode, uint32 InStatusCode, EIoHttpResponseFlags InFlags)
+	: CacheKey(InCacheKey)
+	, Body(InBody)
+	, ErrorCode(InErrorCode)
+	, StatusCode(InStatusCode)
+	, Flags(InFlags)
+{
+}
+
+FIoHttpResponse::FIoHttpResponse(const FIoHash& InCacheKey, FIoHttpHeaders&& InHeaders, const FIoBuffer& InBody, EIoErrorCode InErrorCode, uint32 InStatusCode, EIoHttpResponseFlags InFlags)
+	: CacheKey(InCacheKey)
+	, Headers(MoveTemp(InHeaders))
+	, Body(InBody)
 	, ErrorCode(InErrorCode)
 	, StatusCode(InStatusCode)
 	, Flags(InFlags)
