Optimize Push Data
Removed some DirtyParticles calls that are not necessary
When the particle is updated it is added to the moving kinematics particles list or to the active list.
All those lists are dirty during integration. So it is not required to call dirty on them.
In scenario where particles are moved from the Game Thread, this optimization could get up to 15% improvement on the Push Data

diff --git a/Engine/Source/Runtime/Experimental/Chaos/Private/Chaos/PBDRigidsEvolutionGBF.cpp b/Engine/Source/Runtime/Experimental/Chaos/Private/Chaos/PBDRigidsEvolutionGBF.cpp
--- a/Engine/Source/Runtime/Experimental/Chaos/Private/Chaos/PBDRigidsEvolutionGBF.cpp
+++ b/Engine/Source/Runtime/Experimental/Chaos/Private/Chaos/PBDRigidsEvolutionGBF.cpp
@@ -1748,7 +1748,7 @@ void FPBDRigidsEvolutionGBF::SetParticleTransformSwept(FGeometryParticleHandle*
 {
 	// If the rotation has changed, we need to adjust the initial sweep position to account for the fact that the
 	// particle will have moved in a straight line through its center of mass. This is important when we have shapes
-	// that are not centerd on their center of mass (either it's multi-shape body, or the user modified the CoM).
+	// that are not centered on their center of mass (either it's multi-shape body, or the user modified the CoM).
 	const FVec3 CoMOffset = FConstGenericParticleHandle(InParticle)->CenterOfMass();
 	const FVec3 EndCoM = InPos + InRot * CoMOffset;
 	const FVec3 StartCoM = InParticle->GetX() + InParticle->GetR() * CoMOffset;
@@ -1773,8 +1773,9 @@ void FPBDRigidsEvolutionGBF::SetParticleTransformSwept(FGeometryParticleHandle*
 	}
 }
 
-void FPBDRigidsEvolutionGBF::SetParticleKinematicTarget(FGeometryParticleHandle* ParticleHandle, const FKinematicTarget& NewKinematicTarget)
+bool FPBDRigidsEvolutionGBF::SetParticleKinematicTarget(FGeometryParticleHandle* ParticleHandle, const FKinematicTarget& NewKinematicTarget)
 {
+	bool bIsAddedToParticles = false;
 	FGenericParticleHandle Particle = ParticleHandle;
 
 	// NOTE: If called on a dynamic body we just move the body
@@ -1788,6 +1789,7 @@ void FPBDRigidsEvolutionGBF::SetParticleKinematicTarget(FGeometryParticleHandle*
 		{
 			// move particle from "non-moving" kinematics to "moving" kinematics
 			Particles.MarkMovingKinematic(Kinematic);
+			bIsAddedToParticles = true;
 		}
 		Kinematic->SetKinematicTarget(NewKinematicTarget);
 	}
@@ -1796,8 +1798,10 @@ void FPBDRigidsEvolutionGBF::SetParticleKinematicTarget(FGeometryParticleHandle*
 		if (NewKinematicTarget.GetMode() == EKinematicTargetMode::Position)
 		{
 			SetParticleTransform(ParticleHandle, NewKinematicTarget.GetPosition(), NewKinematicTarget.GetRotation(), false);
+			bIsAddedToParticles = true;
 		}
 	}
+	return bIsAddedToParticles;
 }
 
 void FPBDRigidsEvolutionGBF::OnParticleMoved(FGeometryParticleHandle* InParticle, const FVec3& PrevX, const FRotation3& PrevR, const bool bIsTeleport, const bool bWakeUp)
diff --git a/Engine/Source/Runtime/Experimental/Chaos/Private/PhysicsProxy/SingleParticlePhysicsProxy.cpp b/Engine/Source/Runtime/Experimental/Chaos/Private/PhysicsProxy/SingleParticlePhysicsProxy.cpp
--- a/Engine/Source/Runtime/Experimental/Chaos/Private/PhysicsProxy/SingleParticlePhysicsProxy.cpp
+++ b/Engine/Source/Runtime/Experimental/Chaos/Private/PhysicsProxy/SingleParticlePhysicsProxy.cpp
@@ -125,9 +125,10 @@ void PushToPhysicsStateImp(const Chaos::FDirtyPropertiesManager& Manager, Chaos:
 			}
 			
 			auto NewKinematicTargetGT = bHasKinematicData ? ParticleData.FindKinematicTarget(Manager, DataIdx) : nullptr;
+			bool bSkipDirtyParticle = false;
 			if (NewKinematicTargetGT)
 			{
-				Evolution.SetParticleKinematicTarget(KinematicHandle, *NewKinematicTargetGT);
+				bSkipDirtyParticle = Evolution.SetParticleKinematicTarget(KinematicHandle, *NewKinematicTargetGT);
 			}
 
 			if(NewXR || NewNonFrequentData || NewVelocities || NewKinematicTargetGT)
@@ -138,7 +139,13 @@ void PushToPhysicsStateImp(const Chaos::FDirtyPropertiesManager& Manager, Chaos:
 				const FRigidTransform3 WorldTransform = !bHasKinematicTarget ? FRigidTransform3(Handle->GetX(), Handle->GetR()) : NewKinematicTargetGT->GetTransform();
 				Handle->UpdateWorldSpaceState(WorldTransform, FVec3(0));
 
-				Evolution.DirtyParticle(*Handle);
+				// Kinematic target added the particle in the active list or in the moving kinematic list and so it will be updated
+				// No need to added into the dirty list 
+				if (!bSkipDirtyParticle)
+				{
+					Evolution.DirtyParticle(*Handle);
+				}
+				
 			}
 
 			if(bHasDynamicData)
diff --git a/Engine/Source/Runtime/Experimental/Chaos/Public/Chaos/PBDRigidsEvolutionGBF.h b/Engine/Source/Runtime/Experimental/Chaos/Public/Chaos/PBDRigidsEvolutionGBF.h
--- a/Engine/Source/Runtime/Experimental/Chaos/Public/Chaos/PBDRigidsEvolutionGBF.h
+++ b/Engine/Source/Runtime/Experimental/Chaos/Public/Chaos/PBDRigidsEvolutionGBF.h
@@ -179,8 +179,9 @@ namespace Chaos
 		/**
 		* Set the kinematic target for a particle. This will exist for only one tick - a new target must be set for the next tick if required.
 		* If called on a dynamic object, is equivalent to SetParticleTransform with bIsTeleport=false
+		* Returns true if the particle has been added to the moving kinematic particles list or the the active list.
 		*/
-		CHAOS_API void SetParticleKinematicTarget(FGeometryParticleHandle* ParticleHandle, const FKinematicTarget& NewKinematicTarget);
+		CHAOS_API bool SetParticleKinematicTarget(FGeometryParticleHandle* ParticleHandle, const FKinematicTarget& NewKinematicTarget);
 
 		/*
 		 * [EXPERIMENTAL] Apply a momentumless correction to the particle transform, usually as a result of a server correction.
