Dataflow : add support for bone weights and indices for the render collection
- this is to support edit skin weight tools and nodes

diff --git a/Engine/Plugins/Experimental/Dataflow/Source/DataflowAssetTools/Private/Dataflow/CollectionRenderingPatternUtility.cpp b/Engine/Plugins/Experimental/Dataflow/Source/DataflowAssetTools/Private/Dataflow/CollectionRenderingPatternUtility.cpp
--- a/Engine/Plugins/Experimental/Dataflow/Source/DataflowAssetTools/Private/Dataflow/CollectionRenderingPatternUtility.cpp
+++ b/Engine/Plugins/Experimental/Dataflow/Source/DataflowAssetTools/Private/Dataflow/CollectionRenderingPatternUtility.cpp
@@ -47,6 +47,8 @@ namespace UE::Dataflow
 				const TManagedArray<TArray<FVector2f>>& UVs = Facade.GetVertexUV();
 				const TManagedArray<int32>& FacadeMaterialID = Facade.GetMaterialID();
 
+				const int32 MeshStartVtxID = DynamicMesh.VertexCount();
+
 				int32 LastVertexIndex = StartVertices + NumVertices;
 				for (int32 VertexIndex = StartVertices; VertexIndex < LastVertexIndex; ++VertexIndex)
 				{
@@ -145,6 +147,32 @@ namespace UE::Dataflow
 					}
 				}
 
+				const TManagedArray<TArray<float>>* BoneWeights = Facade.FindBoneWeights();
+				const TManagedArray<TArray<int32>>* BoneIndices = Facade.FindBoneIndices();
+				if (BoneWeights && BoneIndices)
+				{
+					const FName DefaultProfileName("Default");
+					FDynamicMeshVertexSkinWeightsAttribute* Attribute = DynamicMesh.Attributes()->GetSkinWeightsAttribute(DefaultProfileName);
+					if (Attribute == nullptr)
+					{
+						Attribute = new FDynamicMeshVertexSkinWeightsAttribute(&DynamicMesh);
+						DynamicMesh.Attributes()->AttachSkinWeightsAttribute(DefaultProfileName, Attribute);
+					}
+					for (int32 VertexIndex = StartVertices; VertexIndex < LastVertexIndex; ++VertexIndex)
+					{
+						const TArray<float>& VtxBoneWeights = (*BoneWeights)[VertexIndex];
+						const TArray<int32>& VtxBoneIndices = (*BoneIndices)[VertexIndex];
+
+						UE::AnimationCore::FBoneWeights Weights;
+						const int32 NumWeights = FMath::Min(VtxBoneWeights.Num(), VtxBoneIndices.Num());
+						for (int32 WeightIndex = 0; WeightIndex < NumWeights; ++WeightIndex)
+						{
+							Weights.SetBoneWeight(VtxBoneIndices[WeightIndex], VtxBoneWeights[WeightIndex]);
+						}
+						const int32 VtxID = MeshStartVtxID + (VertexIndex - StartVertices);
+						Attribute->SetValue(VtxID, Weights);
+					}
+				}
 			}
 		}
 
diff --git a/Engine/Plugins/Experimental/Dataflow/Source/DataflowEditor/Private/Dataflow/DataflowEngineRendering.cpp b/Engine/Plugins/Experimental/Dataflow/Source/DataflowEditor/Private/Dataflow/DataflowEngineRendering.cpp
--- a/Engine/Plugins/Experimental/Dataflow/Source/DataflowEditor/Private/Dataflow/DataflowEngineRendering.cpp
+++ b/Engine/Plugins/Experimental/Dataflow/Source/DataflowEditor/Private/Dataflow/DataflowEngineRendering.cpp
@@ -17,6 +17,7 @@
 #include "Materials/MaterialInterface.h"
 #include "GeometryCollection/Facades/CollectionUVFacade.h"
 #include "GeometryCollection/Facades/CollectionMeshFacade.h"
+#include "GeometryCollection/Facades/CollectionVertexBoneWeightsFacade.h"
 #include "UDynamicMesh.h"
 
 namespace UE::Dataflow
@@ -450,6 +451,26 @@ namespace UE::Dataflow
 				}
 			}
 
+			TArray<TArray<float>> RenderBoneWeights;
+			TArray<TArray<int32>> RenderBoneIndices;
+
+			using FVertexBoneWeightsFacade = GeometryCollection::Facades::FVertexBoneWeightsFacade;
+
+			const TManagedArray<TArray<float>>* BoneWeights = Collection.FindAttributeTyped<TArray<float>>(FVertexBoneWeightsFacade::BoneWeightsAttributeName, FGeometryCollection::VerticesGroup);
+			const TManagedArray<TArray<int32>>* BoneIndices = Collection.FindAttributeTyped<TArray<int32>>(FVertexBoneWeightsFacade::BoneIndicesAttributeName, FGeometryCollection::VerticesGroup);
+			if (BoneWeights && BoneIndices)
+			{
+				RenderBoneWeights.SetNum(Vertices.Num());
+				RenderBoneIndices.SetNum(Vertices.Num());
+
+				int32 LastVertIndex = VertexStart[Gdx] + VertexCount[Gdx];
+				for (int32 VertexIdx = VertexStart[Gdx], SrcVertexIdx = 0; VertexIdx < LastVertIndex; ++VertexIdx, ++SrcVertexIdx)
+				{
+					RenderBoneWeights[SrcVertexIdx] = (*BoneWeights)[VertexIdx];
+					RenderBoneIndices[SrcVertexIdx] = (*BoneIndices)[VertexIdx];
+				}
+			}
+
 			// Set the data on the RenderCollection
 			if (Vertices.Num() && Tris.Num())
 			{
@@ -460,6 +481,10 @@ namespace UE::Dataflow
 				}
 				int32 GeometryIndex = RenderCollection.StartGeometryGroup(GeometryName);
 				RenderCollection.AddSurface(MoveTemp(Vertices), MoveTemp(Tris), MoveTemp(VertexNormals), MoveTemp(VertexColors));
+				if (BoneWeights && BoneIndices)
+				{
+					RenderCollection.AddSurfaceBoneWeightsAndIndices(MoveTemp(RenderBoneWeights), MoveTemp(RenderBoneIndices));
+				}
 				RenderCollection.EndGeometryGroup(GeometryIndex);
 			}
 		}
diff --git a/Engine/Source/Runtime/Experimental/Chaos/Private/GeometryCollection/Facades/CollectionRenderingFacade.cpp b/Engine/Source/Runtime/Experimental/Chaos/Private/GeometryCollection/Facades/CollectionRenderingFacade.cpp
--- a/Engine/Source/Runtime/Experimental/Chaos/Private/GeometryCollection/Facades/CollectionRenderingFacade.cpp
+++ b/Engine/Source/Runtime/Experimental/Chaos/Private/GeometryCollection/Facades/CollectionRenderingFacade.cpp
@@ -5,6 +5,7 @@
 =============================================================================*/
 
 #include "GeometryCollection/Facades/CollectionRenderingFacade.h"
+#include "GeometryCollection/Facades/CollectionVertexBoneWeightsFacade.h"
 #include "GeometryCollection/GeometryCollection.h"
 #include "Generators/CapsuleGenerator.h"
 
@@ -35,7 +36,8 @@ namespace GeometryCollection::Facades
 		, MaterialStartAttribute(InCollection, "MaterialsStart", FGeometryCollection::GeometryGroup, FGeometryCollection::MaterialGroup)
 		, MaterialCountAttribute(InCollection, "MaterialsCount", FGeometryCollection::GeometryGroup)
 		, GeometrySelectionAttribute(InCollection, "SelectionState", FGeometryCollection::GeometryGroup)
-
+		, BoneWeightsAttribute(InCollection, FVertexBoneWeightsFacade::BoneWeightsAttributeName, FGeometryCollection::VerticesGroup)
+		, BoneIndicesAttribute(InCollection, FVertexBoneWeightsFacade::BoneIndicesAttributeName, FGeometryCollection::VerticesGroup)
 	{}
 
 	FRenderingFacade::FRenderingFacade(const FManagedArrayCollection& InCollection)
@@ -62,6 +64,8 @@ namespace GeometryCollection::Facades
 		, MaterialStartAttribute(InCollection, "MaterialsStart", FGeometryCollection::GeometryGroup, FGeometryCollection::MaterialGroup)
 		, MaterialCountAttribute(InCollection, "MaterialsCount", FGeometryCollection::GeometryGroup)
 		, GeometrySelectionAttribute(InCollection, "SelectionState", FGeometryCollection::GeometryGroup)
+		, BoneWeightsAttribute(InCollection, FVertexBoneWeightsFacade::BoneWeightsAttributeName, FGeometryCollection::VerticesGroup)
+		, BoneIndicesAttribute(InCollection, FVertexBoneWeightsFacade::BoneIndicesAttributeName, FGeometryCollection::VerticesGroup)
 	{}
 
 	//
@@ -452,6 +456,35 @@ namespace GeometryCollection::Facades
 		}
 	}
 
+	void FRenderingFacade::AddSurfaceBoneWeightsAndIndices(TArray<TArray<float>>&& InBoneWeights, TArray<TArray<int32>>&& InBoneIndices)
+	{
+		check(!IsConst());
+		if (IsValid())
+		{
+			const int32 GeomIndex = (GeometryNameAttribute.Num() - 1);
+			check(GeomIndex >= 0);
+
+			const int32 VertexStart = VertexStartAttribute[GeomIndex];
+
+			if (InBoneWeights.Num() == InBoneIndices.Num())
+			{
+				const int32 NumVertices = InBoneIndices.Num();
+
+				BoneWeightsAttribute.Add();
+				BoneIndicesAttribute.Add();
+				TManagedArray<TArray<float>>& BoneWeights = BoneWeightsAttribute.Modify();
+				TManagedArray<TArray<int32>>& BoneIndices = BoneIndicesAttribute.Modify();
+				for (int32 VtxIndex = 0; VtxIndex < NumVertices; ++VtxIndex)
+				{
+					const int32 CollectionCVtxIndex = VertexStart + VtxIndex;
+					BoneWeights[CollectionCVtxIndex] = MoveTemp(InBoneWeights[VtxIndex]);
+					BoneIndices[CollectionCVtxIndex] = MoveTemp(InBoneIndices[VtxIndex]);
+				}
+			}
+		}
+		
+	}
+
 	void FRenderingFacade::AddSurface(TArray<FVector3f>&& InVertices, TArray<FIntVector>&& InIndices, TArray<FVector3f>&& InNormals, TArray<FLinearColor>&& InColors)
 	{
 		check(!IsConst());
diff --git a/Engine/Source/Runtime/Experimental/Chaos/Public/GeometryCollection/Facades/CollectionRenderingFacade.h b/Engine/Source/Runtime/Experimental/Chaos/Public/GeometryCollection/Facades/CollectionRenderingFacade.h
--- a/Engine/Source/Runtime/Experimental/Chaos/Public/GeometryCollection/Facades/CollectionRenderingFacade.h
+++ b/Engine/Source/Runtime/Experimental/Chaos/Public/GeometryCollection/Facades/CollectionRenderingFacade.h
@@ -84,6 +84,9 @@ namespace GeometryCollection::Facades
 		/** Add a surface with UV layers, Material IDs and material paths to the rendering view.*/
 		CHAOS_API void AddSurface(TArray<FVector3f>&& InVertices, TArray<FIntVector>&& InIndices, TArray<FVector3f>&& InNormals, TArray<FLinearColor>&& InColors, TArray<TArray<FVector2f>>&& InUVs, TArray<int32>&& InMaterialIDs, TArray<FString>&& MaterialPaths);
 
+		/** Add surface bone weights and indices to the rendering view - AddSurface must have been called before .*/
+		CHAOS_API void AddSurfaceBoneWeightsAndIndices(TArray<TArray<float>>&& InBoneWeights, TArray<TArray<int32>>&& InBoneIndices);
+
 		/**Add a tetrahedron to the rendering view.*/
 		CHAOS_API void AddTetrahedron(const TArray<FVector3f>& InVertices, const FIntVector4& InIndices);
 
@@ -138,6 +141,10 @@ namespace GeometryCollection::Facades
 		const TManagedArray<TArray<FVector2f>>& GetVertexUV() const { return VertexUVAttribute.Get(); }
 		TManagedArray<TArray<FVector2f>>& ModifyVertexUV() { check(!IsConst()); return VertexUVAttribute.Modify(); }
 
+		// optional per vertex bone attributes
+		const TManagedArray<TArray<float>>* FindBoneWeights() const { return BoneWeightsAttribute.Find(); }
+		const TManagedArray<TArray<int32>>* FindBoneIndices() const { return BoneIndicesAttribute.Find(); }
+
 		//
 		// Geometry Group Attributes
 		//
@@ -225,6 +232,9 @@ namespace GeometryCollection::Facades
 		TManagedArrayAccessor<int32> MaterialStartAttribute;
 		TManagedArrayAccessor<int32> MaterialCountAttribute;
 		TManagedArrayAccessor<int32> GeometrySelectionAttribute;
+
+		TManagedArrayAccessor<TArray<float>> BoneWeightsAttribute;
+		TManagedArrayAccessor<TArray<int32>> BoneIndicesAttribute;
 	};
 
 }
