Added get header value by key in HTTP I/O dispatcher header collection.

#rnx

diff --git a/Engine/Source/Runtime/Experimental/IoStore/OnDemandCore/Private/HttpIoDispatcher.cpp b/Engine/Source/Runtime/Experimental/IoStore/OnDemandCore/Private/HttpIoDispatcher.cpp
--- a/Engine/Source/Runtime/Experimental/IoStore/OnDemandCore/Private/HttpIoDispatcher.cpp
+++ b/Engine/Source/Runtime/Experimental/IoStore/OnDemandCore/Private/HttpIoDispatcher.cpp
@@ -28,6 +28,26 @@ FIoHttpHeaders& FIoHttpHeaders::Add(FAnsiStringView HeaderName, FAnsiStringView
 	return *this;
 }
 
+FAnsiStringView FIoHttpHeaders::Get(FAnsiStringView Key) const
+{
+	for (int32 Idx = 0, Count = Headers.Num(); Idx < Count; Idx += 2)
+	{
+		FAnsiStringView Header = Headers[Idx];
+		if (Header.Equals(Key, ESearchCase::IgnoreCase))
+		{
+			return Headers[Idx + 1];
+		}
+	}
+
+	return FAnsiStringView();
+}
+
+FAnsiStringView FIoHttpHeaders::Get(const FAnsiString& Key) const
+{
+	FAnsiStringView KeyView = Key;
+	return Get(KeyView);
+}
+
 TArray<FAnsiString> FIoHttpHeaders::ToArray() &&
 {
 	return MoveTemp(Headers);
diff --git a/Engine/Source/Runtime/Experimental/IoStore/OnDemandCore/Public/IO/HttpIoDispatcher.h b/Engine/Source/Runtime/Experimental/IoStore/OnDemandCore/Public/IO/HttpIoDispatcher.h
--- a/Engine/Source/Runtime/Experimental/IoStore/OnDemandCore/Public/IO/HttpIoDispatcher.h
+++ b/Engine/Source/Runtime/Experimental/IoStore/OnDemandCore/Public/IO/HttpIoDispatcher.h
@@ -167,6 +167,8 @@ public:
 									FIoHttpHeaders(FIoHttpHeaders&&) = default;
 	UE_API FIoHttpHeaders&			Add(FAnsiString&& HeaderName, FAnsiString&& HeaderValue);
 	UE_API FIoHttpHeaders&			Add(FAnsiStringView HeaderName, FAnsiStringView HeaderValue);
+	UE_API FAnsiStringView			Get(FAnsiStringView Key) const;
+	UE_API FAnsiStringView			Get(const FAnsiString& Key) const;
 	UE_API TArray<FAnsiString>		ToArray() &&;
 	TConstArrayView<FAnsiString>	ToArrayView() const { return Headers; }
 
