Add operator== and GetTypeHash for FIoRelativeUrl

#rnx

diff --git a/Engine/Source/Runtime/Experimental/IoStore/OnDemandCore/Public/IO/HttpIoDispatcher.h b/Engine/Source/Runtime/Experimental/IoStore/OnDemandCore/Public/IO/HttpIoDispatcher.h
--- a/Engine/Source/Runtime/Experimental/IoStore/OnDemandCore/Public/IO/HttpIoDispatcher.h
+++ b/Engine/Source/Runtime/Experimental/IoStore/OnDemandCore/Public/IO/HttpIoDispatcher.h
@@ -191,19 +191,27 @@ private:
 class FIoRelativeUrl
 {
 public:
+									using ElementType = ANSICHAR;
+
 									FIoRelativeUrl() = default;
 	bool							IsEmpty() const		{ return Url.IsEmpty(); }
 	int32							Len() const			{ return Url.Len(); }
 	const ANSICHAR*					ToString() const	{ return *Url; }
-	UE_API static FIoRelativeUrl	From(FAnsiStringView Url); 
+	FAnsiStringView					GetView() const		{ return FAnsiStringView(*Url, Url.Len()); }
 
 	const ANSICHAR*					operator*() const	{ return *Url; }
+	friend inline bool				operator==(const FIoRelativeUrl& Lhs, const FIoRelativeUrl& Rhs)	{ return Lhs.Url == Rhs.Url; }
+	friend inline bool				operator!=(const FIoRelativeUrl& Lhs, const FIoRelativeUrl& Rhs)	{ return Lhs.Url != Rhs.Url; }
+	friend inline bool				operator<(const FIoRelativeUrl& Lhs, const FIoRelativeUrl& Rhs)		{ return Lhs.Url < Rhs.Url; }
+
+	friend inline uint32			GetTypeHash(const FIoRelativeUrl& RelativeUrl) { return GetTypeHash(RelativeUrl.Url); }
+	UE_API static FIoRelativeUrl	From(FAnsiStringView Url); 
 
 private:
 	explicit FIoRelativeUrl(FAnsiStringView InUrl)
 		: Url(InUrl) { }
 
-	TSharedString<ANSICHAR> Url;
+	TSharedString<ElementType> Url;
 };
 
 /**
