Use correct string size in bytes when computing HTTP cache key.

#rnx

diff --git a/Engine/Source/Runtime/Experimental/IoStore/OnDemand/Private/OnDemandHttpIoDispacher.cpp b/Engine/Source/Runtime/Experimental/IoStore/OnDemand/Private/OnDemandHttpIoDispacher.cpp
--- a/Engine/Source/Runtime/Experimental/IoStore/OnDemand/Private/OnDemandHttpIoDispacher.cpp
+++ b/Engine/Source/Runtime/Experimental/IoStore/OnDemand/Private/OnDemandHttpIoDispacher.cpp
@@ -85,10 +85,11 @@ static FIoHash GetCacheKey(const FName& HostGroup, const FIoRelativeUrl Url, con
 
 	TCHAR NameString[NAME_SIZE];
 	const uint32 NameLength = HostGroup.GetPlainNameString(NameString);
-	HashBuilder.Update(NameString, NameLength);
-	HashBuilder.Update(Url.ToString(), Url.Len());
-	
+	HashBuilder.Update(NameString, NameLength * sizeof(TCHAR));
+	HashBuilder.Update(Url.ToString(), Url.Len() * sizeof(FIoRelativeUrl::ElementType));
+
 	TConstArrayView<FAnsiString> HeadersView = Headers.ToArrayView();
+	check(HeadersView.IsEmpty() || ((HeadersView.Num() % 2) == 0));
 	for (int32 Idx = 0; Idx < HeadersView.Num(); Idx += 2)
 	{
 		const FAnsiString Name	= HeadersView[Idx];
