Always negotiate to a support version

#rnx

#FailedValidations Preflight

diff --git a/Engine/Source/Runtime/Experimental/IoStore/HttpClient/Private/Peer.inl b/Engine/Source/Runtime/Experimental/IoStore/HttpClient/Private/Peer.inl
--- a/Engine/Source/Runtime/Experimental/IoStore/HttpClient/Private/Peer.inl
+++ b/Engine/Source/Runtime/Experimental/IoStore/HttpClient/Private/Peer.inl
@@ -269,7 +269,7 @@ static void Ssl_SetupAlpn(SSL* Ssl, int32 HttpVersion)
 ////////////////////////////////////////////////////////////////////////////////
 static int32 Ssl_GetProtocolVersion(SSL* Ssl)
 {
-	int32 Proto = -1;
+	int32 Proto = 1;
 
 	const char* AlpnProto = nullptr;
 	uint32 AlpnProtoLen;
@@ -279,8 +279,6 @@ static int32 Ssl_GetProtocolVersion(SSL* Ssl)
 		return Proto;
 	}
 
-	Proto = 0;
-
 	FAnsiStringView Needle(AlpnProto, AlpnProtoLen);
 	FAnsiStringView Candidates[] = {
 		"http/1.1",
@@ -356,7 +354,7 @@ static SSL*		Ssl_Create(...)				{ return nullptr; }
 static void		Ssl_Destroy(...)			{}
 static void		Ssl_AssociatePeer(...)		{}
 static void		Ssl_SetupAlpn(...)			{}
-static int32	Ssl_GetProtocolVersion(...)	{ return -1; }
+static int32	Ssl_GetProtocolVersion(...)	{ return 1; }
 static FOutcome	Ssl_Handshake(...)			{ return FOutcome::Error("!impl"); }
 static FOutcome	Ssl_Write(...)				{ return FOutcome::Error("!impl"); }
 static FOutcome	Ssl_Read(...)				{ return FOutcome::Error("!impl"); }
@@ -596,10 +594,11 @@ FOutcome FHttpPeer::Handshake()
 		}
 
 		int32 Negotiated = Ssl_GetProtocolVersion(Ssl);
-		if (Negotiated >= 0 && Negotiated != Proto)
+		if (Negotiated != 1 && Negotiated != 2)
 		{
-			return FOutcome::Error("Failed to negotiate protocol version");
+			return FOutcome::Error("Unexpected negotiated protocol version", Negotiated);
 		}
+		Proto = uint8(Negotiated);
 
 		/* intentional fallthrough */
 	}
