Fixed collision bug in broad phase
Static tree should have benn updated if the dynamic tree is not used
fallout from CL 45431155

#FailedValidations Preflight

diff --git a/Engine/Source/Runtime/Experimental/Chaos/Private/Chaos/PBDRigidsEvolutionGBF.cpp b/Engine/Source/Runtime/Experimental/Chaos/Private/Chaos/PBDRigidsEvolutionGBF.cpp
--- a/Engine/Source/Runtime/Experimental/Chaos/Private/Chaos/PBDRigidsEvolutionGBF.cpp
+++ b/Engine/Source/Runtime/Experimental/Chaos/Private/Chaos/PBDRigidsEvolutionGBF.cpp
@@ -1028,7 +1028,8 @@ void FPBDRigidsEvolutionGBF::Integrate(FReal Dt)
 	{
 		// The async data structure is for the static particle. Here we update the dynamic particles.
 		// We need to update them only if dynamic and static particles are sharing the same spatial structure
-		if (AccelerationStructureSplitStaticAndDynamic != 1)
+		// If the dynamic tree is not used we need to update the static one.
+		if (AccelerationStructureSplitStaticAndDynamic != 1 || AccelerationStructureUseDynamicTree == 0)
 		{
 			for (FPBDRigidParticleHandle* Particle : ParticlesArray)
 			{
@@ -1102,7 +1103,8 @@ void FPBDRigidsEvolutionGBF::ApplyKinematicTargets(const FReal Dt, const FReal S
 	{
 		// The async data structure is for the static particle. Here we update the kinematic particles.
 		// We need to update them only if kinematic and static particles are sharing the same spatial structure
-		if (AccelerationStructureSplitStaticAndDynamic != 1)
+		// If the dynamic tree is not used we need to update the static one.
+		if (AccelerationStructureSplitStaticAndDynamic != 1 || AccelerationStructureUseDynamicTree == 0)
 		{
 			for (FTransientPBDRigidParticleHandle& Particle : Particles.GetActiveMovingKinematicParticlesView())
 			{
diff --git a/Engine/Source/Runtime/Experimental/Chaos/Public/Chaos/PBDRigidsEvolutionGBF.h b/Engine/Source/Runtime/Experimental/Chaos/Public/Chaos/PBDRigidsEvolutionGBF.h
--- a/Engine/Source/Runtime/Experimental/Chaos/Public/Chaos/PBDRigidsEvolutionGBF.h
+++ b/Engine/Source/Runtime/Experimental/Chaos/Public/Chaos/PBDRigidsEvolutionGBF.h
@@ -27,6 +27,8 @@ namespace Chaos
 	class IResimCacheBase;
 	class FEvolutionResimCache;
 
+	CHAOS_API extern int32 AccelerationStructureUseDynamicTree;
+
 	namespace CVars
 	{
 		CHAOS_API extern FRealSingle HackMaxAngularVelocity;
