Dropping the IAS cache will now be correctly reflected in the on screen debug display and analytics.

#rnx

- Make sure that the statistics system is informed that we no longer have any data cached when FDiskCache::Drop is called.

#FailedValidations Preflight

diff --git a/Engine/Source/Runtime/Experimental/IoStore/OnDemand/Private/JournaledCache.cpp b/Engine/Source/Runtime/Experimental/IoStore/OnDemand/Private/JournaledCache.cpp
--- a/Engine/Source/Runtime/Experimental/IoStore/OnDemand/Private/JournaledCache.cpp
+++ b/Engine/Source/Runtime/Experimental/IoStore/OnDemand/Private/JournaledCache.cpp
@@ -943,6 +943,8 @@ void FDiskCache::Drop()
 	MappedBytes = 0;
 	DataMap.Reset();
 
+	FOnDemandIoBackendStats::Get()->OnCachePersistedBytes(MappedBytes);
+
 	OpenDataFile();
 }
 
