Allow the new IAX backend to cancel http requests that are inflight.

#rnx

- Currently we can combine several FHttpRequest into a single http request and we should only cancel that http request if all of the FHttpRequest are also canceled. This can be a little bit tricky to keep track of as we do not currently have the concept of a single http request, and the FHttpRequest are linked together as a single direction linked list so we cannot rely on being able to check the other FHttpRequest when one is canceled.
- To solve this when a http request is made we store it's FHttpTicketId (used to cancel it) in a new class FCancelationToken and then each FHttpRequest in the request chain holds a sharedptr to this token. When a FHttpRequest is told to cancel, it will release it's reference on the FCancelationToken and if all FHttpRequest release their reference then the FCancelationToken will be destroyed and the http client tol dit should cancel the request associated with the FHttpTicketId.
-- Note that currently this means we will issue a cancel request when all the FHttpRequest finish processing, but it is safe to try to "cancel" a completed request in the http client, it will just be considered a NOP.
- Currently we maintain a single mutex that locks when ever a FHttpRequest is canceled or a new http request is issued so that we can keep the assignment/release of the FCancelationToken thread safe. Cancelation is a relatively rare event so having a single mutex to lock is often better performing than having a mutex per FHttpRequest.
- Fixed a bug in FHttpDispatcher::CompleteRequest where if one sub request for a http request was canceled that all subsequent FHttpRequest  in the list would be marked as canceled.
- Fixed a bug where a canceled http request was being considered an error and would count towards CDN instability.
-- Also we now correctly inform the statistics system about requests that were canceled inflight.
- Requests that are canceled will now set their status code to 0 in the FIoHttpRequestCompleted rather than 409 to indicate that the problem was internal not not related to the http response itself.
- Replaces test against raw http status codes with a call to ::IsHttpStatusOk
- Added [[nodiscard]] to FMultiEndpointHttpClient::Get to ensure that the caller takes the FHttpTicketId into account.

#FailedValidations Preflight

diff --git a/Engine/Source/Runtime/Experimental/IoStore/OnDemand/Private/OnDemandHttpClient.h b/Engine/Source/Runtime/Experimental/IoStore/OnDemand/Private/OnDemandHttpClient.h
--- a/Engine/Source/Runtime/Experimental/IoStore/OnDemand/Private/OnDemandHttpClient.h
+++ b/Engine/Source/Runtime/Experimental/IoStore/OnDemand/Private/OnDemandHttpClient.h
@@ -129,9 +129,9 @@ public:
 	/** Blocking method */
 	[[nodiscard]] static TIoStatusOr<FMultiEndpointHttpClientResponse> Get(FAnsiStringView Url, const FMultiEndpointHttpClientConfig& Config);
 
-	FHttpTicketId Get(const FOnDemandHostGroup& HostGroup, FAnsiStringView RelativeUrl, FOnHttpResponse&& OnResponse);
-	FHttpTicketId Get(const FOnDemandHostGroup& HostGroup, FAnsiStringView RelativeUrl, const FIoOffsetAndLength& ChunkRange, FOnHttpResponse&& OnResponse);
-	FHttpTicketId Get(const FOnDemandHostGroup& HostGroup, FAnsiStringView RelativeUrl, const FIoOffsetAndLength& ChunkRange, TArray<FAnsiString>&& Headers, EMultiEndpointRequestFlags Flags, FOnHttpResponse&& OnResponse);
+	[[nodiscard]] FHttpTicketId Get(const FOnDemandHostGroup& HostGroup, FAnsiStringView RelativeUrl, FOnHttpResponse&& OnResponse);
+	[[nodiscard]] FHttpTicketId Get(const FOnDemandHostGroup& HostGroup, FAnsiStringView RelativeUrl, const FIoOffsetAndLength& ChunkRange, FOnHttpResponse&& OnResponse);
+	[[nodiscard]] FHttpTicketId Get(const FOnDemandHostGroup& HostGroup, FAnsiStringView RelativeUrl, const FIoOffsetAndLength& ChunkRange, TArray<FAnsiString>&& Headers, EMultiEndpointRequestFlags Flags, FOnHttpResponse&& OnResponse);
 
 	/**
 	 * Process the underlying http client.
diff --git a/Engine/Source/Runtime/Experimental/IoStore/OnDemand/Private/OnDemandHttpIoDispacher.cpp b/Engine/Source/Runtime/Experimental/IoStore/OnDemand/Private/OnDemandHttpIoDispacher.cpp
--- a/Engine/Source/Runtime/Experimental/IoStore/OnDemand/Private/OnDemandHttpIoDispacher.cpp
+++ b/Engine/Source/Runtime/Experimental/IoStore/OnDemand/Private/OnDemandHttpIoDispacher.cpp
@@ -33,8 +33,8 @@
 #define UE_ALLOW_CACHE_POISONING !UE_BUILD_SHIPPING
 
 // TODO - Maybe do LexToString/LexFromString (dupe code in UnrealEngine.cpp ThreadPriorityToString)
-
 extern const TCHAR* LexToString(EThreadPriority Priority);
+
 namespace UE::IoStore
 {
 ////////////////////////////////////////////////////////////////////////////////
@@ -44,6 +44,7 @@ extern FOnRecreateHttpClient OnRecreateHttpClient;
 
 extern int32 GIasHttpTimeOutMs;
 extern int32 GOnDemandBackendThreadPriorityIndex; 
+extern bool GIaxHttpEnableInflightCancellation;
 extern float GIaxInvalidUrlChance;
 #if UE_ALLOW_CACHE_POISONING
 extern bool GIasPoisonCache;
@@ -68,6 +69,33 @@ static FAutoConsoleVariableRef CVar_HttpCacheEnabled(
 
 namespace HttpIoDispatcher
 {
+////////////////////////////////////////////////////////////////////////////////
+
+/**
+ * Tracks the lifespan of a http request for as long as we want to keep it active. When the token is
+ * destroyed we will try to cancel the request if it is still in progress.
+ */
+class FCancelationToken
+{
+public:
+	FCancelationToken(FMultiEndpointHttpClient& InClient, FMultiEndpointHttpClient::FHttpTicketId InTicketId)
+		: Client(InClient)
+		, TicketId(InTicketId)
+	{
+
+	}
+
+	~FCancelationToken()
+	{
+		// Note that if the request has already completed then this will do nothing.
+		Client.CancelRequest(TicketId);
+	}
+
+private:
+	FMultiEndpointHttpClient& Client;
+	FMultiEndpointHttpClient::FHttpTicketId TicketId;
+};
+
 ////////////////////////////////////////////////////////////////////////////////
 static FIoHash GetCacheKey(const FIoHash& ChunkHash, const FIoHttpRange& Range)
 {
@@ -154,6 +182,8 @@ struct FHttpRequest final
 	std::atomic<int32>			RefCount{0};
 	std::atomic<EIoErrorCode>	CompletionStatus{EIoErrorCode::Unknown};
 	uint32						SeqNo = ++NextSeqNo;
+
+	TSharedPtr<FCancelationToken> CancelationToken;
 };
 using FHttpRequestList = TIntrusiveList<FHttpRequest>;
 
@@ -678,7 +708,8 @@ private:
 	FHttpQueue								HttpQueue;
 	FHostGroupRegistered					HostGroupRegisterd;
 	FEventRef								WakeUp;
-	FMutex									Mutex;
+	/** Keeps assignment of the CancelationToken safe between ::CancelRequest and ::ProcessHttpRequests */
+	FMutex									CancelationMutex;
 #if UE_ALLOW_HTTP_PAUSE
 	FConsoleCommandPtr						PauseCommand;
 	FConsoleCommandPtr						UnpauseCommand;
@@ -985,6 +1016,12 @@ void FHttpDispatcher::CancelRequest(FIoHttpRequestHandle Handle)
 			CompleteRequest(*Request, Request->Buffer, FIoHttpRange(), 0, false);
 		}, UE::Tasks::ETaskPriority::BackgroundLow);
 	}
+
+	if (GIaxHttpEnableInflightCancellation)
+	{
+		TUniqueLock _(CancelationMutex);
+		Request->CancelationToken.Reset();
+	}
 }
 
 void FHttpDispatcher::UpdateRequestPriority(FIoHttpRequestHandle Handle, int32 NewPriority)
@@ -1264,7 +1301,7 @@ bool FHttpDispatcher::ProcessHttpRequests()
 				: EMultiEndpointRequestFlags::None;
 
 			// Note that the requests issued here must be completed before this method can return so the capturing of locals is safe.
-			HttpClient->Get(HostGroup.GetUnderlyingHostGroup(), Url, TotalRange.ToOffsetAndLength(), MoveTemp(Headers), Flags,
+			FMultiEndpointHttpClient::FHttpTicketId TicketId =  HttpClient->Get(HostGroup.GetUnderlyingHostGroup(), Url, TotalRange.ToOffsetAndLength(), MoveTemp(Headers), Flags,
 				[this, HttpRequest, TotalRange, &NumConcurrentRequests, HostGroup = MoveTemp(HostGroup)]
 				(FMultiEndpointHttpClientResponse&& HttpResponse) mutable
 				{
@@ -1279,6 +1316,22 @@ bool FHttpDispatcher::ProcessHttpRequests()
 						}, UE::Tasks::ETaskPriority::BackgroundLow);
 				});
 
+			{
+				TUniqueLock _(CancelationMutex);
+			
+				TSharedPtr<FCancelationToken> Token = MakeShared<FCancelationToken>(*HttpClient.Get(), TicketId);
+				FHttpRequest* RequestList = HttpRequest;
+				while (RequestList)
+				{
+					if (!RequestList->Flags.HasAny(EHttpRequestFlags::CancelRequested))
+					{
+						RequestList->CancelationToken = Token;
+					}
+
+					RequestList = RequestList->Next;
+				}
+			}
+
 			if (NumConcurrentRequests >= MaxConcurrentRequests)
 			{
 				TRACE_CPUPROFILER_EVENT_SCOPE(FHttpDispatcher::TickHttpSaturated);
@@ -1338,6 +1391,10 @@ void FHttpDispatcher::CompleteHttpRequest(FHttpRequest& Request, FMultiEndpointH
 		FOnDemandIoBackendStats::Get()->OnHttpCdnCacheReply(Request.GetHttpRequestType(), HttpResponse.CDNCacheStatus);
 		FOnDemandIoBackendStats::Get()->OnHttpGet(Request.GetHttpRequestType(), HttpResponse.Body.DataSize(), HttpResponse.Sample.GetTotalMs());
 	}
+	else if (HttpResponse.IsCanceled())
+	{
+		FOnDemandIoBackendStats::Get()->OnHttpCancel(Request.GetHttpRequestType());
+	}
 	else
 	{
 		FOnDemandIoBackendStats::Get()->OnHttpError(Request.GetHttpRequestType());
@@ -1351,7 +1408,7 @@ void FHttpDispatcher::CompleteHttpRequest(FHttpRequest& Request, FMultiEndpointH
 	CompleteRequest(Request, MoveTemp(HttpResponse.Headers), HttpResponse.Body, TotalRange, HttpResponse.StatusCode, false);
 }
 
-void FHttpDispatcher::CompleteRequest(FHttpRequest& Request, TArray<FAnsiString>&& Headers, FIoBuffer& Body, const FIoHttpRange& TotalRange, uint32 StatusCode, bool bCached)
+void FHttpDispatcher::CompleteRequest(FHttpRequest& Request, TArray<FAnsiString>&& Headers, FIoBuffer& Body, const FIoHttpRange& TotalRange, const uint32 StatusCode, bool bCached)
 {
 #if UE_ALLOW_CACHE_POISONING
 	if (GIasPoisonCache && bCached)
@@ -1370,19 +1427,21 @@ void FHttpDispatcher::CompleteRequest(FHttpRequest& Request, TArray<FAnsiString>
 		NextRequest					= ToComplete->Next;
 		ToComplete->Next			= nullptr;
 
+		uint32 RequestStatusCode = StatusCode; // The status code may be modified per request, so take a copy.
+
 		ensure(!ToComplete->Flags.HasAny(EHttpRequestFlags::Completing | EHttpRequestFlags::Completed));
 		ToComplete->Flags.Add(EHttpRequestFlags::Completing);
 		ToComplete->Flags.Remove(EHttpRequestFlags::CacheInflight | EHttpRequestFlags::HttpInflight);
 
 		const FIoHttpRange& Range		= ToComplete->Options.GetRange();
-		EIoErrorCode CompletionStatus	= StatusCode > 199 && StatusCode < 300 ? EIoErrorCode::Ok : EIoErrorCode::ReadError;
+		EIoErrorCode CompletionStatus	= IsHttpStatusOk(StatusCode) ? EIoErrorCode::Ok : EIoErrorCode::ReadError;
 		FIoBuffer PartialBody			= Body;
 
 		if (ToComplete->Flags.HasAny(EHttpRequestFlags::CancelRequested))
 		{
 			CompletionStatus	= EIoErrorCode::Cancelled;
 			PartialBody			= FIoBuffer();
-			StatusCode			= 409; // Conflict
+			RequestStatusCode = 0; // Mark as a none http response related error
 		}
 
 		// If the request is NOT cached and specified a range we retrieve the partial body
@@ -1404,7 +1463,7 @@ void FHttpDispatcher::CompleteRequest(FHttpRequest& Request, TArray<FAnsiString>
 
 		const EIoHttpResponseFlags Flags	= bCached ? EIoHttpResponseFlags::Cached : EIoHttpResponseFlags::None;
 		FIoHttpRequestCompleted OnCompleted = MoveTemp(ToComplete->OnCompleted);
-		OnCompleted(FIoHttpResponse(ToComplete->CacheKey, FIoHttpHeaders::Create(MoveTemp(Headers)), PartialBody, CompletionStatus, StatusCode, Flags));
+		OnCompleted(FIoHttpResponse(ToComplete->CacheKey, FIoHttpHeaders::Create(MoveTemp(Headers)), PartialBody, CompletionStatus, RequestStatusCode, Flags));
 
 		ToComplete->Flags.Add(EHttpRequestFlags::Completed);
 		ToComplete->CompletionStatus.store(CompletionStatus, std::memory_order_seq_cst);
