Fix bug in on-demand HTTP client not triggering callback for 204, 205 and 304 responses.

#rnx

diff --git a/Engine/Source/Runtime/Experimental/IoStore/OnDemand/Private/OnDemandHttpClient.cpp b/Engine/Source/Runtime/Experimental/IoStore/OnDemand/Private/OnDemandHttpClient.cpp
--- a/Engine/Source/Runtime/Experimental/IoStore/OnDemand/Private/OnDemandHttpClient.cpp
+++ b/Engine/Source/Runtime/Experimental/IoStore/OnDemand/Private/OnDemandHttpClient.cpp
@@ -434,9 +434,15 @@ FMultiEndpointHttpClient::FHttpTicketId FMultiEndpointHttpClient::IssueRequest(F
 				{
 					Request.ResponseMessage = FString(HttpResponse.GetStatusMessage().TrimEnd());
 				}
-				
-				HttpResponse.SetDestination(Request.bIsChunkedTransfer == false ? &Request.Body : &Request.Chunk);
-				
+
+				if (HttpResponse.GetContentLength() > 0)
+				{
+					HttpResponse.SetDestination(Request.bIsChunkedTransfer == false ? &Request.Body : &Request.Chunk);
+				}
+				else
+				{
+					CompleteRequest(MoveTemp(Request), TicketStatus);
+				}
 				break;
 			}
 			case FTicketStatus::EId::Content:
@@ -471,67 +477,7 @@ FMultiEndpointHttpClient::FHttpTicketId FMultiEndpointHttpClient::IssueRequest(F
 				}
 				else
 				{
-					FMultiEndpointHttpClientResponse Response
-					{
-						.Headers				= MoveTemp(Request.ResponseHeaders),
-						.DurationMilliseconds	= uint64(FPlatformTime::ToMilliseconds64(FPlatformTime::Cycles64() - Request.StartTime)),
-						.StatusCode				= Request.StatusCode,
-						.RetryCount				= Request.RetryCount,
-						.HostIndex				= Request.Host,
-						.CDNCacheStatus			= Request.CDNCacheStatus
-					};
-
-					if (IsHttpStatusOk(Request.StatusCode) || Config.bResponseBodyOnError)
-					{
-						if (Request.StatusCode != 206 && Request.Range.IsValid())
-						{
-							// A partial read was requested but the server sent us the full response body so we need to cut down
-							// the data that we return. Note that the partial buffer we return will retain a ref count on the full
-							// sized buffer until it is also destroyed, so we don't need to do any allocations/copies at this point.
-							FMemoryView PartialRange = Request.Body.GetView().Mid(Request.Range.GetOffset(), Request.Range.GetLength());
-							Response.Body = FIoBuffer(PartialRange, Request.Body);
-
-							// Try to limit how frequently we log warnings about this to avoid spam
-							const uint64 LogSpamLimit = 50; 
-							static uint64 LogCounter = 0;
-							UE_CLOG(LogCounter++ % LogSpamLimit == 0, LogIoStoreOnDemand, Warning, TEXT("A partial content request failed and the full data blob was downloaded, check if the CDN supports this"));
-						}
-						else
-						{
-							Response.Body = MoveTemp(Request.Body);
-						}
-					}
-					else if (Config.bResponseBodyOnError)
-					{
-						Response.Body = MoveTemp(Request.Body);
-					}
-
-					if (TicketStatus.GetId() == FTicketStatus::EId::Content)
-					{
-						Response.Sample = TicketStatus.GetPerf().GetSample();
-					}
-					
-					if (TicketStatus.GetId() == FTicketStatus::EId::Error)
-					{
-						Response.Reason = TicketStatus.GetError().Reason;
-					}
-					else if (TicketStatus.GetId() == FTicketStatus::EId::Cancelled)
-					{
-						Response.Reason = TEXT("Canceled");
-						Response.bCanceled = true;
-					}
-					else if (!Request.ResponseMessage.IsEmpty())
-					{
-						Response.Reason = MoveTemp(Request.ResponseMessage);
-					}
-
-					{
-						UE::TUniqueLock _(TicketLookupMutex);
-						TicketLookupMap.Remove(Request.TicketId);
-					}
-
-					FOnHttpResponse OnResponse = MoveTemp(Request.OnResponse);
-					OnResponse(MoveTemp(Response));
+					CompleteRequest(MoveTemp(Request), TicketStatus);
 				}
 				break;
 			}
@@ -549,6 +495,73 @@ FMultiEndpointHttpClient::FHttpTicketId FMultiEndpointHttpClient::IssueRequest(F
 	return RequestId;
 }
 
+void FMultiEndpointHttpClient::CompleteRequest(FRequest&& Request, const UE::IoStore::HTTP::FTicketStatus& TicketStatus)
+{
+	using namespace UE::IoStore::HTTP;
+
+	FMultiEndpointHttpClientResponse Response
+	{
+		.Headers				= MoveTemp(Request.ResponseHeaders),
+		.DurationMilliseconds	= uint64(FPlatformTime::ToMilliseconds64(FPlatformTime::Cycles64() - Request.StartTime)),
+		.StatusCode				= Request.StatusCode,
+		.RetryCount				= Request.RetryCount,
+		.HostIndex				= Request.Host,
+		.CDNCacheStatus			= Request.CDNCacheStatus
+	};
+
+	if (IsHttpStatusOk(Request.StatusCode) || Config.bResponseBodyOnError)
+	{
+		if (Request.StatusCode != 206 && Request.Range.IsValid())
+		{
+			// A partial read was requested but the server sent us the full response body so we need to cut down
+			// the data that we return. Note that the partial buffer we return will retain a ref count on the full
+			// sized buffer until it is also destroyed, so we don't need to do any allocations/copies at this point.
+			FMemoryView PartialRange = Request.Body.GetView().Mid(Request.Range.GetOffset(), Request.Range.GetLength());
+			Response.Body = FIoBuffer(PartialRange, Request.Body);
+
+			// Try to limit how frequently we log warnings about this to avoid spam
+			const uint64 LogSpamLimit = 50; 
+			static uint64 LogCounter = 0;
+			UE_CLOG(LogCounter++ % LogSpamLimit == 0, LogIoStoreOnDemand, Warning, TEXT("A partial content request failed and the full data blob was downloaded, check if the CDN supports this"));
+		}
+		else
+		{
+			Response.Body = MoveTemp(Request.Body);
+		}
+	}
+	else if (Config.bResponseBodyOnError)
+	{
+		Response.Body = MoveTemp(Request.Body);
+	}
+
+	if (TicketStatus.GetId() == FTicketStatus::EId::Content)
+	{
+		Response.Sample = TicketStatus.GetPerf().GetSample();
+	}
+	
+	if (TicketStatus.GetId() == FTicketStatus::EId::Error)
+	{
+		Response.Reason = TicketStatus.GetError().Reason;
+	}
+	else if (TicketStatus.GetId() == FTicketStatus::EId::Cancelled)
+	{
+		Response.Reason = TEXT("Canceled");
+		Response.bCanceled = true;
+	}
+	else if (!Request.ResponseMessage.IsEmpty())
+	{
+		Response.Reason = MoveTemp(Request.ResponseMessage);
+	}
+
+	{
+		UE::TUniqueLock _(TicketLookupMutex);
+		TicketLookupMap.Remove(Request.TicketId);
+	}
+
+	FOnHttpResponse OnResponse = MoveTemp(Request.OnResponse);
+	OnResponse(MoveTemp(Response));
+}
+
 void FMultiEndpointHttpClient::RetryRequest(FRequest&& Request)
 {
 	FConnection& Connection = Request.Connection;
diff --git a/Engine/Source/Runtime/Experimental/IoStore/OnDemand/Private/OnDemandHttpClient.h b/Engine/Source/Runtime/Experimental/IoStore/OnDemand/Private/OnDemandHttpClient.h
--- a/Engine/Source/Runtime/Experimental/IoStore/OnDemand/Private/OnDemandHttpClient.h
+++ b/Engine/Source/Runtime/Experimental/IoStore/OnDemand/Private/OnDemandHttpClient.h
@@ -197,6 +197,7 @@ private:
 										FMultiEndpointHttpClient(const FMultiEndpointHttpClientConfig& Config);
 
 	FHttpTicketId						IssueRequest(FRequest&& Request);
+	void								CompleteRequest(FRequest&& Request, const UE::IoStore::HTTP::FTicketStatus& TicketStatus);
 	void								RetryRequest(FRequest&& Request);
 
 	void								ProcessRetryAttempts(uint32 TicketCount);
