[NetPhysics] Improve various aspects of resimulation, mainly for sim-proxy pawns.
- Option to trigger a resim when receiving new inputs
- Option to apply input decay curve over a set time instead of over the remaining resim frames
- Option to decay post-resim error in render interpolation exponentially, with a half-life setting
- Moved some render interpolation logic from the header file

diff --git a/Engine/Source/Runtime/Engine/Private/PhysicsEngine/NetworkPhysicsComponent.cpp b/Engine/Source/Runtime/Engine/Private/PhysicsEngine/NetworkPhysicsComponent.cpp
--- a/Engine/Source/Runtime/Engine/Private/PhysicsEngine/NetworkPhysicsComponent.cpp
+++ b/Engine/Source/Runtime/Engine/Private/PhysicsEngine/NetworkPhysicsComponent.cpp
@@ -64,6 +64,13 @@ namespace PhysicsReplicationCVars
 		static FAutoConsoleVariableRef CVarResimApplySimProxyStateAtRuntime(TEXT("np2.Resim.ApplySimProxyStateAtRuntime"), bApplySimProxyStateAtRuntime, TEXT("When true, call ApplyData on received states for simulated proxies at runtime."));
 		bool bApplySimProxyInputAtRuntime = true;
 		static FAutoConsoleVariableRef CVarResimApplySimProxyInputAtRuntime(TEXT("np2.Resim.ApplySimProxyInputAtRuntime"), bApplySimProxyInputAtRuntime, TEXT("When true, call ApplyData on received inputs for simulated proxies at runtime."));
+		bool bTriggerResimOnInputReceive = false;
+		static FAutoConsoleVariableRef CVarTriggerResimOnInputReceive(TEXT("np2.Resim.TriggerResimOnInputReceive"), bTriggerResimOnInputReceive, TEXT("When true, a resim will be requested to the frame of the latest frame of received inputs this frame"));
+
+		bool bApplyInputDecayOverSetTime = false;
+		static FAutoConsoleVariableRef CVarApplyInputDecayOverSetTime(TEXT("np2.Resim.ApplyInputDecayOverSetTime"), bApplyInputDecayOverSetTime, TEXT("When true, apply the Input Decay Curve over a set amount of time instead of over the start of input prediction and end of resim which is variable each resimulation"));
+		float InputDecaySetTime = 0.15f;
+		static FAutoConsoleVariableRef CVarInputDecaySetTime(TEXT("np2.Resim.InputDecaySetTime"), InputDecaySetTime, TEXT("Applied when np2.Resim.ApplyInputDecayOverSetTime is true, read there for more info. Set time to apply Input Decay Curve over while predicting inputs during resimulation"));
 
 		bool bEnableStatefulDeltaSerialization = true;
 		static FAutoConsoleVariableRef CVarResimEnableStatefulDeltaSerialization(TEXT("np2.Resim.StatefulDeltaSerialization.Enable"), bEnableStatefulDeltaSerialization, TEXT("Enables stateful delta serialization for FNetworkPhysicsData derived inputs and states. During FNetworkPhysicsData::NetSerialize there will be a valid pointer to a previous FNetworkPhysicsData which can be used for delta serialization, FNetworkPhysicsData::DeltaSourceData. NOTE: Switching this during gameplay might cause disconnections."));
@@ -1914,6 +1921,20 @@ void FAsyncNetworkPhysicsComponent::ConsumeAsyncInput(const int32 PhysicsStep)
 						}
 					}
 				}
+				
+				// If setting is true, request resimulation for sim-proxies when receiving inputs that are newer than latest already received input
+				if (!bIsServer && ComponentSettings.GetTriggerResimOnInputReceive() && !IsLocallyControlled())
+				{
+					const int32 LatestReceiveFrame = AsyncInput->InputData.Get()->GetLatestFrame() - NetworkPhysicsTickOffset;
+					const int32 NextInputFrame = InputHistory.Get()->GetLatestFrame() + 1;
+					if (LatestReceiveFrame >= NextInputFrame)
+					{
+						const int32 EarliestReceivedFrame = AsyncInput->InputData.Get()->GetEarliestFrame() - NetworkPhysicsTickOffset;
+						const int32 EarliestNewInputFrame = FMath::Max(EarliestReceivedFrame, NextInputFrame);
+						TriggerResimulation(EarliestNewInputFrame);
+					}
+				}
+
 				ReceiveHelper(InputHistory.Get(), AsyncInput->InputData.Get(), /*bImportant*/false, bCompareInput);
 			}
 
@@ -2483,10 +2504,21 @@ const float FAsyncNetworkPhysicsComponent::GetCurrentInputDecay(const FNetworkPh
 	const FRuntimeFloatCurve& InputDecayCurve = ComponentSettings.GetInputDecayCurve();
 
 	const float NumPredictedInputs = RewindData->CurrentFrame() - PhysicsData->LocalFrame; // Number of frames we have used the same PhysicsData for during resim
-	const float MaxPredictedInputs = RewindData->GetLatestFrame() - 1 - PhysicsData->LocalFrame; // Max number of frames PhysicsData registered frame until end of resim
+	float MaxPredictedInputs = 0;
 
+	if (ComponentSettings.GetApplyInputDecayOverSetTime())
+	{
+		// If the decay curve should be applied over a set amount of time, calculate how many frames that time translates to
+		MaxPredictedInputs = ComponentSettings.GetInputDecaySetTime() / BaseSolver->GetAsyncDeltaTime();
+	}
+	else
+	{
+		// Number of frames the input will predict and decay over until resimulation has completed
+		MaxPredictedInputs = RewindData->GetLatestFrame() - 1 - PhysicsData->LocalFrame;
+	}
+	
 	// Linear decay
-	const float PredictionAlpha = MaxPredictedInputs > 0 ? (NumPredictedInputs / MaxPredictedInputs) : 0.0f;
+	const float PredictionAlpha = MaxPredictedInputs > 0 ? FMath::Clamp(NumPredictedInputs / MaxPredictedInputs, 0.0f, 1.0f) : 0.0f;
 
 	// Get decay from curve
 	const float InputDecay = InputDecayCurve.GetRichCurveConst()->Eval(PredictionAlpha);
diff --git a/Engine/Source/Runtime/Engine/Private/PhysicsEngine/NetworkPhysicsSettingsComponent.cpp b/Engine/Source/Runtime/Engine/Private/PhysicsEngine/NetworkPhysicsSettingsComponent.cpp
--- a/Engine/Source/Runtime/Engine/Private/PhysicsEngine/NetworkPhysicsSettingsComponent.cpp
+++ b/Engine/Source/Runtime/Engine/Private/PhysicsEngine/NetworkPhysicsSettingsComponent.cpp
@@ -253,6 +253,11 @@ void FNetworkPhysicsSettingsResimulationErrorCorrection::ApplySettings_External(
 					Settings.MaximumErrorCorrectionBeforeSnapping = ResimErrorMaximumDistanceBeforeSnapping;
 					Settings.MaximumErrorCorrectionDesyncTimeBeforeSnapping = ResimErrorMaximumDesyncTimeBeforeSnapping;
 					Settings.ErrorDirectionalDecayMultiplier = ResimErrorDirectionalDecayMultiplier;
+					Settings.bApplyExponentialDecay = bRenderInterpApplyExponentialDecay;
+					Settings.ExponentialDecayLinearHalfLife = RenderInterpExponentialDecayLinearHalfLife;
+					Settings.ExponentialDecayAngularHalfLife = RenderInterpExponentialDecayAngularHalfLife;
+					Settings.MinimumLinearThreshold = RenderInterpMinimumLinearThreshold;
+					Settings.MinimumAngularThreshold = RenderInterpMinimumAngularThreshold;
 				}
 			};
 
diff --git a/Engine/Source/Runtime/Engine/Private/PhysicsEngine/PhysicsReplication.cpp b/Engine/Source/Runtime/Engine/Private/PhysicsEngine/PhysicsReplication.cpp
--- a/Engine/Source/Runtime/Engine/Private/PhysicsEngine/PhysicsReplication.cpp
+++ b/Engine/Source/Runtime/Engine/Private/PhysicsEngine/PhysicsReplication.cpp
@@ -82,6 +82,8 @@ namespace RenderInterpolationCVars
 {
 	bool bRenderInterpDebugDrawResimTrigger = false;
 	static FAutoConsoleVariableRef CVarRenderInterpDebugDrawResimTrigger(TEXT("p.RenderInterp.DebugDraw.ResimTrigger"), bRenderInterpDebugDrawResimTrigger, TEXT("Draw debug lines for physics render interpolation, also needs p.Chaos.DebugDraw.Enabled set"));
+	float RenderInterpDebugDrawResimBoxScale= 1.0f;
+	static FAutoConsoleVariableRef CVarRenderInterpDebugDrawResimBoxScale(TEXT("p.RenderInterp.DebugDraw.ResimBoxScale"), RenderInterpDebugDrawResimBoxScale, TEXT("Draw debug lines for physics render interpolation, also needs p.Chaos.DebugDraw.Enabled set"));
 }
 
 namespace PhysicsReplicationCVars
@@ -2299,7 +2301,8 @@ bool FPhysicsReplicationAsync::ResimulationReplication(Chaos::FPBDRigidParticleH
 		if (bShouldTriggerResim)
 		{
 			FVector Box = RenderInterpolationCVars::bRenderInterpDebugDrawResimTrigger ? FVector(6, 3, 2) : FVector(40, 20, 10);
-			const float DrawThickness = RenderInterpolationCVars::bRenderInterpDebugDrawResimTrigger ? 0.5f : 1.5f;
+			Box *= RenderInterpolationCVars::RenderInterpDebugDrawResimBoxScale;
+			const float DrawThickness = (RenderInterpolationCVars::bRenderInterpDebugDrawResimTrigger ? 0.5f : 1.5f) * RenderInterpolationCVars::RenderInterpDebugDrawResimBoxScale;
 
 			if (RenderInterpolationCVars::bRenderInterpDebugDrawResimTrigger) // Resim debug draw extension for render interpolation 
 			{
diff --git a/Engine/Source/Runtime/Engine/Public/Physics/NetworkPhysicsSettingsComponent.h b/Engine/Source/Runtime/Engine/Public/Physics/NetworkPhysicsSettingsComponent.h
--- a/Engine/Source/Runtime/Engine/Public/Physics/NetworkPhysicsSettingsComponent.h
+++ b/Engine/Source/Runtime/Engine/Public/Physics/NetworkPhysicsSettingsComponent.h
@@ -74,6 +74,9 @@ namespace PhysicsReplicationCVars
 		extern bool bCompareInputToTriggerRewind;
 		extern bool bApplySimProxyStateAtRuntime;
 		extern bool bApplySimProxyInputAtRuntime;
+		extern bool bTriggerResimOnInputReceive;
+		extern bool bApplyInputDecayOverSetTime;
+		extern float InputDecaySetTime;
 	}
 }
 
@@ -294,6 +297,11 @@ struct FNetworkPhysicsSettingsResimulationErrorCorrection
 		, ResimErrorMaximumDistanceBeforeSnapping(250.0f)
 		, ResimErrorMaximumDesyncTimeBeforeSnapping(0.6f)
 		, ResimErrorDirectionalDecayMultiplier(0.0f)
+		, bRenderInterpApplyExponentialDecay(false)
+		, RenderInterpExponentialDecayLinearHalfLife(0.06f)
+		, RenderInterpExponentialDecayAngularHalfLife(0.06f)
+		, RenderInterpMinimumLinearThreshold(0.1f)
+		, RenderInterpMinimumAngularThreshold(0.001f)
 	{};
 
 	GENERATED_BODY()
@@ -322,6 +330,26 @@ struct FNetworkPhysicsSettingsResimulationErrorCorrection
 	/** Overrides CVar: p.RenderInterp.DirectionalDecayMultiplier -- Decay error offset in the direction that the physics object is moving, value is multiplier of projected offset direction, 0.25 means a 25 % decay of the magnitude in the direction of physics travel.Deactivate by setting to 0 */
 	UPROPERTY(config, EditDefaultsOnly, Category = "Overrides", Meta = (EditCondition = "bOverrideResimErrorInterpolationSettings", ClampMin = "0.0", UIMin = "0.0", ClampMax = "1.0", UIMax = "1.0"))
 	float ResimErrorDirectionalDecayMultiplier;
+	
+	/** Overrides CVar: p.RenderInterp.ApplyExponentialDecay -- When enabled a post-resim error will decay exponentially (instead of linearly) based on half-life time set in ExponentialDecayLinearHalfLife and ExponentialDecayAngularHalfLife. */
+	UPROPERTY(config, EditDefaultsOnly, Category = "Overrides", Meta = (EditCondition = "bOverrideResimErrorInterpolationSettings"))
+	bool bRenderInterpApplyExponentialDecay;
+	
+	/** Overrides CVar: p.RenderInterp.ExponentialDecayLinearHalfLife -- Sets the positional half-life time for when bApplyExponentialDecay is enabled. */
+	UPROPERTY(config, EditDefaultsOnly, Category = "Overrides", Meta = (EditCondition = "bOverrideResimErrorInterpolationSettings", ClampMin = "0.0", UIMin = "0.0"))
+	float RenderInterpExponentialDecayLinearHalfLife;
+
+	/** Overrides CVar: p.RenderInterp.ExponentialDecayAngularHalfLife -- Sets the rotational half-life time for when bApplyExponentialDecay is enabled. */
+	UPROPERTY(config, EditDefaultsOnly, Category = "Overrides", Meta = (EditCondition = "bOverrideResimErrorInterpolationSettings", ClampMin = "0.0", UIMin = "0.0"))
+	float RenderInterpExponentialDecayAngularHalfLife;
+
+	/** Overrides CVar: p.RenderInterp.MinimumLinearThreshold -- Squared value, when the remaining render error is below this we clear it, if ApplyExponentialDecay is enabled. */
+	UPROPERTY(config, EditDefaultsOnly, Category = "Overrides", Meta = (EditCondition = "bOverrideResimErrorInterpolationSettings", ClampMin = "0.0", UIMin = "0.0"))
+	float RenderInterpMinimumLinearThreshold;
+
+	/** Overrides CVar: p.RenderInterp.MinimumAngularThreshold -- When the remaining render error angle is below this we clear it, if ApplyExponentialDecay is enabled. */
+	UPROPERTY(config, EditDefaultsOnly, Category = "Overrides", Meta = (EditCondition = "bOverrideResimErrorInterpolationSettings", ClampMin = "0.0", UIMin = "0.0"))
+	float RenderInterpMinimumAngularThreshold;
 };
 
 USTRUCT()
@@ -457,6 +485,9 @@ struct FNetworkPhysicsSettingsNetworkPhysicsComponent
 		, bOverrideValidateDataOnGameThread(0)
 		, bOverrideApplySimProxyStateAtRuntime(0)
 		, bOverrideApplySimProxyInputAtRuntime(0)
+		, bOverrideTriggerResimOnInputReceive(0)
+		, bOverrideApplyInputDecayOverSetTime(0)
+		, bOverrideInputDecaySetTime(0)
 	{};
 
 	GENERATED_BODY()
@@ -498,6 +529,12 @@ struct FNetworkPhysicsSettingsNetworkPhysicsComponent
 	uint32 bOverrideApplySimProxyStateAtRuntime : 1;
 	UPROPERTY(config, EditDefaultsOnly, Category = "Overrides", Meta = (InlineEditConditionToggle))
 	uint32 bOverrideApplySimProxyInputAtRuntime : 1;
+	UPROPERTY(config, EditDefaultsOnly, Category = "Overrides", Meta = (InlineEditConditionToggle))
+	uint32 bOverrideTriggerResimOnInputReceive : 1;
+	UPROPERTY(config, EditDefaultsOnly, Category = "Overrides", Meta = (InlineEditConditionToggle))
+	uint32 bOverrideApplyInputDecayOverSetTime : 1;
+	UPROPERTY(config, EditDefaultsOnly, Category = "Overrides", Meta = (InlineEditConditionToggle))
+	uint32 bOverrideInputDecaySetTime : 1;
 
 	/* Overrides CVar : np2.Resim.RedundantInputs -- How many extra inputs to send with each unreliable network message, to account for packetloss.From owning client to server and server to owning client.
 	 * NOTE: This is disabled while np2.Resim.DynamicInputScaling.Enabled is enabled. */
@@ -569,6 +606,21 @@ struct FNetworkPhysicsSettingsNetworkPhysicsComponent
 	bool bApplySimProxyInputAtRuntime = PhysicsReplicationCVars::ResimulationCVars::bApplySimProxyInputAtRuntime;
 	const bool GetApplySimProxyInputAtRuntime() const { return bOverrideApplySimProxyInputAtRuntime ? bApplySimProxyInputAtRuntime : PhysicsReplicationCVars::ResimulationCVars::bApplySimProxyInputAtRuntime; }
 
+	// Overrides CVar: np2.Resim.TriggerResimOnInputReceive -- When true, a resimulation will be requested to the frame of the latest frame of received inputs this frame
+	UPROPERTY(config, EditDefaultsOnly, Category = "Overrides", Meta = (EditCondition = "bOverrideTriggerResimOnInputReceive"))
+	bool bTriggerResimOnInputReceive = PhysicsReplicationCVars::ResimulationCVars::bTriggerResimOnInputReceive;
+	const bool GetTriggerResimOnInputReceive() const { return bOverrideTriggerResimOnInputReceive ? bTriggerResimOnInputReceive : PhysicsReplicationCVars::ResimulationCVars::bTriggerResimOnInputReceive; }
+
+	// Overrides CVar: np2.Resim.ApplyInputDecayOverSetTime -- When true, apply the Input Decay Curve over a set amount of time instead of over the start of input prediction and end of resim which is variable each resimulation
+	UPROPERTY(config, EditDefaultsOnly, Category = "Overrides", Meta = (EditCondition = "bOverrideApplyInputDecayOverSetTime"))
+	bool bApplyInputDecayOverSetTime = PhysicsReplicationCVars::ResimulationCVars::bApplyInputDecayOverSetTime;
+	const bool GetApplyInputDecayOverSetTime() const { return bOverrideApplyInputDecayOverSetTime ? bApplyInputDecayOverSetTime : PhysicsReplicationCVars::ResimulationCVars::bApplyInputDecayOverSetTime; }
+
+	// Overrides CVar: np2.Resim.InputDecaySetTime -- Applied when np2.Resim.ApplyInputDecayOverSetTime is true, read there for more info. Set time to apply Input Decay Curve over while predicting inputs during resimulation
+	UPROPERTY(config, EditDefaultsOnly, Category = "Overrides", Meta = (EditCondition = "bOverrideInputDecaySetTime"))
+	float InputDecaySetTime = PhysicsReplicationCVars::ResimulationCVars::InputDecaySetTime;
+	const float GetInputDecaySetTime() const { return bOverrideInputDecaySetTime ? InputDecaySetTime : PhysicsReplicationCVars::ResimulationCVars::InputDecaySetTime; }
+
 	/* Curve for input decay during resimulation if input is being reused. 
 	* XAxis = Alpha value from 0.0 to 1.0 where 0 is the start of reusing input and 1 is the last time we will reuse the input this resimulation.
 	* YAxis = The Input Decay value from 0.0 to 1.0 (as a percentage where 1.0 = 100% decay) for the given Alpha. */
diff --git a/Engine/Source/Runtime/Experimental/Chaos/Private/Framework/PhysicsProxyBase.cpp b/Engine/Source/Runtime/Experimental/Chaos/Private/Framework/PhysicsProxyBase.cpp
--- a/Engine/Source/Runtime/Experimental/Chaos/Private/Framework/PhysicsProxyBase.cpp
+++ b/Engine/Source/Runtime/Experimental/Chaos/Private/Framework/PhysicsProxyBase.cpp
@@ -31,6 +31,21 @@ namespace RenderInterpolationCVars
 
 	float RenderInterpDebugDrawZOffset = 0.0f;
 	static FAutoConsoleVariableRef CVarRenderInterpDebugDrawZOffset(TEXT("p.RenderInterp.DebugDrawZOffset"), RenderInterpDebugDrawZOffset, TEXT("Add Z axis offset to DebugDraw calls for Render Interpolation."));
+
+	bool bRenderInterpApplyExponentialDecay = false;
+	static FAutoConsoleVariableRef CVarRenderInterpApplyExponentialDecay(TEXT("p.RenderInterp.ApplyExponentialDecay"), bRenderInterpApplyExponentialDecay, TEXT("When enabled a post-resim error will decay exponentially (instead of linearly) based on half-life time set in ExponentialDecayLinearHalfLife and ExponentialDecayAngularHalfLife."));
+
+	float RenderInterpExponentialDecayLinearHalfLife = 0.06f;
+	static FAutoConsoleVariableRef CVarRenderInterpExponentialDecayLinearHalfLife(TEXT("p.RenderInterp.ExponentialDecayLinearHalfLife"), RenderInterpExponentialDecayLinearHalfLife, TEXT("Sets the positional half-life time for when bApplyExponentialDecay is enabled."));
+
+	float RenderInterpExponentialDecayAngularHalfLife = 0.06f;
+	static FAutoConsoleVariableRef CVarRenderInterpExponentialDecayAngularHalfLife(TEXT("p.RenderInterp.ExponentialDecayAngularHalfLife"), RenderInterpExponentialDecayAngularHalfLife, TEXT("Sets the rotational half-life time for when bApplyExponentialDecay is enabled."));
+
+	float RenderInterpMinimumLinearThreshold = 0.1f;
+	static FAutoConsoleVariableRef CVarRenderInterpMinimumLinearThreshold(TEXT("p.RenderInterp.MinimumLinearThreshold"), RenderInterpMinimumLinearThreshold, TEXT("Squared value, when the remaining render error is below this we clear it, if ApplyExponentialDecay is enabled."));
+
+	float RenderInterpMinimumAngularThreshold = 0.001f;
+	static FAutoConsoleVariableRef CVarRenderInterpMinimumAngularThreshold(TEXT("p.RenderInterp.MinimumAngularThreshold"), RenderInterpMinimumAngularThreshold, TEXT("When the remaining render error angle is below this we clear it, if ApplyExponentialDecay is enabled."));
 }
 
 IPhysicsProxyBase::~IPhysicsProxyBase()
@@ -50,4 +65,149 @@ int32 IPhysicsProxyBase::GetSolverSyncTimestamp_External() const
 	}
 
 	return INDEX_NONE;
+}
+
+const bool FProxyInterpolationError::IsErrorSmoothing() const
+{
+	const bool ApplyExponentialDecay = ErrorInterpolationSettings.Get() ? ErrorInterpolationSettings.Get()->bApplyExponentialDecay : RenderInterpolationCVars::bRenderInterpApplyExponentialDecay;
+	if (ApplyExponentialDecay)
+	{
+		const float MinimumLinearThreshold = ErrorInterpolationSettings.Get() ? ErrorInterpolationSettings.Get()->MinimumLinearThreshold : RenderInterpolationCVars::RenderInterpMinimumLinearThreshold;
+		const float MinimumAngularThreshold = ErrorInterpolationSettings.Get() ? ErrorInterpolationSettings.Get()->MinimumAngularThreshold : RenderInterpolationCVars::RenderInterpMinimumAngularThreshold;
+		
+		return ErrorX.SizeSquared() > MinimumLinearThreshold
+			|| ErrorR.GetAngle() > MinimumAngularThreshold;
+	}
+	return LastSimTick < EndDecayTick;
+}
+
+
+void FProxyInterpolationError::AccumlateErrorXR(const Chaos::FVec3 X, const FQuat R, const int32 CurrentSimTick, const int32 ErrorSmoothDuration)
+{
+	LastSimTick = CurrentSimTick - 1; // Error is from the previous simulation tick, not the current
+	EndDecayTick = LastSimTick + ErrorSmoothDuration;
+	SimTicks = 0;
+
+	const bool ApplyExponentialDecay = ErrorInterpolationSettings.Get() ? ErrorInterpolationSettings.Get()->bApplyExponentialDecay : RenderInterpolationCVars::bRenderInterpApplyExponentialDecay;
+	if (ApplyExponentialDecay || IsErrorSmoothing())
+	{
+		ErrorX += X;
+		ErrorR = R * ErrorR;
+	}
+	else
+	{
+		Reset();
+	}
+}
+
+const bool FProxyInterpolationError::UpdateError(const int32 CurrentSimTick, const Chaos::FReal AsyncFixedTimeStep)
+{
+	// Cache how many simulation ticks have passed since last call
+	SimTicks = CurrentSimTick - LastSimTick;
+	LastSimTick = CurrentSimTick;
+
+	if (SimTicks > 0)
+	{
+		return DecayError(AsyncFixedTimeStep);
+	}
+	return false;
+}
+
+const bool FProxyInterpolationError::DirectionalDecay(Chaos::FVec3 Direction, float ErrorDirectionalDecayMultiplier)
+{
+	if (ErrorX.IsNearlyZero())
+	{
+		return false;
+	}
+
+	ErrorDirectionalDecayMultiplier = ErrorInterpolationSettings.Get() ? ErrorInterpolationSettings.Get()->ErrorDirectionalDecayMultiplier : ErrorDirectionalDecayMultiplier;
+	if (ErrorDirectionalDecayMultiplier > 0.0f && IsErrorSmoothing() && SimTicks > 0)
+	{
+		Chaos::FVec3 DirProjection = Direction.ProjectOnTo(ErrorX) * ErrorDirectionalDecayMultiplier;
+		const Chaos::FRealDouble DotProd = Chaos::FVec3::DotProduct(DirProjection, ErrorX);
+
+		if (DotProd > 0.0f)
+		{
+			if (DirProjection.SizeSquared() < ErrorX.SizeSquared())
+			{
+				ErrorX -= DirProjection;
+			}
+			else
+			{
+				ErrorX = Chaos::FVec3::Zero();
+			}
+			return true;
+		}
+	}
+	return false;
+}
+
+const bool FProxyInterpolationError::DecayError(const Chaos::FReal AsyncFixedTimeStep)
+{
+	if (SimTicks <= 0)
+	{
+		return false;
+	}
+
+	if (!IsErrorSmoothing())
+	{
+		Reset();
+		return false;
+	}
+
+	const bool ApplyExponentialDecay = ErrorInterpolationSettings.Get() ? ErrorInterpolationSettings.Get()->bApplyExponentialDecay : RenderInterpolationCVars::bRenderInterpApplyExponentialDecay;
+
+	if (ApplyExponentialDecay)
+	{
+		ErrorXPrev = ErrorX;
+		ErrorRPrev = ErrorR;
+
+		const float ExponentialDecayLinearHalfLife = ErrorInterpolationSettings.Get() ? ErrorInterpolationSettings.Get()->ExponentialDecayLinearHalfLife : RenderInterpolationCVars::RenderInterpExponentialDecayLinearHalfLife;
+		const float ExponentialDecayAngularHalfLife = ErrorInterpolationSettings.Get() ? ErrorInterpolationSettings.Get()->ExponentialDecayAngularHalfLife : RenderInterpolationCVars::RenderInterpExponentialDecayAngularHalfLife;
+		const double SimulatedTime = SimTicks * AsyncFixedTimeStep;
+
+		if (ExponentialDecayLinearHalfLife > UE_SMALL_NUMBER)
+		{
+			const double DecayAmountX = FMath::Exp2(-SimulatedTime / ExponentialDecayLinearHalfLife); // Same as FMath::Pow(0.5f, SimulatedTime / RenderInterpolationCVars::RenderInterpExponentialDecayLinearHalfLife)
+			ErrorX = ErrorX * DecayAmountX;
+		}
+		else
+		{
+			ErrorX = Chaos::FVec3::ZeroVector;
+		}
+
+		if (ExponentialDecayAngularHalfLife > UE_SMALL_NUMBER)
+		{
+			const double DecayAmountR = FMath::Exp2(-SimulatedTime / ExponentialDecayAngularHalfLife); // Same as FMath::Pow(0.5f, SimulatedTime / RenderInterpolationCVars::RenderInterpExponentialDecayAngularHalfLife)
+			ErrorR = FQuat::Slerp(FQuat::Identity, ErrorR, DecayAmountR);
+		}
+		else
+		{
+			ErrorR = FQuat::Identity;
+		}
+
+		return true;
+	}
+	else
+	{
+		// Linear decay
+		// Example: If we want to decay an error of 100 over 10ticks (i.e. 10% each tick)
+		// First step:  9/10 = 0.9   |  100 * 0.9  = 90 error
+		// Second step: 8/9  = 0.888 |  90 * 0.888 = 80 error
+		// Third step: 7/8  = 0.875 |  80 * 0.875 = 70 error
+		// etc.
+		const int32 ErrorSmoothingCount = EndDecayTick - LastSimTick;
+		int32 Ticks = FMath::Clamp(SimTicks, 0, ErrorSmoothingCount);
+		for (; Ticks > 0; Ticks--)
+		{
+			Chaos::FRealSingle Alpha = Chaos::FRealSingle(ErrorSmoothingCount - 1) / Chaos::FRealSingle(ErrorSmoothingCount);
+			Alpha = FMath::Clamp(Alpha, 0.f, 1.f);
+
+			ErrorXPrev = ErrorX;
+			ErrorX *= Alpha;
+			ErrorRPrev = ErrorR;
+			ErrorR = FQuat::Slerp(FQuat::Identity, ErrorR, Alpha);
+		}
+		return true;
+	}
 }
\ No newline at end of file
diff --git a/Engine/Source/Runtime/Experimental/Chaos/Public/Chaos/Framework/PhysicsProxyBase.h b/Engine/Source/Runtime/Experimental/Chaos/Public/Chaos/Framework/PhysicsProxyBase.h
--- a/Engine/Source/Runtime/Experimental/Chaos/Public/Chaos/Framework/PhysicsProxyBase.h
+++ b/Engine/Source/Runtime/Experimental/Chaos/Public/Chaos/Framework/PhysicsProxyBase.h
@@ -38,6 +38,11 @@ namespace RenderInterpolationCVars
 	extern bool bRenderInterpDebugDraw;
 	extern bool bRenderInterpDebugDrawGC;
 	extern float RenderInterpDebugDrawZOffset;
+	extern bool bRenderInterpApplyExponentialDecay;
+	extern float RenderInterpExponentialDecayLinearHalfLife;
+	extern float RenderInterpExponentialDecayAngularHalfLife;
+	extern float RenderInterpMinimumLinearThreshold;
+	extern float RenderInterpMinimumAngularThreshold;
 }
 
 enum class EProxyInterpolationType : uint32
@@ -212,6 +217,11 @@ struct FErrorInterpolationSettings
 		, MaximumErrorCorrectionBeforeSnapping(250.0f)
 		, MaximumErrorCorrectionDesyncTimeBeforeSnapping(0.6f)
 		, ErrorDirectionalDecayMultiplier(0.0f)
+		, bApplyExponentialDecay(false)
+		, ExponentialDecayLinearHalfLife(0.06f)
+		, ExponentialDecayAngularHalfLife(0.06f)
+		, MinimumLinearThreshold(0.1f)
+		, MinimumAngularThreshold(0.001f)
 	{}
 	~FErrorInterpolationSettings() { }
 
@@ -227,6 +237,21 @@ struct FErrorInterpolationSettings
 
 	/** Decay error offset in the direction that the physics object is moving, value is multiplier of projected offset direction, 0.25 means a 25 % decay of the magnitude in the direction of physics travel.Deactivate by setting to 0 */
 	float ErrorDirectionalDecayMultiplier;
+
+	/** When enabled a post-resim error will decay exponentially (instead of linearly) based on half-life time set in ExponentialDecayLinearHalfLife and ExponentialDecayAngularHalfLife. */
+	bool bApplyExponentialDecay;
+
+	/** Sets the positional half-life time for when bApplyExponentialDecay is enabled. */
+	float ExponentialDecayLinearHalfLife;
+
+	/** Sets the rotational half-life time for when bApplyExponentialDecay is enabled. */
+	float ExponentialDecayAngularHalfLife;
+
+	/** Squared value, when the remaining render error is below this we clear it, if ApplyExponentialDecay is enabled. */
+	float MinimumLinearThreshold;
+
+	/** When the remaining render error angle is below this we clear it, if ApplyExponentialDecay is enabled. */
+	float MinimumAngularThreshold;
 };
 
 struct FProxyInterpolationBase
@@ -307,106 +332,20 @@ struct FProxyInterpolationError : FProxyInterpolationBase
 	static const EProxyInterpolationType InterpolationType = EProxyInterpolationType::ErrorLinear;
 	virtual const EProxyInterpolationType GetInterpolationType() const override { return FProxyInterpolationError::InterpolationType; }
 
-	virtual const bool IsErrorSmoothing() const override { return LastSimTick < EndDecayTick; }
+	CHAOS_API virtual const bool IsErrorSmoothing() const override;
 
 	virtual const Chaos::FVec3 GetErrorX(const Chaos::FRealSingle Alpha) const override { return FMath::Lerp(ErrorXPrev, ErrorX, Alpha); }
 
 	virtual const FQuat GetErrorR(const Chaos::FRealSingle Alpha) const override { return FQuat::Slerp(ErrorRPrev, ErrorR, Alpha); }
 
-	virtual void AccumlateErrorXR(const Chaos::FVec3 X, const FQuat R, const int32 CurrentSimTick, const int32 ErrorSmoothDuration) override
-	{
-		LastSimTick = CurrentSimTick - 1; // Error is from the previous simulation tick, not the current
-		EndDecayTick = LastSimTick + ErrorSmoothDuration;
-		SimTicks = 0;
-
-		if (IsErrorSmoothing())
-		{
-			ErrorX += X;
-			ErrorR = R * ErrorR;
-		}
-		else
-		{
-			Reset();
-		}
-	}
-
-	virtual const bool UpdateError(const int32 CurrentSimTick, const Chaos::FReal AsyncFixedTimeStep) override
-	{
-		// Cache how many simulation ticks have passed since last call
-		SimTicks = CurrentSimTick - LastSimTick;
-		LastSimTick = CurrentSimTick;
-
-		if (SimTicks > 0)
-		{
-			return DecayError();
-		}
-		return false;
-	}
-
-	virtual const bool DirectionalDecay(Chaos::FVec3 Direction, float ErrorDirectionalDecayMultiplier) override
-	{
-		if (ErrorX.IsNearlyZero())
-		{
-			return false;
-		}
+	CHAOS_API virtual void AccumlateErrorXR(const Chaos::FVec3 X, const FQuat R, const int32 CurrentSimTick, const int32 ErrorSmoothDuration) override;
 
-		ErrorDirectionalDecayMultiplier = ErrorInterpolationSettings.Get() ? ErrorInterpolationSettings.Get()->ErrorDirectionalDecayMultiplier : ErrorDirectionalDecayMultiplier;
-		if (ErrorDirectionalDecayMultiplier > 0.0f && IsErrorSmoothing() && SimTicks > 0)
-		{
-			Chaos::FVec3 DirProjection = Direction.ProjectOnTo(ErrorX) * ErrorDirectionalDecayMultiplier;
-			const Chaos::FRealDouble DotProd = Chaos::FVec3::DotProduct(DirProjection, ErrorX);
-			
-			if (DotProd > 0.0f)
-			{
-				if (DirProjection.SizeSquared() < ErrorX.SizeSquared())
-				{
-					ErrorX -= DirProjection;
-				}
-				else
-				{
-					ErrorX = Chaos::FVec3::Zero();
-				}
-				return true;
-			}
-		}
-		return false;
-	}
+	CHAOS_API virtual const bool UpdateError(const int32 CurrentSimTick, const Chaos::FReal AsyncFixedTimeStep) override;
 
+	CHAOS_API virtual const bool DirectionalDecay(Chaos::FVec3 Direction, float ErrorDirectionalDecayMultiplier) override;
+	
 protected:
-	virtual const bool DecayError()
-	{
-		if (SimTicks <= 0)
-		{
-			return false;
-		}
-
-		if (!IsErrorSmoothing())
-		{
-			Reset();
-			return false;
-		}
-
-		// Linear decay
-		// Example: If we want to decay an error of 100 over 10ticks (i.e. 10% each tick)
-		// First step:  9/10 = 0.9   |  100 * 0.9  = 90 error
-		// Second step: 8/9  = 0.888 |  90 * 0.888 = 80 error
-		// Third step: 7/8  = 0.875 |  80 * 0.875 = 70 error
-		// etc.
-		const int32 ErrorSmoothingCount = EndDecayTick - LastSimTick;
-		int32 Ticks = FMath::Clamp(SimTicks, 0, ErrorSmoothingCount);
-		for (; Ticks > 0; Ticks--)
-		{
-			Chaos::FRealSingle Alpha = Chaos::FRealSingle(ErrorSmoothingCount - 1) / Chaos::FRealSingle(ErrorSmoothingCount);
-			Alpha = FMath::Clamp(Alpha, 0.f, 1.f);
-
-			ErrorXPrev = ErrorX;
-			ErrorX *= Alpha;
-			ErrorRPrev = ErrorR;
-			ErrorR = FQuat::Slerp(FQuat::Identity, ErrorR, Alpha);
-		}
-
-		return true;
-	}
+	CHAOS_API virtual const bool DecayError(const Chaos::FReal AsyncFixedTimeStep);
 
 	virtual void Reset()
 	{
