ClusterUnion : propagate values and gravity to cluster union ( when using normal cluster union or cluster group index )

diff --git a/Engine/Source/Runtime/Experimental/Chaos/Private/Chaos/ClusterUnionManager.cpp b/Engine/Source/Runtime/Experimental/Chaos/Private/Chaos/ClusterUnionManager.cpp
--- a/Engine/Source/Runtime/Experimental/Chaos/Private/Chaos/ClusterUnionManager.cpp
+++ b/Engine/Source/Runtime/Experimental/Chaos/Private/Chaos/ClusterUnionManager.cpp
@@ -642,6 +642,10 @@ namespace Chaos
 
 		// Use the minimum sleep multiplier of all member particles
 		FRealSingle MinSleepThresholdMultiplier = TNumericLimits<FRealSingle>::Max();
+		
+		Chaos::FReal MaxLinearEtherDrag = Cluster->InternalCluster->LinearEtherDrag();
+		Chaos::FReal MaxAngularEtherDrag = Cluster->InternalCluster->AngularEtherDrag();
+		bool bEnableGravity = Cluster->InternalCluster->GravityEnabled();
 
 		for (FPBDRigidParticleHandle* Particle : FinalParticlesToAdd)
 		{
@@ -663,6 +667,9 @@ namespace Chaos
 			bIsOneWayInteraction &= Particle->OneWayInteraction();
 			bIsMACD |= Particle->MACDEnabled();
 			MinSleepThresholdMultiplier = FMath::Min(MinSleepThresholdMultiplier, Particle->SleepThresholdMultiplier());
+			MaxLinearEtherDrag = FMath::Max(MaxLinearEtherDrag, Particle->LinearEtherDrag());
+			MaxAngularEtherDrag = FMath::Max(MaxAngularEtherDrag, Particle->AngularEtherDrag());
+			bEnableGravity &= Particle->GravityEnabled(); // favor disabled gravity 
 
 			SolverIterationSettings = Private::FIterationSettings::Merge(SolverIterationSettings, Particle->IterationSettings());
 
@@ -734,6 +741,9 @@ namespace Chaos
 				Cluster->InternalCluster->SetInitialOverlapDepenetrationVelocity(ParticlePropertySource->InitialOverlapDepenetrationVelocity());
 			}
 		}
+		Cluster->InternalCluster->SetLinearEtherDrag(MaxLinearEtherDrag);
+		Cluster->InternalCluster->SetAngularEtherDrag(MaxAngularEtherDrag);
+		Cluster->InternalCluster->SetGravityEnabled(bEnableGravity);
 
 		Cluster->InternalCluster->SetOneWayInteraction(bIsOneWayInteraction);
 		Cluster->InternalCluster->SetMACDEnabled(bIsMACD);
