Fix code rot in the old IAX backend that was preventing a fast exit if the process was stuck failing to resolve the distributed backend.

#rnx

- When the http code was moved to it's own class the call to trigger DistributedEndpointEvent on shutdown was lost as it used to be in the ::Stop method of the threading code. This meant that if the distributed endpoint had failed to resolve and we were waiting for N seconds to try again that the shutdown process would fail to wake up that wait early and we'd have to wait the full N seconds.
- We once again trigger DistributedEndpointEvent on shutdown.
- Removed the event from FDistributionEndpoints::ResolveEndpoints as it wasn't actually doing anything. It was triggered when the resolve was a success but the trigger was only waited on if the resolution failed.

diff --git a/Engine/Source/Runtime/Experimental/IoStore/OnDemand/Private/DistributionEndpoints.cpp b/Engine/Source/Runtime/Experimental/IoStore/OnDemand/Private/DistributionEndpoints.cpp
--- a/Engine/Source/Runtime/Experimental/IoStore/OnDemand/Private/DistributionEndpoints.cpp
+++ b/Engine/Source/Runtime/Experimental/IoStore/OnDemand/Private/DistributionEndpoints.cpp
@@ -35,12 +35,6 @@ static FAutoConsoleVariableRef CVar_DistributedEndpointTimeout(
 static FDistributionEndpoints::EResult ParseResponse(FIoBuffer, TArray<FString>&);
 
 FDistributionEndpoints::EResult FDistributionEndpoints::ResolveEndpoints(const FString& DistributionUrl, TArray<FString>& OutServiceUrls)
-{
-	FEventRef Event;
-	return ResolveEndpoints(DistributionUrl, OutServiceUrls, *Event.Get());
-}
-
-FDistributionEndpoints::EResult FDistributionEndpoints::ResolveEndpoints(const FString& DistributionUrl, TArray<FString>& OutServiceUrls, FEvent& Event)
 {
 	using namespace HTTP;
 
@@ -53,7 +47,6 @@ FDistributionEndpoints::EResult FDistributionEndpoints::ResolveEndpoints(const F
 	auto OnRequestStatus = [
 		&bHasResponse,
 		&OutServiceUrls,
-		&Event,
 		&Result,
 		bIsChunkedTransfer = false,
 		Dest=FIoBuffer(),
@@ -82,7 +75,6 @@ FDistributionEndpoints::EResult FDistributionEndpoints::ResolveEndpoints(const F
 
 			bHasResponse = true;
 			Result = ParseResponse(Dest, OutServiceUrls);
-			Event.Trigger();
 
 			return;
 		}
diff --git a/Engine/Source/Runtime/Experimental/IoStore/OnDemand/Private/DistributionEndpoints.h b/Engine/Source/Runtime/Experimental/IoStore/OnDemand/Private/DistributionEndpoints.h
--- a/Engine/Source/Runtime/Experimental/IoStore/OnDemand/Private/DistributionEndpoints.h
+++ b/Engine/Source/Runtime/Experimental/IoStore/OnDemand/Private/DistributionEndpoints.h
@@ -29,15 +29,6 @@ public:
 	 * @param OutServiceUrls	An array into which the resolved urls will be placed
 	 */
 	FDistributionEndpoints::EResult ResolveEndpoints(const FString& DistributionUrl, TArray<FString>& OutServiceUrls);
-
-	/**
-	 * Queries the distributed endpoint in order to get a list of service urls from which we can download data.
-	 * 
-	 * @param DistributionUrl	The url of the distributed endpoint to resolve
-	 * @param OutServiceUrls	An array into which the resolved urls will be placed
-	 * @param Event				An event which can be triggered to cancel the request
-	 */
-	FDistributionEndpoints::EResult ResolveEndpoints(const FString& DistributionUrl, TArray<FString>& OutServiceUrls, FEvent& Event);
 };
 
 } // namespace UE::IoStore
diff --git a/Engine/Source/Runtime/Experimental/IoStore/OnDemand/Private/OnDemandIoDispatcherBackend.cpp b/Engine/Source/Runtime/Experimental/IoStore/OnDemand/Private/OnDemandIoDispatcherBackend.cpp
--- a/Engine/Source/Runtime/Experimental/IoStore/OnDemand/Private/OnDemandIoDispatcherBackend.cpp
+++ b/Engine/Source/Runtime/Experimental/IoStore/OnDemand/Private/OnDemandIoDispatcherBackend.cpp
@@ -757,6 +757,8 @@ void FOnDemandIoBackend::Shutdown()
 	UE_LOG(LogIas, Log, TEXT("Shutting down on demand I/O dispatcher backend"));
 
 	bStopRequested.store(true, std::memory_order_relaxed);
+	DistributedEndpointEvent->Trigger();
+
 	HostGroupTask.Wait();
 
 	// The CompleteRequest tasks may still be executing a while after the IoDispatcher has been notified about the completed io requests.
@@ -1235,7 +1237,7 @@ bool FOnDemandIoBackend::ResolveDistributedEndpoint(const FDistributedEndpointUr
 	while (!bStopRequested)
 	{
 		FDistributionEndpoints Resolver;
-		FDistributionEndpoints::EResult Result = Resolver.ResolveEndpoints(DistributedEndpointUrl.EndpointUrl, OutUrls, *DistributedEndpointEvent.Get());
+		FDistributionEndpoints::EResult Result = Resolver.ResolveEndpoints(DistributedEndpointUrl.EndpointUrl, OutUrls);
 		if (Result == FDistributionEndpoints::EResult::Success)
 		{
 			Stats.OnHttpDistributedEndpointResolved();
@@ -1253,8 +1255,10 @@ bool FOnDemandIoBackend::ResolveDistributedEndpoint(const FDistributedEndpointUr
 
 		if (!bStopRequested)
 		{
+			UE_LOG(LogIas, Error, TEXT("ResolveDistributedEndpoint failed, waiting for %d seconds"), GDistributedEndpointRetryWaitTime);
 			const uint32 WaitTime = GDistributedEndpointRetryWaitTime >= 0 ? (static_cast<uint32>(GDistributedEndpointRetryWaitTime) * 1000) : MAX_uint32;
 			DistributedEndpointEvent->Wait(WaitTime);
+			UE_LOG(LogIas, Error, TEXT("ResolveDistributedEndpoint ready to try again"));
 		}
 	}
 
