Missed a frame reset

#rnx

diff --git a/Engine/Source/Runtime/Experimental/IoStore/HttpClient/Private/TransactionTwo.inl b/Engine/Source/Runtime/Experimental/IoStore/HttpClient/Private/TransactionTwo.inl
--- a/Engine/Source/Runtime/Experimental/IoStore/HttpClient/Private/TransactionTwo.inl
+++ b/Engine/Source/Runtime/Experimental/IoStore/HttpClient/Private/TransactionTwo.inl
@@ -514,7 +514,13 @@ FOutcome FDriverNg::RecvBody(FTlsPeer& Peer, FMutableMemoryView Dest)
 		FMutableMemoryView View = Dest.Left(RecvRemaining);
 		if (View.IsEmpty())
 		{
-			return EndOfStream ? FOutcome::Ok(RecvSize) : FOutcome::WaitBuffer(RecvSize);
+			if (EndOfStream)
+			{
+				RecvFrame = {};
+				return FOutcome::Ok(RecvSize);
+			}
+
+			return FOutcome::WaitBuffer(RecvSize);
 		}
 
 		FOutcome Outcome = Peer.Recv((char*)View.GetData(), int32(View.GetSize()));
