Deleted unused code.

#rnx

diff --git a/Engine/Source/Runtime/Experimental/IoStore/HttpClient/Private/TransactionOne.inl b/Engine/Source/Runtime/Experimental/IoStore/HttpClient/Private/TransactionOne.inl
--- a/Engine/Source/Runtime/Experimental/IoStore/HttpClient/Private/TransactionOne.inl
+++ b/Engine/Source/Runtime/Experimental/IoStore/HttpClient/Private/TransactionOne.inl
@@ -956,9 +956,6 @@ public:
 	bool				IsValid() const;
 
 private:
-	friend FTransactRef CreateTransactOne(FBuffer&);
-
-						FTransactRef(UPTRINT Ptr);
 	const FTransaction*	Get() const;
 	FTransaction*		Get();
 	UPTRINT				Self = 0;
