Version accessor

#rnx

diff --git a/Engine/Source/Runtime/Experimental/IoStore/HttpClient/Private/Peer.inl b/Engine/Source/Runtime/Experimental/IoStore/HttpClient/Private/Peer.inl
--- a/Engine/Source/Runtime/Experimental/IoStore/HttpClient/Private/Peer.inl
+++ b/Engine/Source/Runtime/Experimental/IoStore/HttpClient/Private/Peer.inl
@@ -541,6 +541,7 @@ public:
 	FHttpPeer&		operator = (FHttpPeer&& Rhs);
 	FHttpPeer&		operator = (const FHttpPeer& Rhs) = delete;
 					FHttpPeer(FParams&& Params);
+	uint32			GetVersion() const;
 	FOutcome		Handshake();
 	FTransactRef	Transact();
 	FOutcome		GetPendingTransactId();
@@ -583,6 +584,12 @@ FHttpPeer& FHttpPeer::operator = (FHttpPeer&& Rhs)
 	return *this;
 }
 
+////////////////////////////////////////////////////////////////////////////////
+uint32 FHttpPeer::GetVersion() const
+{
+	return Proto;
+}
+
 ////////////////////////////////////////////////////////////////////////////////
 FOutcome FHttpPeer::Handshake()
 {
