Surface version result on the off chance it might be needed

#rnx

diff --git a/Engine/Source/Runtime/Experimental/IoStore/HttpClient/Private/Peer.inl b/Engine/Source/Runtime/Experimental/IoStore/HttpClient/Private/Peer.inl
--- a/Engine/Source/Runtime/Experimental/IoStore/HttpClient/Private/Peer.inl
+++ b/Engine/Source/Runtime/Experimental/IoStore/HttpClient/Private/Peer.inl
@@ -610,12 +610,20 @@ FOutcome FHttpPeer::Handshake()
 		/* intentional fallthrough */
 	}
 
+	FOutcome Outcome = FOutcome::None();
 	switch (Proto)
 	{
-	case 1:		return HandshakeOne(*this, PeerData);
-	case 2:		return HandshakeTwo(*this, PeerData);
+	case 1:		Outcome = HandshakeOne(*this, PeerData); break;
+	case 2:		Outcome = HandshakeTwo(*this, PeerData); break;
 	default:	return FOutcome::Error("Unsupported");
 	}
+
+	if (!Outcome.IsOk())
+	{
+		return Outcome;
+	}
+
+	return FOutcome::Ok(Proto);
 }
 
 ////////////////////////////////////////////////////////////////////////////////
