Improve error logging when failing file operations in the IAS cache.

#rnx

- Check if deleting the journal failed and if so report the error.
- Use the new OpenWrite api when trying to open the journal handle as it should contain more error info.
-- Note that due to TManagedStoragePlatformFile this doesn't work great for runtime games at the moment as the LastErrorCode tends to get wiped by critical section locking/unlocking depending on the target platform. This will be improved in the future.

diff --git a/Engine/Source/Runtime/Experimental/IoStore/OnDemand/Private/JournaledCache.cpp b/Engine/Source/Runtime/Experimental/IoStore/OnDemand/Private/JournaledCache.cpp
--- a/Engine/Source/Runtime/Experimental/IoStore/OnDemand/Private/JournaledCache.cpp
+++ b/Engine/Source/Runtime/Experimental/IoStore/OnDemand/Private/JournaledCache.cpp
@@ -477,7 +477,14 @@ void FDiskJournal::Drop()
 	GetPath(JrnPath);
 
 	IPlatformFile& Ipf = IPlatformFile::GetPlatformPhysical();
-	Ipf.DeleteFile(*JrnPath);
+	if (!Ipf.DeleteFile(*JrnPath))
+	{
+		const uint32 SystemError = FPlatformMisc::GetLastError();
+		TCHAR SystemErrorMsg[MAX_SPRINTF] = { 0 };
+		FPlatformMisc::GetSystemErrorMessage(SystemErrorMsg, UE_ARRAY_COUNT(SystemErrorMsg), SystemError);
+
+		UE_LOG(LogIas, Error, TEXT("Failed to delete '%s' for FDiskJournal ('%d' - '%s'"), *JrnPath, SystemError, SystemErrorMsg);
+	}
 
 	Cursor = 0;
 	Entries.Reset();
@@ -493,9 +500,15 @@ void FDiskJournal::OpenJrnFile()
 
 	IPlatformFile& Ipf = IPlatformFile::GetPlatformPhysical();
 
-	IFileHandle* Handle = Ipf.OpenWrite(*JrnPath, true, true);
-	UE_CLOG(Handle == nullptr, LogIas, Error, TEXT("Failed to open '%s' for FDiskJournal"), *JrnPath);
-	JrnHandle.Reset(Handle);
+	FFileOpenResult Result = Ipf.OpenWrite(*JrnPath, IPlatformFile::EOpenWriteFlags::Append | IPlatformFile::EOpenWriteFlags::AllowRead);
+	if (Result.IsValid())
+	{
+		JrnHandle = Result.StealValue();
+	}
+	else
+	{
+		UE_LOG(LogIas, Error, TEXT("Failed to open '%s' for FDiskJournal (%s)"), *JrnPath, *Result.GetError().GetMessage());
+	}
 }
 
 ////////////////////////////////////////////////////////////////////////////////
