Constraints: also evaluate any constraints that target the component bound to the rig (even if they do not target the manipulated control)

#rnx

#FailedValidations Preflight

diff --git a/Engine/Plugins/Animation/ControlRig/Source/ControlRigEditor/Private/EditMode/ControlRigEditMode.cpp b/Engine/Plugins/Animation/ControlRig/Source/ControlRigEditor/Private/EditMode/ControlRigEditMode.cpp
--- a/Engine/Plugins/Animation/ControlRig/Source/ControlRigEditor/Private/EditMode/ControlRigEditMode.cpp
+++ b/Engine/Plugins/Animation/ControlRig/Source/ControlRigEditor/Private/EditMode/ControlRigEditMode.cpp
@@ -5298,7 +5298,10 @@ void FControlRigEditMode::SetControlShapeTransform(
 		{
 			KeyframeData.bConstraintSpace = true;
 		}
+	}
 
+	if (bNeedsConstraintPostProcess || ConstraintsCache.HasAnyActiveConstraint(GetHostingSceneComponent(ControlRig), World))
+	{
 		TGuardValue<bool> CompensateGuard(FMovieSceneConstraintChannelHelper::bDoNotCompensate, true);
 		const FConstraintsManagerController& Controller = FConstraintsManagerController::Get(World);
 		Controller.EvaluateAllConstraints();
diff --git a/Engine/Source/Runtime/Experimental/Animation/Constraints/Private/Transform/TransformConstraintUtil.cpp b/Engine/Source/Runtime/Experimental/Animation/Constraints/Private/Transform/TransformConstraintUtil.cpp
--- a/Engine/Source/Runtime/Experimental/Animation/Constraints/Private/Transform/TransformConstraintUtil.cpp
+++ b/Engine/Source/Runtime/Experimental/Animation/Constraints/Private/Transform/TransformConstraintUtil.cpp
@@ -129,10 +129,54 @@ TOptional<FTransform> FConstraintsInteractionCache::GetParentTransform(const uin
 	static const TOptional<FTransform> Invalid;
 	return Invalid;
 }
+
+const TArray< TWeakObjectPtr<UTickableConstraint> >& FConstraintsInteractionCache::Get(UObject* InTarget, UWorld* InWorld)
+{
+	static const TArray< TWeakObjectPtr<UTickableConstraint> > Empty;
+	if (!InTarget)
+	{
+		return Empty;
+	}
+	
+	UWorld* World = InWorld ? InWorld : InTarget->GetWorld();
+	if (!World)
+	{
+		return Empty;
+	}
+
+	if (TArray< TWeakObjectPtr<UTickableConstraint> >* FoundConstraints = PerTargetConstraints.Find(InTarget))
+	{
+		return *FoundConstraints;
+	}
+
+	TArray< TWeakObjectPtr<UTickableConstraint> >& TransformConstraints = PerTargetConstraints.Add(InTarget);
+	
+	static constexpr bool bSorted = true;
+	const FConstraintsManagerController& Controller = FConstraintsManagerController::Get(World);
+	auto SameTargetPredicate = [InTarget](const TWeakObjectPtr<UTickableConstraint>& WeakConstraint)
+	{
+		if (const UTickableTransformConstraint* TransformConstraint = Cast<UTickableTransformConstraint>(WeakConstraint.Get()))
+		{
+			const UObject* Target = IsValid(TransformConstraint->ChildTRSHandle) ? TransformConstraint->ChildTRSHandle->GetTarget().Get() : nullptr;
+			return Target && Target == InTarget;
+		}
+		return false;
+	};
+	TransformConstraints = Controller.GetConstraintsByPredicate(SameTargetPredicate, bSorted);
+
+	return TransformConstraints;
+}
+	
+bool FConstraintsInteractionCache::HasAnyActiveConstraint(UObject* InTarget, UWorld* InWorld)
+{
+	const TArray< TWeakObjectPtr<UTickableConstraint> >& TransformConstraints = Get(InTarget, InWorld);
+	return GetLastActiveConstraintIndex(TransformConstraints) != INDEX_NONE;
+}
 	
 void FConstraintsInteractionCache::Reset()
 {
 	PerHandleConstraints.Reset();
+	PerTargetConstraints.Reset();
 }
 
 void FConstraintsInteractionCache::RegisterNotifications(UWorld* InWorld)
diff --git a/Engine/Source/Runtime/Experimental/Animation/Constraints/Public/Transform/TransformConstraintUtil.h b/Engine/Source/Runtime/Experimental/Animation/Constraints/Public/Transform/TransformConstraintUtil.h
--- a/Engine/Source/Runtime/Experimental/Animation/Constraints/Public/Transform/TransformConstraintUtil.h
+++ b/Engine/Source/Runtime/Experimental/Animation/Constraints/Public/Transform/TransformConstraintUtil.h
@@ -8,8 +8,7 @@
 #include "Containers/Map.h"
 
 #include "Delegates/IDelegateInstance.h"
-
-#include "UObject/WeakObjectPtr.h"
+#include "UObject/ObjectKey.h"
 
 #define UE_API CONSTRAINTS_API
 
@@ -25,7 +24,7 @@ namespace UE::TransformConstraintUtil
 /**
  * FConstraintsInteractionCache is designed to minimize the number of requests made to the ConstraintManager and thus improve performance.
  * it is particularly useful for interface components that need to be updated frequently.
- * note that PerHandleConstraints is reset when the constraint graph is updated, to keep it synchronized.
+ * note that PerHandleConstraints and PerTargetConstraints are reset when the constraint graph is updated, to keep it synchronized.
  */
 	
 struct FConstraintsInteractionCache
@@ -36,17 +35,21 @@ public:
 	UE_API const TArray< TWeakObjectPtr<UTickableConstraint> >& Get(const UObject* InObject, const FName& InAttachmentName);
 	UE_API bool HasAnyActiveConstraint(const UObject* InObject, const FName& InAttachmentName);
 	UE_API TOptional<FTransform> GetParentTransform(const UObject* InObject, const FName& InAttachmentName);
-	
+
 	UE_API const TArray< TWeakObjectPtr<UTickableConstraint> >& Get(const uint32 InHandleHash, UWorld* InWorld);
 	UE_API bool HasAnyActiveConstraint(const uint32 InHandleHash, UWorld* InWorld);
 	UE_API TOptional<FTransform> GetParentTransform(const uint32 InHandleHash, UWorld* InWorld);
 
+	UE_API const TArray< TWeakObjectPtr<UTickableConstraint> >& Get(UObject* InTarget, UWorld* InWorld = nullptr);
+	UE_API bool HasAnyActiveConstraint(UObject* InTarget, UWorld* InWorld = nullptr);
+	
 	UE_API void RegisterNotifications(UWorld* InWorld);
 	UE_API void UnregisterNotifications(UWorld* InWorld);
 
 private:
 	
 	TMap<uint32, TArray< TWeakObjectPtr<UTickableConstraint> >> PerHandleConstraints;
+	TMap<TObjectKey<UObject>, TArray< TWeakObjectPtr<UTickableConstraint> >> PerTargetConstraints;
 	FDelegateHandle ConstraintsNotificationHandle;
 };
 
