Move the experiental 'FConsoleCommandPtr' to it's own header file to make it easier to use elsewhere in the module.

#rnx

diff --git a/Engine/Source/Runtime/Experimental/IoStore/OnDemand/Private/CVarUtilities.h b/Engine/Source/Runtime/Experimental/IoStore/OnDemand/Private/CVarUtilities.h
new file mode 100644
--- /dev/null
+++ b/Engine/Source/Runtime/Experimental/IoStore/OnDemand/Private/CVarUtilities.h
@@ -0,0 +1,16 @@
+// Copyright Epic Games, Inc. All Rights Reserved.
+
+#pragma once
+
+#include "Templates/UniquePtr.h"
+
+// The following works in the header, make sure that you include 'HAL/IConsoleManager.h' in the cpp!
+struct IConsoleCommand;
+
+struct FConsoleCommandUnregister
+{
+	void operator()(IConsoleCommand* ConsoleCommandPtr) const;
+};
+
+/** Experimental wrapper around a registered IConsoleCommand */
+using FConsoleCommandPtr = TUniquePtr<IConsoleCommand, FConsoleCommandUnregister>;
diff --git a/Engine/Source/Runtime/Experimental/IoStore/OnDemand/Private/OnDemandHttpThread.h b/Engine/Source/Runtime/Experimental/IoStore/OnDemand/Private/OnDemandHttpThread.h
--- a/Engine/Source/Runtime/Experimental/IoStore/OnDemand/Private/OnDemandHttpThread.h
+++ b/Engine/Source/Runtime/Experimental/IoStore/OnDemand/Private/OnDemandHttpThread.h
@@ -4,6 +4,7 @@
 
 #include "Async/Mutex.h"
 #include "Containers/StringFwd.h"
+#include "CVarUtilities.h"
 #include "Delegates/Delegate.h"
 #include "HAL/Event.h"
 #include "HAL/PlatformAffinity.h"
@@ -19,20 +20,6 @@
 
 class FRunnableThread;
 
-#if UE_ALLOW_HTTP_PAUSE
-
-struct IConsoleCommand;
-
-struct FConsoleCommandUnregister
-{
-	void operator()(IConsoleCommand* ConsoleCommandPtr) const;
-};
-
-/** Experimental wrapper around a registered IConsoleCommand */
-using FConsoleCommandPtr = TUniquePtr<IConsoleCommand, FConsoleCommandUnregister>;
-
-#endif // UE_ALLOW_HTTP_PAUSE
-
 namespace UE::IoStore
 {
 
