Add a new cvar 'iax.ResetOnScreenStatistics' to reset the values seen in the IAX debug display.

#rnx

- Note that once this has been invoked we can no longer trust the values being gathered for analytics and so they will not be sent until the process has been restarted.
- This might be fixed in the future when IAnalyticsRecording is fixed but isn't immediately important as we generally only care about analytics in shipping config which won't have this stuff anyway.

diff --git a/Engine/Source/Runtime/Experimental/IoStore/OnDemand/Private/Statistics.cpp b/Engine/Source/Runtime/Experimental/IoStore/OnDemand/Private/Statistics.cpp
--- a/Engine/Source/Runtime/Experimental/IoStore/OnDemand/Private/Statistics.cpp
+++ b/Engine/Source/Runtime/Experimental/IoStore/OnDemand/Private/Statistics.cpp
@@ -344,7 +344,7 @@ private:
 
 #define UE_IAX_COUNTER(Name, Type) Type G##Name[(uint8)EHttpRequestType::NUM_SOURCES] = {{TEXT(UE_STRINGIZE(Ias/Name)), TraceCounterDisplayHint_None}, {TEXT(UE_STRINGIZE(Iad/Name)), TraceCounterDisplayHint_None}};
 #define UE_IAX_MEMORY_COUNTER(Name, Type) Type G##Name[(uint8)EHttpRequestType::NUM_SOURCES] = {{TEXT(UE_STRINGIZE(Ias/Name)), TraceCounterDisplayHint_None}, {TEXT(UE_STRINGIZE(Iad/Name)), TraceCounterDisplayHint_Memory}};
-	
+
 // iorequest stats
 FCounterInt				GIoRequestCount(TEXT("Ias/IoRequestCount"), TraceCounterDisplayHint_None);
 FCounterAtomicInt		GIoRequestReadCount(TEXT("Ias/IoRequestReadCount"), TraceCounterDisplayHint_None);
@@ -588,6 +588,15 @@ FOnDemandIoBackendStats::FOnDemandIoBackendStats(FBackendStatus& InStatus)
 
 			PrintOnScreenStatistics(OutMessages);
 		});
+
+	ResetStatisticsCommand = FConsoleCommandPtr(IConsoleManager::Get().RegisterConsoleCommand(
+		TEXT("iax.ResetOnScreenStatistics"),
+		TEXT(""),
+		FConsoleCommandDelegate::CreateLambda([this]() -> void
+			{
+				ResetOnScreenStatistics();
+			}),
+		ECVF_Default));
 #endif // UE_ENABLE_ONSCREEN_STATISTICS
 }
 
@@ -713,6 +722,37 @@ void FOnDemandIoBackendStats::PrintPeriodicLogging()
 	}
 }
 
+void FOnDemandIoBackendStats::ResetOnScreenStatistics()
+{
+#if UE_ENABLE_ONSCREEN_STATISTICS
+
+#define UE_IAX_RESET_COUNTER(Counter) {for(auto& Value:Counter){Value.Set(0);}}	
+#define UE_IAX_RESET_VALUE(Counter) {for(auto& Value:Counter){Value = 0;}}
+
+	UE_IAX_RESET_COUNTER(GHttpDownloadedBytes);
+	UE_IAX_RESET_COUNTER(GHttpGetCount);
+	UE_IAX_RESET_VALUE(GHttpDurationMsAvg);
+	UE_IAX_RESET_COUNTER(GHttpRetryCount);
+
+	GHttpCdnCacheHit = 0;
+	GHttpCdnCacheMiss = 0;
+	GHttpCdnCacheUnknown = 0;
+
+	GCacheRejectBytes.Set(0);
+	GCacheReadBytes.Set(0);
+	GCacheGetCount.Set(0);
+
+	GCacheDecodeErrorCount.Set(0);
+
+	UE_IAX_RESET_COUNTER(GHttpDecodeErrorCount);
+	UE_IAX_RESET_COUNTER(GHttpErrorCount);
+
+	bValuesValidForAnalytics = false;
+
+#undef UE_IAX_RESET_COUNTER
+
+#endif // UE_ENABLE_ONSCREEN_STATISTICS
+}
 
 void FOnDemandIoBackendStats::PrintOnScreenStatistics(FCoreDelegates::FSeverityMessageMap& OutMessages)
 {
@@ -886,7 +926,7 @@ void FOnDemandIoBackendStats::ReportGeneralAnalytics(TArray<FAnalyticsEventAttri
 
 void FOnDemandIoBackendStats::ReportEndPointAnalytics(TArray<FAnalyticsEventAttribute>& OutAnalyticsArray) const
 {
-	if (GIasReportHttpAnalyticsEnabled)
+	if (GIasReportHttpAnalyticsEnabled && bValuesValidForAnalytics)
 	{
 		auto ReportHttpStats = [&OutAnalyticsArray](const TCHAR* Prefix, EHttpRequestType Type)
 		{
diff --git a/Engine/Source/Runtime/Experimental/IoStore/OnDemand/Private/Statistics.h b/Engine/Source/Runtime/Experimental/IoStore/OnDemand/Private/Statistics.h
--- a/Engine/Source/Runtime/Experimental/IoStore/OnDemand/Private/Statistics.h
+++ b/Engine/Source/Runtime/Experimental/IoStore/OnDemand/Private/Statistics.h
@@ -2,6 +2,7 @@
 
 #pragma once
 
+#include "CVarUtilities.h"
 #include "HAL/LowLevelMemTracker.h"
 #include "IO/OnDemandError.h"
 #include "Misc/CoreDelegates.h"
@@ -119,13 +120,18 @@ private:
 	void UpdateCSVValues();
 	void PrintPeriodicLogging();
 
+	void ResetOnScreenStatistics();
 	void PrintOnScreenStatistics(FCoreDelegates::FSeverityMessageMap& OutMessages);
 
 	FDelegateHandle OnApplicationResumeHandle;
 	FDelegateHandle EndFrameDelegateHandle;
 	FDelegateHandle OnScreenDelegateHandle;
 
+	FConsoleCommandPtr ResetStatisticsCommand;
+
 	const FBackendStatus& BackendStatus;
+
+	bool bValuesValidForAnalytics = true;
 #endif //IAS_WITH_STATISTICS
 };
 
