Perf optimization for partial island sleep:
- Parallelize constraint to particle sync for large islands

diff --git a/Engine/Source/Runtime/Experimental/Chaos/Private/Chaos/Island/IslandManager.cpp b/Engine/Source/Runtime/Experimental/Chaos/Private/Chaos/Island/IslandManager.cpp
--- a/Engine/Source/Runtime/Experimental/Chaos/Private/Chaos/Island/IslandManager.cpp
+++ b/Engine/Source/Runtime/Experimental/Chaos/Private/Chaos/Island/IslandManager.cpp
@@ -8,6 +8,7 @@
 #include "Chaos/PerParticleGravity.h"
 #include "Chaos/Evolution/IterationSettings.h"
 #include "Chaos/PhysicsMaterialUtilities.h"
+#include "Tasks/Task.h"
 
 #include "ChaosStats.h"
 
@@ -58,7 +59,7 @@ namespace Chaos::CVars
 	/** Cvar to enable/disable partial island waking for all islands */
 	bool bChaosSolverPartialIslandWake = true;
 	FAutoConsoleVariableRef CVarChaosSolverPartialIslandWake(TEXT("p.Chaos.Solver.Sleep.PartialIslandSleep.PartialIslandWake"), bChaosSolverPartialIslandWake, TEXT("[EXPERIMENTAL] Safeguard only: New constraints and particle state changes will trigger partial wake events instead of waking up the entire island.[def:true]"));
-	
+
 	/** Cvar to enable/disable forcing a partial island waking for all islands even if all particles are already awake.
 	This is just a safeguard to force resetting the particle sleep counter when a new wake event occurs. It should generally be set to false. */
 	bool bChaosSolverForcePartialIslandWake = false;
@@ -105,7 +106,7 @@ namespace Chaos::CVars
 	/** Cvar to adjust the angular wake threshold for partial waking */
 	FRealSingle AngularPartialWakeThresholdMultiplier = 10.0f;
 	FAutoConsoleVariableRef CVarChaosSolverAngularWakeThresholdMultiplier(TEXT("p.Chaos.Solver.Sleep.PartialIslandSleep.AngularWakeThresholdMultiplier"), AngularPartialWakeThresholdMultiplier, TEXT("[EXPERIMENTAL] A multiplier applied to AngularSleepThreshold for particle waking.[def:10]"));
-	
+
 	/** Cvar to override the sleep counter threshold if necessary */
 	int32 ChaosSolverCollisionDefaultSleepCounterThreshold = 20;
 	FAutoConsoleVariableRef CVarChaosSolverCollisionDefaultSleepCounterThreshold(TEXT("p.Chaos.Solver.Sleep.Defaults.SleepCounterThreshold"), ChaosSolverCollisionDefaultSleepCounterThreshold, TEXT("Default counter threshold for sleeping.[def:20]"));
@@ -3460,11 +3461,7 @@ namespace Chaos::Private
 
 	void FPBDIslandManager::SyncConstraintSleepStateToParticles(FPBDIsland* Island)
 	{
-		// Iterate over the edges
-		// Sync the constraint sleep state to match the connected particles
-		for (TArray<FPBDIslandConstraint*>& IslandEdges : Island->ContainerEdges)
-		{
-			for (FPBDIslandConstraint* IslandEdge : IslandEdges)
+		auto SyncConstraintSleepStateToParticlesLambda = [](FPBDIslandConstraint* IslandEdge)
 			{
 				const FPBDIslandParticle* Node0 = IslandEdge->Nodes[0];
 				const FPBDIslandParticle* Node1 = IslandEdge->Nodes[1];
@@ -3518,8 +3515,50 @@ namespace Chaos::Private
 						Constraint->SetIsSleeping(false);
 					}
 				}
+			};
+
+		TArray<UE::Tasks::FTask> PendingEdgesTask;
+
+		for (TArray<FPBDIslandConstraint*>& IslandEdges : Island->ContainerEdges)
+		{
+			int32 NumEdges = IslandEdges.Num();
+			if (NumEdges > 0)
+			{
+				constexpr int32 MinEdgesPerTask = 100;
+				const int32 NumTasks = FMath::Max(FMath::Min(FTaskGraphInterface::Get().GetNumWorkerThreads() - 1, Chaos::MaxNumWorkers), 1);
+				const int32 NumEdgeTask = FMath::Max(FMath::Min(NumTasks, NumEdges), 1);
+				const int32 EdgesByTask = FMath::Max(FMath::DivideAndRoundUp(NumEdges, NumEdgeTask), MinEdgesPerTask);
+				const int32 NumEdgeBatches = FMath::DivideAndRoundUp(NumEdges, EdgesByTask);
+
+				if (NumEdgeBatches == 1) // in series
+				{
+					for (FPBDIslandConstraint* IslandEdge : IslandEdges)
+					{
+						SyncConstraintSleepStateToParticlesLambda(IslandEdge);
+					}
+				}
+				else // in parallel
+				{
+					for (int32 BatchIndex = 0; BatchIndex < NumEdgeBatches; BatchIndex++)
+					{
+						const int32 StartIndex = BatchIndex * EdgesByTask;
+						int32 EndIndex = (BatchIndex + 1) * EdgesByTask;
+						EndIndex = FMath::Min(NumEdges, EndIndex);
+
+						UE::Tasks::FTask EdgeTask = UE::Tasks::Launch(UE_SOURCE_LOCATION, [StartIndex, EndIndex, &IslandEdges, SyncConstraintSleepStateToParticlesLambda]()
+							{
+								for (int32 Index = StartIndex; Index < EndIndex; ++Index)
+								{
+									FPBDIslandConstraint* IslandEdge = IslandEdges[Index];
+									SyncConstraintSleepStateToParticlesLambda(IslandEdge);
+								}
+							});
+						PendingEdgesTask.Add(EdgeTask);
+					}
+				}
 			}
 		}
+		UE::Tasks::Wait(PendingEdgesTask);
 	}
 
 	void FPBDIslandManager::WakeUpParticles(TArray<FGeometryParticleHandle*>& ParticlesToWake, const bool bResetSleepCounter)
