Add new command 'ias.DropCache' that resets the IAS cache which works for both the old and new backend.

#rnx

- The cvar is compiled out in shipping and is not available. The code api will remain available but it is currently intended for development work.

diff --git a/Engine/Source/Runtime/Experimental/IoStore/OnDemand/Private/IasCache.h b/Engine/Source/Runtime/Experimental/IoStore/OnDemand/Private/IasCache.h
--- a/Engine/Source/Runtime/Experimental/IoStore/OnDemand/Private/IasCache.h
+++ b/Engine/Source/Runtime/Experimental/IoStore/OnDemand/Private/IasCache.h
@@ -12,6 +12,8 @@ class FIoBuffer;
 class FIoReadOptions;
 struct FIoHash;
 
+#define UE_ALLOW_DROP_CACHE !UE_BUILD_SHIPPING
+
 namespace UE::IoStore
 {
 
@@ -27,6 +29,9 @@ public:
 	to cancel and collect any Get() tasks beforehand. */
 	virtual void Abandon() = 0;
 
+	/** Clears out the cache and resets it back to the default state. */
+	virtual void Drop() = 0;
+
 	/** Returns whether the specified cache key is present in the cache. */
 	virtual bool ContainsChunk(const FIoHash& Key) const = 0;
 
diff --git a/Engine/Source/Runtime/Experimental/IoStore/OnDemand/Private/JournaledCache.cpp b/Engine/Source/Runtime/Experimental/IoStore/OnDemand/Private/JournaledCache.cpp
--- a/Engine/Source/Runtime/Experimental/IoStore/OnDemand/Private/JournaledCache.cpp
+++ b/Engine/Source/Runtime/Experimental/IoStore/OnDemand/Private/JournaledCache.cpp
@@ -135,6 +135,7 @@ public:
 	const FIoBuffer*Get(uint64 Key) const;
 	bool			Put(uint64 Key, FIoBuffer&& Data);
 	bool			Evict(uint64 Key, bool bAlways);
+	void			Drop();
 	int32			Peel(int32 PeelThreshold, PeelItems& Out);
 	uint32			DebugVisit(void* Param, FDebugCacheEntry::Callback* Callback);
 
@@ -233,6 +234,15 @@ bool FMemCache::Evict(uint64 Key, bool bAlways)
 	return bEviction;
 }
 
+////////////////////////////////////////////////////////////////////////////////
+void FMemCache::Drop()
+{
+	UE_LOG(LogIas, Display, TEXT("Dropping FMemCache"));
+
+	Items.Empty();
+	UsedSize = 0;
+}
+
 ////////////////////////////////////////////////////////////////////////////////
 int32 FMemCache::Peel(int32 PeelThreshold, PeelItems& Out)
 {
@@ -944,6 +954,8 @@ int32 FDiskCache::Flush()
 ////////////////////////////////////////////////////////////////////////////////
 void FDiskCache::Drop()
 {
+	UE_LOG(LogIas, Display, TEXT("Dropping FDiskCache"));
+
 	DataHandle.Reset();
 
 	IPlatformFile& Ipf = IPlatformFile::GetPlatformPhysical();
@@ -1381,6 +1393,13 @@ bool FCache::Load()
 void FCache::Drop()
 {
 	DiskCache.Drop();
+
+	{
+		FWriteScopeLock _(MemLock);
+		MemCache.Drop();
+	}
+
+	Partial = FPartial();
 }
 
 ////////////////////////////////////////////////////////////////////////////////
@@ -1644,6 +1663,7 @@ public:
 	uint32					ClaimReadId();
 	void					BeginRead(const FCache* Cache, const FReadRequest& Request);
 	void					CancelRead(const FIoBuffer* GivenDest);
+	void					Drop();
 
 private:
 	union FWork
@@ -1654,6 +1674,7 @@ private:
 			Work_Gov,
 			Work_Read,
 			Work_Cancel,
+			Work_Drop,
 		};
 
 		template <typename Type>
@@ -1806,6 +1827,14 @@ void FServiceThread::CancelRead(const FIoBuffer* GivenData)
 	SubmitWork(&Work, 1);
 }
 
+////////////////////////////////////////////////////////////////////////////////
+void FServiceThread::Drop()
+{
+	FWork Work;
+	Work.What = FWork::Work_Drop;
+	SubmitWork(&Work, 1);
+}
+
 ////////////////////////////////////////////////////////////////////////////////
 void FServiceThread::StartThread()
 {
@@ -1939,10 +1968,17 @@ void FServiceThread::ReceiveWork()
 	PendingPrev = PendingLoad;
 
 	// Unregisters first
+	bool bDropCache = false;
+
 	for (uint32 i = 0, n = InboundWork.Num(); i < n; ++i)
 	{
 		const FWork& Work = InboundWork[i];
 
+		if (Work.What == FWork::Work_Drop)
+		{
+			bDropCache = true;
+		}
+
 		if (Work.What != FWork::Work_Unregister)
 		{
 			i += (Work.What == FWork::Work_Read); // Consume parameter work item too.
@@ -1982,7 +2018,7 @@ void FServiceThread::ReceiveWork()
 	{
 		const FWork& Work = InboundWork[i];
 
-		if (Work.What == FWork::Work_Unregister || Work.What == FWork::Work_Cancel)
+		if (Work.What == FWork::Work_Unregister || Work.What == FWork::Work_Cancel || Work.What == FWork::Work_Drop)
 		{
 			continue;
 		}
@@ -2028,6 +2064,21 @@ void FServiceThread::ReceiveWork()
 		}
 	}
 
+	// lastly if there were one ore more requests to drop the cache
+	if (bDropCache)
+	{
+		for (TUniquePtr<FCache>& Cache : Caches)
+		{
+			Cache->Drop();
+		}
+
+		// cancel all active reads
+		for (uint32 j = 1, m = ActiveReads.Num(); j < m; j += 2)
+		{
+			ActiveReads[j].Dest = nullptr;
+		}
+	}
+
 	check((ActiveReads.Num() & 1) == 0);
 }
 
@@ -2087,6 +2138,7 @@ public:
 								~FJournaledCache();
 	bool						Initialize(const TCHAR* RootDir, const FIasCacheConfig& Config);
 	virtual void				Abandon() override;
+	virtual void				Drop() override;
 	virtual bool				ContainsChunk(const FIoHash& Key) const override;
 	virtual EIoErrorCode		Get(const FIoHash& Key, FIoBuffer& OutData) override;
 	virtual void				Materialize(const FIoHash& Key, FIoBuffer& Dest, EIoErrorCode& Status, UE::Tasks::FTaskEvent DoneEvent) override;
@@ -2194,6 +2246,12 @@ void FJournaledCache::Abandon()
 	delete this;
 }
 
+////////////////////////////////////////////////////////////////////////////////
+void FJournaledCache::Drop()
+{
+	JournaledCache::FServiceThread::Get().Drop();
+}
+
 ////////////////////////////////////////////////////////////////////////////////
 bool FJournaledCache::ContainsChunk(const FIoHash& Key) const
 {
diff --git a/Engine/Source/Runtime/Experimental/IoStore/OnDemand/Private/OnDemandHttpIoDispacher.cpp b/Engine/Source/Runtime/Experimental/IoStore/OnDemand/Private/OnDemandHttpIoDispacher.cpp
--- a/Engine/Source/Runtime/Experimental/IoStore/OnDemand/Private/OnDemandHttpIoDispacher.cpp
+++ b/Engine/Source/Runtime/Experimental/IoStore/OnDemand/Private/OnDemandHttpIoDispacher.cpp
@@ -6,17 +6,18 @@
 #include "Async/UniqueLock.h"
 #include "AtomicFlags.h"
 #include "Containers/AnsiString.h"
+#include "CVarUtilities.h"
 #include "Experimental/ConcurrentLinearAllocator.h"
 #include "HAL/Platform.h"
 #include "HAL/Runnable.h"
 #include "HAL/RunnableThread.h"
-#include "IasCache.h"
-#include "IasHostGroup.h"
+#include "IO/HttpIoDispatcher.h"
 #include "IO/IoContainers.h"
 #include "IO/IoStoreOnDemand.h"
-#include "IO/HttpIoDispatcher.h"
-#include "Logging/StructuredLog.h"
+#include "IasCache.h"
+#include "IasHostGroup.h"
 #include "Logging/LogVerbosity.h"
+#include "Logging/StructuredLog.h"
 #include "Misc/CommandLine.h"
 #include "Misc/EnumClassFlags.h"
 #include "Misc/Fork.h"
@@ -710,13 +711,18 @@ private:
 	FEventRef								WakeUp;
 	/** Keeps assignment of the CancelationToken safe between ::CancelRequest and ::ProcessHttpRequests */
 	FMutex									CancelationMutex;
+	FDelegateHandle							OnRecreateHttpClientHandle;
+	std::atomic_bool						bRecreateHttpClient = false;
+	std::atomic_bool						bStopRequested = false;
+
 #if UE_ALLOW_HTTP_PAUSE
 	FConsoleCommandPtr						PauseCommand;
 	FConsoleCommandPtr						UnpauseCommand;
 #endif // UE_ALLOW_HTTP_PAUSE
-	FDelegateHandle							OnRecreateHttpClientHandle;
-	std::atomic_bool						bRecreateHttpClient = false;
-	std::atomic_bool						bStopRequested = false;
+
+#if UE_ALLOW_DROP_CACHE
+	FConsoleCommandPtr						DropCacheCommand;
+#endif // UE_ALLOW_DROP_CACHE
 };
 
 static TUniquePtr<FHttpDispatcher> HttpDispatcher;
@@ -745,6 +751,20 @@ FHttpDispatcher::FHttpDispatcher(TUniquePtr<IIasCache>&& InCache)
 		ECVF_Default));
 #endif // UE_ALLOW_HTTP_PAUSE
 
+#if UE_ALLOW_DROP_CACHE
+	DropCacheCommand = FConsoleCommandPtr(IConsoleManager::Get().RegisterConsoleCommand(
+		TEXT("ias.DropCache"),
+		TEXT("Resets the IAS cache and deletes the data from disk."),
+		FConsoleCommandDelegate::CreateLambda([this]() -> void
+			{
+				if (Cache)
+				{
+					Cache->Drop();
+				}
+			}),
+		ECVF_Default));
+#endif // UE_ALLOW_DROP_CACHE
+
 	const int32 ThreadPriorityIndex			= FMath::Clamp(GOnDemandBackendThreadPriorityIndex, 0, (int32)UE_ARRAY_COUNT(GThreadPriorities) - 1);
 	EThreadPriority DesiredThreadPriority	= (EThreadPriority)GThreadPriorities[ThreadPriorityIndex];
 	ThreadPriority							= DesiredThreadPriority;
diff --git a/Engine/Source/Runtime/Experimental/IoStore/OnDemand/Private/OnDemandIoDispatcherBackend.cpp b/Engine/Source/Runtime/Experimental/IoStore/OnDemand/Private/OnDemandIoDispatcherBackend.cpp
--- a/Engine/Source/Runtime/Experimental/IoStore/OnDemand/Private/OnDemandIoDispatcherBackend.cpp
+++ b/Engine/Source/Runtime/Experimental/IoStore/OnDemand/Private/OnDemandIoDispatcherBackend.cpp
@@ -4,6 +4,7 @@
 
 #include "Containers/BitArray.h"
 #include "Containers/StringView.h"
+#include "CVarUtilities.h"
 #include "DistributionEndpoints.h"
 #include "HAL/Event.h"
 #include "HAL/LowLevelMemTracker.h"
@@ -704,6 +705,10 @@ private:
 	UE::Tasks::FTask HostGroupTask;
 	std::atomic_uint32_t InflightCacheRequestCount{0};
 	std::atomic_bool bStopRequested{false};
+
+#if UE_ALLOW_DROP_CACHE
+	FConsoleCommandPtr DropCacheCommand;
+#endif // UE_ALLOW_DROP_CACHE
 };
 
 ///////////////////////////////////////////////////////////////////////////////
@@ -729,6 +734,20 @@ FOnDemandIoBackend::FOnDemandIoBackend(const FOnDemandEndpointConfig& Config, FO
 	// Don't enable HTTP until the background thread has been started
 	BackendStatus.SetHttpEnabled(false);
 	BackendStatus.SetCacheEnabled(Cache.IsValid());
+
+#if UE_ALLOW_DROP_CACHE
+	DropCacheCommand = FConsoleCommandPtr(IConsoleManager::Get().RegisterConsoleCommand(
+		TEXT("ias.DropCache"),
+		TEXT("Resets the IAS cache and deletes the data from disk."),
+		FConsoleCommandDelegate::CreateLambda([this]() -> void
+			{
+				if (Cache)
+				{
+					Cache->Drop();
+				}
+			}),
+		ECVF_Default));
+#endif // UE_ALLOW_DROP_CACHE
 }
 
 FOnDemandIoBackend::~FOnDemandIoBackend()
