Changed cvar iax.EnableHttpIoDispatcher to use ECVF_SaveForNextBoot and added screen stats to display which ever is in use.

#rnx

diff --git a/Engine/Source/Runtime/Experimental/IoStore/OnDemand/Private/OnDemandIoStore.cpp b/Engine/Source/Runtime/Experimental/IoStore/OnDemand/Private/OnDemandIoStore.cpp
--- a/Engine/Source/Runtime/Experimental/IoStore/OnDemand/Private/OnDemandIoStore.cpp
+++ b/Engine/Source/Runtime/Experimental/IoStore/OnDemand/Private/OnDemandIoStore.cpp
@@ -58,7 +58,7 @@ static FAutoConsoleVariableRef CVar_HttpIoDispatcherEnabled(
 	TEXT("iax.HttpIoDispatcherEnabled"),
 	GHttpIoDispatcherEnabled,
 	TEXT("Whether the HTTP I/O dispatcher should be enabled."),
-	ECVF_ReadOnly
+	ECVF_SaveForNextBoot
 );
 
 ///////////////////////////////////////////////////////////////////////////////
@@ -345,6 +345,8 @@ FIoStatus FOnDemandIoStore::Initialize()
 
 	const bool bHttpIoDispatcherEnabled = (GHttpIoDispatcherEnabled || FParse::Param(FCommandLine::Get(), TEXT("HttpIoDispatcher")))
 		&& !FParse::Param(FCommandLine::Get(), TEXT("NoHttpIoDispatcher"));
+	FOnDemandIoBackendStats::Get()->SetHttpIoDispatcherEnabled(bHttpIoDispatcherEnabled);
+
 	if (bHttpIoDispatcherEnabled)
 	{
 		TUniquePtr<IIasCache> HttpCache;
diff --git a/Engine/Source/Runtime/Experimental/IoStore/OnDemand/Private/Statistics.cpp b/Engine/Source/Runtime/Experimental/IoStore/OnDemand/Private/Statistics.cpp
--- a/Engine/Source/Runtime/Experimental/IoStore/OnDemand/Private/Statistics.cpp
+++ b/Engine/Source/Runtime/Experimental/IoStore/OnDemand/Private/Statistics.cpp
@@ -434,6 +434,8 @@ struct FHttpRecentHistoryStatistics
 
 FHttpRecentHistoryStatistics GHttpHistory[(uint8)EHttpRequestType::NUM_SOURCES];
 
+// Debug flag indicating which HTTP backend is running
+static bool				GHttpIoDispatcherEnabled = false;
 // Experimental Http Stats
 static uint32			GHttpCdnCacheHit = 0;
 static uint32			GHttpCdnCacheMiss = 0;
@@ -762,6 +764,7 @@ void FOnDemandIoBackendStats::PrintOnScreenStatistics(FCoreDelegates::FSeverityM
 		TStringBuilder<256> BackendStatusText;
 		BackendStatusText << TEXT("IAS Backend Status: ");
 		BackendStatus.ToString(BackendStatusText);
+		BackendStatusText.Appendf(TEXT(" | HTTP I/O dispatcher: %s"), GHttpIoDispatcherEnabled ? TEXT("Enabled") : TEXT("Disabled"));
 
 		OutMessages.Add(FCoreDelegates::EOnScreenMessageSeverity::Info, FText::FromStringView(BackendStatusText.ToView()));
 	}
@@ -957,6 +960,7 @@ void FOnDemandIoBackendStats::ReportEndPointAnalytics(TArray<FAnalyticsEventAttr
 
 		ReportHttpStats(TEXT("Ias"), EHttpRequestType::Streaming);
 		ReportHttpStats(TEXT("Iad"), EHttpRequestType::Installed);
+		AppendAnalyticsEventAttributeArray(OutAnalyticsArray, TEXT("HttpIoDispatcherEnabled"), GHttpIoDispatcherEnabled);
 	}
 
 	if (GIasReportCacheAnalyticsEnabled)
@@ -1366,6 +1370,11 @@ void FOnDemandIoBackendStats::OnApplicationResume()
 	++GAppResumeCount;
 }
 
+void FOnDemandIoBackendStats::SetHttpIoDispatcherEnabled(bool bEnabled)
+{
+	GHttpIoDispatcherEnabled = bEnabled;
+}
+
 ////////////////////////////////////////////////////////////////////////////////
 void FOnDemandContentInstallerStats::OnRequestEnqueued()
 {
diff --git a/Engine/Source/Runtime/Experimental/IoStore/OnDemand/Private/Statistics.h b/Engine/Source/Runtime/Experimental/IoStore/OnDemand/Private/Statistics.h
--- a/Engine/Source/Runtime/Experimental/IoStore/OnDemand/Private/Statistics.h
+++ b/Engine/Source/Runtime/Experimental/IoStore/OnDemand/Private/Statistics.h
@@ -112,7 +112,7 @@ public:
 	void OnHttpRequestStarted() IAS_STATISTICS_IMPL()
 	void OnHttpRequestCompleted() IAS_STATISTICS_IMPL()
 
-
+	void SetHttpIoDispatcherEnabled(bool bEnabled) IAS_STATISTICS_IMPL()
 private:
 
 #if IAS_WITH_STATISTICS
