Disable additional validation when running the 'ChunkPlugin' command

#rnx

- IASTool is not currently able to decode the chunked data for all platforms, depending on the compression formats used. This should be fixed in the future but for now it's easiest to just disable the checks.

diff --git a/Engine/Source/Runtime/Experimental/IoStore/OnDemand/Private/Tool/ChunkPlugin.cpp b/Engine/Source/Runtime/Experimental/IoStore/OnDemand/Private/Tool/ChunkPlugin.cpp
--- a/Engine/Source/Runtime/Experimental/IoStore/OnDemand/Private/Tool/ChunkPlugin.cpp
+++ b/Engine/Source/Runtime/Experimental/IoStore/OnDemand/Private/Tool/ChunkPlugin.cpp
@@ -398,6 +398,7 @@ static int32 ChunkContainer(const FString& ContainerUTocPath, const FChunkPlugin
 		}
 
 		// Additional validation to make sure that the chunks can be decrypted and decoded
+#if 0
 		{
 			FIoChunkDecodingParams DecodingParams;
 			DecodingParams.CompressionFormat = CompressionFormat;
@@ -424,6 +425,7 @@ static int32 ChunkContainer(const FString& ContainerUTocPath, const FChunkPlugin
 				return -1;
 			}
 		}
+#endif
 
 		FIoStatus WriteStatus = ChunkWriter->WriteChunk(ChunksRelativeFolder, ReadResult.IoBuffer, ChunkHash);
 		if (WriteStatus.IsOk() == false)
