Implement ResetRestLengthsWithMorphTarget for all eligible cloth constraint types.

This was originally only implemented for a subset of constraints to support an internal project.

diff --git a/Engine/Plugins/ChaosCloth/Source/ChaosCloth/Private/ChaosCloth/ChaosClothConstraints.cpp b/Engine/Plugins/ChaosCloth/Source/ChaosCloth/Private/ChaosCloth/ChaosClothConstraints.cpp
--- a/Engine/Plugins/ChaosCloth/Source/ChaosCloth/Private/ChaosCloth/ChaosClothConstraints.cpp
+++ b/Engine/Plugins/ChaosCloth/Source/ChaosCloth/Private/ChaosCloth/ChaosClothConstraints.cpp
@@ -2638,14 +2638,49 @@ void FClothConstraints::ResetRestLengths(
 	const TMap<FString, TConstArrayView<FRealSingle>>& WeightMaps
 )
 {
+	if (XStretchBiasConstraints)
+	{ 
+		// TODO. There is no way to create these from the cloth UI
+	}
+	if (XEdgeConstraints)
+	{
+		XEdgeConstraints->ResetRestLengths(NewRestLengthPositions);
+	}
 	if (EdgeConstraints)
 	{
 		EdgeConstraints->ResetRestLengths(NewRestLengthPositions);
 	}
-
+	if (XAnisoSpringConstraints)
+	{
+		XAnisoSpringConstraints->ResetRestLengths(NewRestLengthPositions);
+	}
+	if (XBendingConstraints)
+	{
+		XBendingConstraints->ResetRestLengths(NewRestLengthPositions);
+	}
+	if (BendingConstraints)
+	{
+		BendingConstraints->ResetRestLengths(NewRestLengthPositions);
+	}
 	if (BendingElementConstraints)
 	{
 		BendingElementConstraints->ResetRestLengths(NewRestLengthPositions, ConfigProperties, WeightMaps);
 	}
+	if (XBendingElementConstraints)
+	{
+		XBendingElementConstraints->ResetRestLengths(NewRestLengthPositions, ConfigProperties, WeightMaps);
+	}
+	if (XAnisoBendingElementConstraints)
+	{
+		XAnisoBendingElementConstraints->ResetRestLengths(NewRestLengthPositions, ConfigProperties, WeightMaps);
+	}
+	if (XAreaConstraints)
+	{
+		XAreaConstraints->ResetRestLengths(NewRestLengthPositions);
+	}
+	if (AreaConstraints)
+	{
+		AreaConstraints->ResetRestLengths(NewRestLengthPositions);
+	}
 }
 }  // End namespace Chaos
diff --git a/Engine/Source/Runtime/Experimental/Chaos/Private/Chaos/PBDSpringConstraints.cpp b/Engine/Source/Runtime/Experimental/Chaos/Private/Chaos/PBDSpringConstraints.cpp
--- a/Engine/Source/Runtime/Experimental/Chaos/Private/Chaos/PBDSpringConstraints.cpp
+++ b/Engine/Source/Runtime/Experimental/Chaos/Private/Chaos/PBDSpringConstraints.cpp
@@ -374,12 +374,9 @@ void FPBDEdgeSpringConstraints::ApplyProperties(const FSolverReal Dt, const int3
 void FPBDEdgeSpringConstraints::ResetRestLengths(const TConstArrayView<FSolverVec3>& Positions)
 {
 	TRACE_CPUPROFILER_EVENT_SCOPE(FPBDSpringConstraints_ResetRestLengths);
-	
-	TArray<FSolverReal>& DistsToCompute = bWarpWeftScaleEnabled ? BaseDists : Dists;
-	for (int32 Index = 0; Index < Constraints.Num(); ++Index)
-	{
-		DistsToCompute[Index] = (Positions[Constraints[Index][0]] - Positions[Constraints[Index][1]]).Size();
-	}
+
+	TArrayView<FSolverReal> DistsToCompute = bWarpWeftScaleEnabled ? TArrayView<FSolverReal>(BaseDists) : TArrayView<FSolverReal>(Dists);
+	CalculateRestLengths(Positions, DistsToCompute);
 
 	UpdateDists();
 }
diff --git a/Engine/Source/Runtime/Experimental/Chaos/Public/Chaos/PBDAxialSpringConstraintsBase.h b/Engine/Source/Runtime/Experimental/Chaos/Public/Chaos/PBDAxialSpringConstraintsBase.h
--- a/Engine/Source/Runtime/Experimental/Chaos/Public/Chaos/PBDAxialSpringConstraintsBase.h
+++ b/Engine/Source/Runtime/Experimental/Chaos/Public/Chaos/PBDAxialSpringConstraintsBase.h
@@ -81,7 +81,27 @@ public:
 	const TArray<TVec3<int32>>& GetConstraints() const { return Constraints; }
 	const TArray<FSolverReal>& GetBarys() const { return Barys; }
 
+	void ResetRestLengths(const TConstArrayView<FSolverVec3>& Positions)
+	{
+		CalculateRestLengths(Positions, TArrayView<FSolverReal>(Dists));
+	}
+
 protected:
+
+	void CalculateRestLengths(const TConstArrayView<FSolverVec3>& Positions, const TArrayView<FSolverReal>& DistsToCompute) const
+	{
+		check(DistsToCompute.Num() == Constraints.Num());
+		for (int32 Index = 0; Index < Constraints.Num(); ++Index)
+		{
+
+			const FSolverVec3& P1 = Positions[Constraints[Index][0]];
+			const FSolverVec3& P2 = Positions[Constraints[Index][1]];
+			const FSolverVec3& P3 = Positions[Constraints[Index][2]];
+			const FSolverVec3 P = (P2 - P3) * Barys[Index] + P3;
+			DistsToCompute[Index] = (P1 - P).Size();
+		}
+	}
+
 	template<typename VectorType>
 	inline VectorType GetSpringDiff(const VectorType& P1, const VectorType& P2, const VectorType& P3, const FSolverReal Bary) const
 	{
diff --git a/Engine/Source/Runtime/Experimental/Chaos/Public/Chaos/PBDSpringConstraintsBase.h b/Engine/Source/Runtime/Experimental/Chaos/Public/Chaos/PBDSpringConstraintsBase.h
--- a/Engine/Source/Runtime/Experimental/Chaos/Public/Chaos/PBDSpringConstraintsBase.h
+++ b/Engine/Source/Runtime/Experimental/Chaos/Public/Chaos/PBDSpringConstraintsBase.h
@@ -42,13 +42,8 @@ public:
 			MaxStiffness)
 	{
 		// Update distances
-		Dists.Reset(Constraints.Num());
-		for (const TVec2<int32>& Constraint : Constraints)
-		{
-			const FSolverVec3& P0 = Particles.X(Constraint[0]);
-			const FSolverVec3& P1 = Particles.X(Constraint[1]);
-			Dists.Add((P1 - P0).Size());
-		}
+		Dists.SetNumUninitialized(Constraints.Num());
+		CalculateRestLengths(Particles.XArray(), Dists);
 	}
 
 	template<int32 Valence UE_REQUIRES(Valence >= 2 && Valence <= 4)>
@@ -79,13 +74,8 @@ public:
 			MaxStiffness)
 	{
 		// Update distances
-		Dists.Reset(Constraints.Num());
-		for (const TVec2<int32>& Constraint : Constraints)
-		{
-			const FSolverVec3& P0 = Particles.GetX(Constraint[0]);
-			const FSolverVec3& P1 = Particles.GetX(Constraint[1]);
-			Dists.Add((P1 - P0).Size());
-		}
+		Dists.SetNumUninitialized(Constraints.Num());
+		CalculateRestLengths(Particles.XArray(), Dists);
 	}
 
 	virtual ~FPBDSpringConstraintsBase()
@@ -99,6 +89,11 @@ public:
 
 	const TArray<TVec2<int32>>& GetConstraints() const { return Constraints; }
 
+	void ResetRestLengths(const TConstArrayView<FSolverVec3>& Positions)
+	{
+		CalculateRestLengths(Positions, TArrayView<FSolverReal>(Dists));
+	}
+
 protected:
 	template<typename SolverParticlesOrRange>
 	FSolverVec3 GetDelta(const SolverParticlesOrRange& Particles, const int32 ConstraintIndex, const FSolverReal ExpStiffnessValue) const
@@ -122,6 +117,15 @@ protected:
 		return ExpStiffnessValue * Delta / CombinedMass;
 	}
 
+	void CalculateRestLengths(const TConstArrayView<FSolverVec3>& Positions, const TArrayView<FSolverReal>& DistsToCompute) const
+	{
+		check(DistsToCompute.Num() == Constraints.Num());
+		for (int32 Index = 0; Index < Constraints.Num(); ++Index)
+		{
+			DistsToCompute[Index] = (Positions[Constraints[Index][0]] - Positions[Constraints[Index][1]]).Size();
+		}
+	}
+
 private:
 	template<int32 Valence, typename Predicate UE_REQUIRES(Valence >= 2 && Valence <= 4)>
 	TArray<TVector<int32, 2>> TrimConstraints(const TArray<TVector<int32, Valence>>& InConstraints, Predicate TrimPredicate)
diff --git a/Engine/Source/Runtime/Experimental/Chaos/Public/Chaos/XPBDAnisotropicBendingConstraints.h b/Engine/Source/Runtime/Experimental/Chaos/Public/Chaos/XPBDAnisotropicBendingConstraints.h
--- a/Engine/Source/Runtime/Experimental/Chaos/Public/Chaos/XPBDAnisotropicBendingConstraints.h
+++ b/Engine/Source/Runtime/Experimental/Chaos/Public/Chaos/XPBDAnisotropicBendingConstraints.h
@@ -36,16 +36,6 @@ public:
 		const TMap<FString, TConstArrayView<FRealSingle>>& WeightMaps,
 		const FCollectionPropertyConstFacade& PropertyCollection);
 
-	UE_DEPRECATED(5.4, "XPBD Constraints must always trim kinematic constraints")
-	FXPBDAnisotropicBendingConstraints(const FSolverParticlesRange& InParticles,
-		const FTriangleMesh& TriangleMesh,
-		const TArray<TVec3<FVec2f>>& FaceVertexPatternPositions,
-		const TMap<FString, TConstArrayView<FRealSingle>>& WeightMaps,
-		const FCollectionPropertyConstFacade& PropertyCollection,
-		bool bTrimKinematicConstraints)
-		: FXPBDAnisotropicBendingConstraints(InParticles, TriangleMesh, FaceVertexPatternPositions, WeightMaps, PropertyCollection)
-	{}
-
 	CHAOS_API FXPBDAnisotropicBendingConstraints(const FSolverParticles& InParticles,
 		int32 InParticleOffset,
 		int32 InParticleCount,
@@ -54,18 +44,6 @@ public:
 		const TMap<FString, TConstArrayView<FRealSingle>>& WeightMaps,
 		const FCollectionPropertyConstFacade& PropertyCollection);
 
-	UE_DEPRECATED(5.4, "XPBD Constraints must always trim kinematic constraints")
-	FXPBDAnisotropicBendingConstraints(const FSolverParticles& InParticles,
-		int32 InParticleOffset,
-		int32 InParticleCount,
-		const FTriangleMesh& TriangleMesh,
-		const TArray<TVec3<FVec2f>>& FaceVertexPatternPositions,
-		const TMap<FString, TConstArrayView<FRealSingle>>& WeightMaps,
-		const FCollectionPropertyConstFacade& PropertyCollection,
-		bool bTrimKinematicConstraints)
-		: FXPBDAnisotropicBendingConstraints(InParticles, InParticleOffset, InParticleCount, TriangleMesh, FaceVertexPatternPositions, WeightMaps, PropertyCollection)
-	{}
-
 	CHAOS_API FXPBDAnisotropicBendingConstraints(const FSolverParticles & InParticles,
 		int32 InParticleOffset,
 		int32 InParticleCount,
@@ -91,103 +69,6 @@ public:
 		const FSolverVec2& RestAngleValue,
 		ERestAngleConstructionType RestAngleConstructionType);
 
-	UE_DEPRECATED(5.5, "Use constructor with BucklingRatioMultipliers")
-	FXPBDAnisotropicBendingConstraints(const FSolverParticles& InParticles,
-		int32 InParticleOffset,
-		int32 InParticleCount,
-		const FTriangleMesh& TriangleMesh,
-		const TArray<TVec3<FVec2f>>& FaceVertexPatternPositions,
-		const TConstArrayView<FRealSingle>& StiffnessWarpMultipliers,
-		const TConstArrayView<FRealSingle>& StiffnessWeftMultipliers,
-		const TConstArrayView<FRealSingle>& StiffnessBiasMultipliers,
-		const TConstArrayView<FRealSingle>& BucklingStiffnessWarpMultipliers,
-		const TConstArrayView<FRealSingle>& BucklingStiffnessWeftMultipliers,
-		const TConstArrayView<FRealSingle>& BucklingStiffnessBiasMultipliers,
-		const TConstArrayView<FRealSingle>& DampingMultipliers,
-		const FSolverVec2& InStiffnessWarp,
-		const FSolverVec2& InStiffnessWeft,
-		const FSolverVec2& InStiffnessBias,
-		const FSolverReal InBucklingRatio,
-		const FSolverVec2& InBucklingStiffnessWarp,
-		const FSolverVec2& InBucklingStiffnessWeft,
-		const FSolverVec2& InBucklingStiffnessBias,
-		const FSolverVec2& InDampingRatio)
-		:FXPBDAnisotropicBendingConstraints(
-			InParticles,
-			InParticleOffset,
-			InParticleCount,
-			TriangleMesh,
-			FaceVertexPatternPositions,
-			StiffnessWarpMultipliers,
-			StiffnessWeftMultipliers,
-			StiffnessBiasMultipliers,
-			TConstArrayView<FRealSingle>(),
-			BucklingStiffnessWarpMultipliers,
-			BucklingStiffnessWeftMultipliers,
-			BucklingStiffnessBiasMultipliers,
-			DampingMultipliers,
-			TConstArrayView<FRealSingle>(),
-			InStiffnessWarp,
-			InStiffnessWeft,
-			InStiffnessBias,
-			FSolverVec2(InBucklingRatio),
-			InBucklingStiffnessWarp,
-			InBucklingStiffnessWeft,
-			InBucklingStiffnessBias,
-			InDampingRatio,
-			FSolverVec2((FSolverReal)0.f),
-			ERestAngleConstructionType::Use3DRestAngles)
-	{}
-
-	UE_DEPRECATED(5.4, "XPBD Constraints must always trim kinematic constraints")
-	FXPBDAnisotropicBendingConstraints(const FSolverParticles& InParticles,
-		int32 InParticleOffset,
-		int32 InParticleCount,
-		const FTriangleMesh& TriangleMesh,
-		const TArray<TVec3<FVec2f>>& FaceVertexPatternPositions,
-		const TConstArrayView<FRealSingle>& StiffnessWarpMultipliers,
-		const TConstArrayView<FRealSingle>& StiffnessWeftMultipliers,
-		const TConstArrayView<FRealSingle>& StiffnessBiasMultipliers,
-		const TConstArrayView<FRealSingle>& BucklingStiffnessWarpMultipliers,
-		const TConstArrayView<FRealSingle>& BucklingStiffnessWeftMultipliers,
-		const TConstArrayView<FRealSingle>& BucklingStiffnessBiasMultipliers,
-		const TConstArrayView<FRealSingle>& DampingMultipliers,
-		const FSolverVec2& InStiffnessWarp,
-		const FSolverVec2& InStiffnessWeft,
-		const FSolverVec2& InStiffnessBias,
-		const FSolverReal InBucklingRatio,
-		const FSolverVec2& InBucklingStiffnessWarp,
-		const FSolverVec2& InBucklingStiffnessWeft,
-		const FSolverVec2& InBucklingStiffnessBias,
-		const FSolverVec2& InDampingRatio,
-		bool bTrimKinematicConstraints)
-		: FXPBDAnisotropicBendingConstraints(
-			InParticles,
-			InParticleOffset,
-			InParticleCount,
-			TriangleMesh,
-			FaceVertexPatternPositions,
-			StiffnessWarpMultipliers,
-			StiffnessWeftMultipliers,
-			StiffnessBiasMultipliers,
-			TConstArrayView<FRealSingle>(),
-			BucklingStiffnessWarpMultipliers,
-			BucklingStiffnessWeftMultipliers,
-			BucklingStiffnessBiasMultipliers,
-			DampingMultipliers,
-			TConstArrayView<FRealSingle>(),
-			InStiffnessWarp,
-			InStiffnessWeft,
-			InStiffnessBias,
-			FSolverVec2(InBucklingRatio),
-			InBucklingStiffnessWarp,
-			InBucklingStiffnessWeft,
-			InBucklingStiffnessBias,
-			InDampingRatio,
-			FSolverVec2((FSolverReal)0.f),
-			ERestAngleConstructionType::Use3DRestAngles)
-	{}
-
 	virtual ~FXPBDAnisotropicBendingConstraints() override {}
 
 	template<typename SolverParticlesOrRange>
@@ -227,7 +108,15 @@ public:
 
 	CHAOS_API void AddInternalForceDifferential(const FSolverParticles& InParticles, const TArray<TVector<FSolverReal, 3>>& DeltaParticles, TArray<TVector<FSolverReal, 3>>& ndf);
 
-
+	void ResetRestLengths(const TConstArrayView<FSolverVec3>& Positions,
+		const FCollectionPropertyConstFacade& PropertyCollection,
+		const TMap<FString, TConstArrayView<FRealSingle>>& WeightMaps)
+	{
+		CalculateRestAngles(Positions, ParticleOffset, ParticleCount,
+			GetRestAngleMapFromCollection(WeightMaps, PropertyCollection),
+			GetRestAngleValueFromCollection(PropertyCollection),
+			(ERestAngleConstructionType)GetXPBDAnisoRestAngleType(PropertyCollection, (int32)ERestAngleConstructionType::Use3DRestAngles));
+	}
 private:
 	template<typename SolverParticlesOrRange>
 	void InitColor(const SolverParticlesOrRange& InParticles);
diff --git a/Engine/Source/Runtime/Experimental/Chaos/Public/Chaos/XPBDAnisotropicSpringConstraints.h b/Engine/Source/Runtime/Experimental/Chaos/Public/Chaos/XPBDAnisotropicSpringConstraints.h
--- a/Engine/Source/Runtime/Experimental/Chaos/Public/Chaos/XPBDAnisotropicSpringConstraints.h
+++ b/Engine/Source/Runtime/Experimental/Chaos/Public/Chaos/XPBDAnisotropicSpringConstraints.h
@@ -146,6 +146,12 @@ public:
 	const TArray<int32>& GetConstraintsPerColorStartIndex() const { return ConstraintsPerColorStartIndex; }
 	const TArray<FSolverVec3>& GetWarpWeftBiasBaseMultipliers() const { return WarpWeftBiasBaseMultipliers; }
 
+	void ResetRestLengths(const TConstArrayView<FSolverVec3>& Positions)
+	{
+		CalculateRestLengths(Positions, TArrayView<FSolverReal>(BaseDists));
+		UpdateDists();
+	}
+
 private:
 	template<typename SolverParticlesOrRange>
 	void InitColor(const SolverParticlesOrRange& InParticles);
@@ -308,6 +314,12 @@ public:
 	const TArray<int32>& GetConstraintsPerColorStartIndex() const { return ConstraintsPerColorStartIndex; }
 	const TArray<FSolverVec3>& GetWarpWeftBiasBaseMultipliers() const { return WarpWeftBiasBaseMultipliers; }
 
+	void ResetRestLengths(const TConstArrayView<FSolverVec3>& Positions)
+	{
+		CalculateRestLengths(Positions, TArrayView<FSolverReal>(BaseDists));
+		UpdateDists();
+	}
+
 private:
 	template<typename SolverParticlesOrRange>
 	void InitColor(const SolverParticlesOrRange& InParticles);
@@ -499,6 +511,12 @@ public:
 	const FXPBDAnisotropicEdgeSpringConstraints& GetEdgeConstraints() const { return EdgeConstraints; }
 	const FXPBDAnisotropicAxialSpringConstraints& GetAxialConstraints() const { return AxialConstraints; }
 
+	void ResetRestLengths(const TConstArrayView<FSolverVec3>& Positions)
+	{
+		EdgeConstraints.ResetRestLengths(Positions);
+		AxialConstraints.ResetRestLengths(Positions);
+	}
+
 private:
 	FXPBDAnisotropicEdgeSpringConstraints EdgeConstraints;
 	FXPBDAnisotropicAxialSpringConstraints AxialConstraints;
diff --git a/Engine/Source/Runtime/Experimental/Chaos/Public/Chaos/XPBDBendingConstraints.h b/Engine/Source/Runtime/Experimental/Chaos/Public/Chaos/XPBDBendingConstraints.h
--- a/Engine/Source/Runtime/Experimental/Chaos/Public/Chaos/XPBDBendingConstraints.h
+++ b/Engine/Source/Runtime/Experimental/Chaos/Public/Chaos/XPBDBendingConstraints.h
@@ -324,6 +324,16 @@ public:
 	CHAOS_API void AddInternalForceDifferential(const FSolverParticles& InParticles, const TArray<TVector<FSolverReal, 3>>& DeltaParticles, TArray<TVector<FSolverReal, 3>>& ndf);
 
 	CHAOS_API FSolverReal ComputeTotalEnergy(const FSolverParticles& InParticles, const FSolverReal ExplicitStiffness = -1.f);
+
+	void ResetRestLengths(const TConstArrayView<FSolverVec3>& Positions,
+		const FCollectionPropertyConstFacade& PropertyCollection,
+		const TMap<FString, TConstArrayView<FRealSingle>>& WeightMaps)
+	{
+		CalculateRestAngles(Positions, ParticleOffset, ParticleCount,
+			GetRestAngleMapFromCollection(WeightMaps, PropertyCollection),
+			GetRestAngleValueFromCollection(PropertyCollection),
+			(ERestAngleConstructionType)GetXPBDRestAngleType(PropertyCollection, (int32)ERestAngleConstructionType::Use3DRestAngles));
+	}
 private:
 	template<typename SolverParticlesOrRange>
 	CHAOS_API void InitColor(const SolverParticlesOrRange& InParticles);
