Geometry collection : add extra check for transform index when pulling the data from physics thread to the game thread

diff --git a/Engine/Source/Runtime/Experimental/Chaos/Private/PhysicsProxy/GeometryCollectionPhysicsProxy.cpp b/Engine/Source/Runtime/Experimental/Chaos/Private/PhysicsProxy/GeometryCollectionPhysicsProxy.cpp
--- a/Engine/Source/Runtime/Experimental/Chaos/Private/PhysicsProxy/GeometryCollectionPhysicsProxy.cpp
+++ b/Engine/Source/Runtime/Experimental/Chaos/Private/PhysicsProxy/GeometryCollectionPhysicsProxy.cpp
@@ -4269,6 +4269,14 @@ bool FGeometryCollectionPhysicsProxy::PullNonInterpolatableDataFromSinglePhysics
 		const FGeometryCollectionResults::FStateData& StateData = CurrentResults.GetState(EntryIndex);
 		const int32 TransformGroupIndex = StateData.TransformIndex;
 
+		// make sure we are not crashing if the transform index is invalid
+		const bool bIsValidTransformIndex = FromTransformToParticleIndex.IsValidIndex(TransformGroupIndex);
+		ensureMsgf(bIsValidTransformIndex, TEXT(" TransformGroupIndex is not valid : value = [%d] - Max valid TransformIndex = [%d]"), TransformGroupIndex, FromTransformToParticleIndex.Num());
+		if (!bIsValidTransformIndex)
+		{
+			continue;
+		}
+
 		if (Seen && (*Seen)[TransformGroupIndex])
 		{
 			continue;
