Geometry collection : fix anchored and other attribute been stripped when using strip on cook on nanite based geometry collection
- list of all attributes that should will not be skipped : Anchored | AutoInstanceMeshIndex | Level | Volume | Size
- also fix reporting of the allocated size of the bool typed attributes

diff --git a/Engine/Source/Runtime/Experimental/Chaos/Public/GeometryCollection/ManagedArray.h b/Engine/Source/Runtime/Experimental/Chaos/Public/GeometryCollection/ManagedArray.h
--- a/Engine/Source/Runtime/Experimental/Chaos/Public/GeometryCollection/ManagedArray.h
+++ b/Engine/Source/Runtime/Experimental/Chaos/Public/GeometryCollection/ManagedArray.h
@@ -800,7 +800,7 @@ public:
 
 	virtual SIZE_T GetAllocatedSize() const override
 	{
-		return ManagedArrayTypeSize::GetAllocatedSize(Array);
+		return Array.GetAllocatedSize();
 	}
 
 	/**
diff --git a/Engine/Source/Runtime/Experimental/GeometryCollectionEngine/Private/GeometryCollection/GeometryCollectionObject.cpp b/Engine/Source/Runtime/Experimental/GeometryCollectionEngine/Private/GeometryCollection/GeometryCollectionObject.cpp
--- a/Engine/Source/Runtime/Experimental/GeometryCollectionEngine/Private/GeometryCollection/GeometryCollectionObject.cpp
+++ b/Engine/Source/Runtime/Experimental/GeometryCollectionEngine/Private/GeometryCollection/GeometryCollectionObject.cpp
@@ -963,7 +963,6 @@ static void SerializeOldNaniteData(FArchive& Ar, UGeometryCollection* Owner)
 	}
 }
 
-
 /** Serialize */
 void UGeometryCollection::Serialize(FArchive& Ar)
 {
@@ -1585,8 +1584,6 @@ TSharedPtr<FGeometryCollection, ESPMode::ThreadSafe> UGeometryCollection::Genera
 	static TSet<FName> GeometryGroups{ FGeometryCollection::GeometryGroup, FGeometryCollection::VerticesGroup, FGeometryCollection::FacesGroup };
 	if (bStripOnCook)
 	{
-		// Remove all geometry
-		//static TSet<FName> GeometryGroups{ FGeometryCollection::GeometryGroup, FGeometryCollection::VerticesGroup, FGeometryCollection::FacesGroup, FGeometryCollection::MaterialGroup };
 		for (const FName& GeometryGroup : GeometryGroups)
 		{
 			TSet<FName>& SkipAttributes = SkipList.Add(GeometryGroup);
@@ -1601,7 +1598,37 @@ TSharedPtr<FGeometryCollection, ESPMode::ThreadSafe> UGeometryCollection::Genera
 	DuplicateGeometryCollection->AddAttribute<FTransform>("MassToLocal", FGeometryCollection::TransformGroup);
 	DuplicateGeometryCollection->AddAttribute<Chaos::FImplicitObjectPtr>(
 		FGeometryDynamicCollection::ImplicitsAttribute, FTransformCollection::TransformGroup);
+
+	// TODO(chaos) : for now, let's fix up the missing attributes and we should consolidate GenerateMinimalGeometryCollection and CopyCollectionAndRemoveGeometry for later releases
+	// we should also probably have a way to specify which attributes are needed in the final collection instead of hardcoding it here
+	static const FName AnchoredAttributeName(TEXT("Anchored"));
+	if (GeometryCollection->HasAttribute(AnchoredAttributeName, FTransformCollection::TransformGroup))
+	{
+		DuplicateGeometryCollection->AddAttribute<bool>(AnchoredAttributeName, FTransformCollection::TransformGroup);
+	}
+	static const FName AutoInstanceMeshIndexAttributeName(TEXT("AutoInstanceMeshIndex"));
+	if (GeometryCollection->HasAttribute(AutoInstanceMeshIndexAttributeName, FTransformCollection::TransformGroup))
+	{
+		DuplicateGeometryCollection->AddAttribute<int32>(AutoInstanceMeshIndexAttributeName, FTransformCollection::TransformGroup);
+	}
+	static const FName LevelAttributeName(TEXT("Level"));
+	if (GeometryCollection->HasAttribute(LevelAttributeName, FTransformCollection::TransformGroup))
+	{
+		DuplicateGeometryCollection->AddAttribute<int32>(LevelAttributeName, FTransformCollection::TransformGroup);
+	}
+	static const FName VolumeAttributeName(TEXT("Volume"));
+	if (GeometryCollection->HasAttribute(VolumeAttributeName, FTransformCollection::TransformGroup))
+	{
+		DuplicateGeometryCollection->AddAttribute<float>(VolumeAttributeName, FTransformCollection::TransformGroup);
+	}
+	static const FName SizeAttributeName(TEXT("Size"));
+	if (GeometryCollection->HasAttribute(SizeAttributeName, FTransformCollection::TransformGroup))
+	{
+		DuplicateGeometryCollection->AddAttribute<float>(SizeAttributeName, FTransformCollection::TransformGroup);
+	}
+
 	DuplicateGeometryCollection->CopyMatchingAttributesFrom(*GeometryCollection, &SkipList);
+
 	// If we've removed all geometry, we need to make sure any references to that geometry are removed.
 	// We also need to resize geometry groups to ensure that they are empty.
 	if (bStripOnCook)
