[NetPhysics] Make logic that interpolates received states on the client safer.
- Don't access pointer during a for-loop which can alter the collection holding the pointer data.

diff --git a/Engine/Source/Runtime/Experimental/Chaos/Private/RewindData.cpp b/Engine/Source/Runtime/Experimental/Chaos/Private/RewindData.cpp
--- a/Engine/Source/Runtime/Experimental/Chaos/Private/RewindData.cpp
+++ b/Engine/Source/Runtime/Experimental/Chaos/Private/RewindData.cpp
@@ -1998,13 +1998,13 @@ FJointState FRewindData::GetPastJointStateAtFrame(const FPBDJointConstraintHandl
 	return GetPastStateAtFrameImp<FJointState>(DirtyJoints, Handle, Frame, bRewindDataOptimization ? FFrameAndPhase::EParticleHistoryPhase::PrePushData : Phase);
 }
 
-CHAOS_API int32 bInterpolateTargetGaps = 5;
-FAutoConsoleVariableRef CVarResimInterpolateTargetGaps(TEXT("p.Resim.InterpolateTargetGaps"), bInterpolateTargetGaps, TEXT("How many frame gaps in replicated targets we should fill by interpolating between the previous and the new target received. Value in max number of frames to interpolate, deactivate by setting to 0."));
+CHAOS_API int32 InterpolateTargetGaps = 5;
+FAutoConsoleVariableRef CVarResimInterpolateTargetGaps(TEXT("p.Resim.InterpolateTargetGaps"), InterpolateTargetGaps, TEXT("How many frame gaps in replicated targets we should fill by interpolating between the previous and the new target received. Value in max number of frames to interpolate, deactivate by setting to 0."));
 
 void FRewindData::SetTargetStateAtFrame(FGeometryParticleHandle& Handle, const int32 Frame, FFrameAndPhase::EParticleHistoryPhase Phase,
 	const FVector& Position, const FQuat& Quaternion, const FVector& LinVelocity, const FVector& AngVelocity, const bool bShouldSleep)
 {
-	if (bInterpolateTargetGaps)
+	if (InterpolateTargetGaps > 0)
 	{
 		FDirtyParticleInfo& Info = FindOrAddDirtyObj(Handle);
 		FGeometryParticleStateBase& History = Info.GetHistory();
@@ -2013,23 +2013,30 @@ void FRewindData::SetTargetStateAtFrame(FGeometryParticleHandle& Handle, const i
 		if (History.TargetPositions.GetHeadFrameAndPhase(FrameAndPhase))
 		{
 			const int32 FrameDiff = Frame - FrameAndPhase.Frame;
-			if (FrameDiff > 1 && FrameDiff <= bInterpolateTargetGaps)
+			if (FrameDiff > 1 && FrameDiff <= InterpolateTargetGaps)
 			{
 				const FParticlePositionRotation* TargetXR = History.TargetPositions.Read(FrameAndPhase, PropertiesPool);
 				const FParticleVelocities* TargetVW = History.TargetVelocities.Read(FrameAndPhase, PropertiesPool);
 				const FParticleDynamicMisc* TargetDynamic = History.TargetStates.Read(FrameAndPhase, PropertiesPool);
 				if (TargetXR && TargetVW && TargetDynamic)
 				{
+					// Copy the previous state to lerp from
+					const FVector PrevX = TargetXR->GetX();
+					const FQuat PrevR = TargetXR->GetR();
+					const FVector PrevV = TargetVW->GetV();
+					const FVector PrevW = TargetVW->GetW();
+					const bool bSleep = bShouldSleep && TargetDynamic->ObjectState() == EObjectStateType::Sleeping;
+
 					for (int32 InterpFrame = 1; InterpFrame < FrameDiff; InterpFrame++)
 					{
-						const float Alpha = (1.0f / (float)FrameDiff) * (float)InterpFrame;
+						const float Alpha = static_cast<float>(InterpFrame) / static_cast<float>(FrameDiff);
 
 						PushStateAtFrame(Handle, FrameAndPhase.Frame + InterpFrame, Phase,
-							FMath::Lerp(TargetXR->GetX(), Position, Alpha),
-							FRotation3::Slerp(TargetXR->GetR(), Quaternion, Alpha),
-							FMath::Lerp(TargetVW->GetV(), LinVelocity, Alpha),
-							FMath::Lerp(TargetVW->GetW(), AngVelocity, Alpha),
-							bShouldSleep && TargetDynamic->ObjectState() == EObjectStateType::Sleeping);
+							FMath::Lerp(PrevX, Position, Alpha),
+							FRotation3::Slerp(PrevR, Quaternion, Alpha),
+							FMath::Lerp(PrevV, LinVelocity, Alpha),
+							FMath::Lerp(PrevW, AngVelocity, Alpha),
+							bSleep);
 					}
 				}
 			}
