Chaos - Character ground constraint fix to prevent slidng on slopes when standing

#rnx

diff --git a/Engine/Source/Runtime/Experimental/Chaos/Private/Chaos/Character/CharacterGroundConstraintSolver.cpp b/Engine/Source/Runtime/Experimental/Chaos/Private/Chaos/Character/CharacterGroundConstraintSolver.cpp
--- a/Engine/Source/Runtime/Experimental/Chaos/Private/Chaos/Character/CharacterGroundConstraintSolver.cpp
+++ b/Engine/Source/Runtime/Experimental/Chaos/Private/Chaos/Character/CharacterGroundConstraintSolver.cpp
@@ -165,20 +165,20 @@ namespace Chaos
 			const FReal ProjectedOverlap = -ProjectedError.Dot(GroundNormal);
 			const bool WillOverlap = ProjectedOverlap > 0.0;
 
-			if (WillOverlap)
+			const FReal InitialOverlap = -InitialError.Dot(GroundNormal);
+			const bool CurrentlyOverlapping = InitialOverlap > 0.0f;
+
+			const bool NeedsCorrection = CurrentlyOverlapping && WillOverlap;
+			CorrectionSolveFunction = NeedsCorrection ? &SolveCorrectionSingleBody : &NoSolve;
+
+			if (CurrentlyOverlapping)
 			{
 				FReal DampingFactor = FMath::Clamp(Settings.DampingFactor, 0.0, 1.0);
-				FVec3 Diff = DampingFactor * ProjectedOverlap * GroundNormal;
+				FVec3 Diff = DampingFactor * InitialOverlap * GroundNormal;
 				InitialError += Diff;
 				ProjectedError += Diff;
 			}
 
-			// Correction
-			const FReal InitialErrorN = InitialError.Dot(GroundNormal);
-			const bool CurrentlyOverlapping = InitialErrorN < 0.0f;
-			const bool NeedsCorrection = CurrentlyOverlapping && WillOverlap;
-			CorrectionSolveFunction = NeedsCorrection ? &SolveCorrectionSingleBody : &NoSolve;
-
 			// Write constraint data
 			ConstraintData.CharacterInvM = FSolverReal(CharacterInvM);
 			ConstraintData.GroundInvM = FSolverReal(GroundInvM);
@@ -186,7 +186,7 @@ namespace Chaos
 			ConstraintData.EffectiveMassN = FSolverReal(TwoBodyEffectiveMass);
 			ConstraintData.Normal = FSolverVec3(GroundNormal);
 			ConstraintData.GroundOffset = FSolverVec3(GroundOffset);
-			ConstraintData.InitialError = FSolverReal(InitialErrorN);
+			ConstraintData.InitialError = FSolverReal(InitialError.Dot(GroundNormal));
 			ConstraintData.InitialProjectedError = FSolverReal(ProjectedError.Dot(GroundNormal));
 
 			//////////////////////////////////////////////////////////////////////////
