Geometry collection : kinematically anchored geometry collection should remain kinematic if the object type is sleeping

diff --git a/Engine/Source/Runtime/Experimental/Chaos/Private/PhysicsProxy/GeometryCollectionPhysicsProxy.cpp b/Engine/Source/Runtime/Experimental/Chaos/Private/PhysicsProxy/GeometryCollectionPhysicsProxy.cpp
--- a/Engine/Source/Runtime/Experimental/Chaos/Private/PhysicsProxy/GeometryCollectionPhysicsProxy.cpp
+++ b/Engine/Source/Runtime/Experimental/Chaos/Private/PhysicsProxy/GeometryCollectionPhysicsProxy.cpp
@@ -1164,7 +1164,7 @@ void FGeometryCollectionPhysicsProxy::SetSleepingState(const Chaos::FPBDRigidsSo
 {
 	for (FClusterHandle* Handle: SolverParticleHandles)
 	{
-		if (Handle && !Handle->Disabled())
+		if (Handle && !Handle->Disabled() && Handle->IsDynamic())
 		{
 			// Sleeping Geometry Collections:
 			//   A sleeping geometry collection is dynamic internally, and then the top level
