Ensure that constraint sleep states match connected particles at the end of the step
- Rewaking can wake-up particles which have not been awake during solve
- Complete sync of constraint sleep states is needed to prevent assert
- If constraint sleep states are out of sync at the end of the step, sleeping constraints can be removed from the graph and then re-added while still sleeping.

diff --git a/Engine/Source/Programs/HeadlessChaos/Private/HeadlessChaosTestGraphEvolution.cpp b/Engine/Source/Programs/HeadlessChaos/Private/HeadlessChaosTestGraphEvolution.cpp
--- a/Engine/Source/Programs/HeadlessChaos/Private/HeadlessChaosTestGraphEvolution.cpp
+++ b/Engine/Source/Programs/HeadlessChaos/Private/HeadlessChaosTestGraphEvolution.cpp
@@ -63,6 +63,11 @@ namespace ChaosTest
 			check(CVarAngularWakeThresholdMultiplier && CVarAngularWakeThresholdMultiplier->IsVariableFloat());
 			AngularWakeThresholdMultiplier_EditorDefault = CVarAngularWakeThresholdMultiplier->GetFloat();
 			CVarAngularWakeThresholdMultiplier->Set(AngularWakeThresholdMultiplier_UnitTestDefault);
+
+			CVarValidateConstraintSleepState = IConsoleManager::Get().FindConsoleVariable(TEXT("p.Chaos.Solver.Sleep.PartialIslandSleep.ValidateConstraintSleepState"), false);
+			check(CVarValidateConstraintSleepState && CVarValidateConstraintSleepState->IsVariableBool());
+			bValidateConstraintSleepState_EditorDefault = CVarValidateConstraintSleepState->GetBool();
+			CVarValidateConstraintSleepState->Set(bValidateConstraintSleepState_UnitTestDefault);
 		}
 
 		void RestorePartialIslandSleepEditorDefaults()
@@ -74,6 +79,7 @@ namespace ChaosTest
 			CVarMinMotionlessRatio->Set(MinimumMotionlessThreshold_EditorDefault);
 			CVarLinearWakeThresholdMultiplier->Set(LinearWakeThresholdMultiplier_EditorDefault);
 			CVarAngularWakeThresholdMultiplier->Set(AngularWakeThresholdMultiplier_EditorDefault);
+			CVarValidateConstraintSleepState->Set(bValidateConstraintSleepState_EditorDefault);
 		}
 
 		// Default values for unit testing
@@ -84,6 +90,7 @@ namespace ChaosTest
 		FRealSingle MinimumMotionlessThreshold_UnitTestDefault = 0.0f;
 		FRealSingle LinearWakeThresholdMultiplier_UnitTestDefault = 1.0f;
 		FRealSingle AngularWakeThresholdMultiplier_UnitTestDefault = 1.0f;
+		bool bValidateConstraintSleepState_UnitTestDefault = true;
 
 		// Storage for default values editor hard-coded in IslandManager.cpp
 		FRealSingle PostStepWake_EditorDefault;
@@ -93,6 +100,7 @@ namespace ChaosTest
 		FRealSingle MinimumMotionlessThreshold_EditorDefault;
 		FRealSingle LinearWakeThresholdMultiplier_EditorDefault;
 		FRealSingle AngularWakeThresholdMultiplier_EditorDefault;
+		bool bValidateConstraintSleepState_EditorDefault;
 
 	public:
 		// Storage of console variable pointer
@@ -103,6 +111,7 @@ namespace ChaosTest
 		IConsoleVariable* CVarMinMotionlessRatio;
 		IConsoleVariable* CVarLinearWakeThresholdMultiplier;
 		IConsoleVariable* CVarAngularWakeThresholdMultiplier;
+		IConsoleVariable* CVarValidateConstraintSleepState;
 	};
 
 	class FGraphEvolutionTest
diff --git a/Engine/Source/Programs/HeadlessChaos/Private/HeadlessChaosTestSleeping.cpp b/Engine/Source/Programs/HeadlessChaos/Private/HeadlessChaosTestSleeping.cpp
--- a/Engine/Source/Programs/HeadlessChaos/Private/HeadlessChaosTestSleeping.cpp
+++ b/Engine/Source/Programs/HeadlessChaos/Private/HeadlessChaosTestSleeping.cpp
@@ -60,6 +60,11 @@ namespace ChaosTest
 			check(CVarAngularWakeThresholdMultiplier && CVarAngularWakeThresholdMultiplier->IsVariableFloat());
 			AngularWakeThresholdMultiplier_EditorDefault = CVarAngularWakeThresholdMultiplier->GetFloat();
 			CVarAngularWakeThresholdMultiplier->Set(AngularWakeThresholdMultiplier_UnitTestDefault);
+			
+			CVarValidateConstraintSleepState = IConsoleManager::Get().FindConsoleVariable(TEXT("p.Chaos.Solver.Sleep.PartialIslandSleep.ValidateConstraintSleepState"), false);
+			check(CVarValidateConstraintSleepState && CVarValidateConstraintSleepState->IsVariableBool());
+			bValidateConstraintSleepState_EditorDefault = CVarValidateConstraintSleepState->GetBool();
+			CVarValidateConstraintSleepState->Set(bValidateConstraintSleepState_UnitTestDefault);
 		}
 
 		void RestorePartialIslandSleepEditorDefaults()
@@ -71,6 +76,7 @@ namespace ChaosTest
 			CVarMinMotionlessRatio->Set(MinimumMotionlessThreshold_EditorDefault);
 			CVarLinearWakeThresholdMultiplier->Set(LinearWakeThresholdMultiplier_EditorDefault);
 			CVarAngularWakeThresholdMultiplier->Set(AngularWakeThresholdMultiplier_EditorDefault);
+			CVarValidateConstraintSleepState->Set(bValidateConstraintSleepState_EditorDefault);
 		}
 
 		// Default values for unit testing
@@ -81,6 +87,7 @@ namespace ChaosTest
 		FRealSingle MinimumMotionlessThreshold_UnitTestDefault = 0.0f;
 		FRealSingle LinearWakeThresholdMultiplier_UnitTestDefault = 1.0f;
 		FRealSingle AngularWakeThresholdMultiplier_UnitTestDefault = 1.0f;
+		bool bValidateConstraintSleepState_UnitTestDefault = true;
 
 		// Storage for default values editor hard-coded in IslandManager.cpp
 		FRealSingle PostStepWake_EditorDefault;
@@ -90,6 +97,7 @@ namespace ChaosTest
 		FRealSingle MinimumMotionlessThreshold_EditorDefault;
 		FRealSingle LinearWakeThresholdMultiplier_EditorDefault;
 		FRealSingle AngularWakeThresholdMultiplier_EditorDefault;
+		bool bValidateConstraintSleepState_EditorDefault;
 
 	public:
 		// Storage of console variable pointer
@@ -100,6 +108,7 @@ namespace ChaosTest
 		IConsoleVariable* CVarMinMotionlessRatio;
 		IConsoleVariable* CVarLinearWakeThresholdMultiplier;
 		IConsoleVariable* CVarAngularWakeThresholdMultiplier;
+		IConsoleVariable* CVarValidateConstraintSleepState;
 	};
 
 	class FSleepingTest
diff --git a/Engine/Source/Runtime/Experimental/Chaos/Private/Chaos/Island/IslandManager.cpp b/Engine/Source/Runtime/Experimental/Chaos/Private/Chaos/Island/IslandManager.cpp
--- a/Engine/Source/Runtime/Experimental/Chaos/Private/Chaos/Island/IslandManager.cpp
+++ b/Engine/Source/Runtime/Experimental/Chaos/Private/Chaos/Island/IslandManager.cpp
@@ -107,6 +107,10 @@ namespace Chaos::CVars
 	FRealSingle AngularPartialWakeThresholdMultiplier = 10.0f;
 	FAutoConsoleVariableRef CVarChaosSolverAngularWakeThresholdMultiplier(TEXT("p.Chaos.Solver.Sleep.PartialIslandSleep.AngularWakeThresholdMultiplier"), AngularPartialWakeThresholdMultiplier, TEXT("[EXPERIMENTAL] A multiplier applied to AngularSleepThreshold for particle waking.[def:10]"));
 
+	/** Cvar to enable validate constraint sleep state at the end of the step (debugging only) */
+	bool bChaosSolverValidateConstraintSleepState = false;
+	FAutoConsoleVariableRef ChaosSolverValidateConstraintSleepState(TEXT("p.Chaos.Solver.Sleep.PartialIslandSleep.ValidateConstraintSleepState"), bChaosSolverValidateConstraintSleepState, TEXT("[EXPERIMENTAL] Debugging only: Enable constraint sleep state validation checks/assertions at the end of the tick.[def:false]"));
+
 	/** Cvar to override the sleep counter threshold if necessary */
 	int32 ChaosSolverCollisionDefaultSleepCounterThreshold = 20;
 	FAutoConsoleVariableRef CVarChaosSolverCollisionDefaultSleepCounterThreshold(TEXT("p.Chaos.Solver.Sleep.Defaults.SleepCounterThreshold"), ChaosSolverCollisionDefaultSleepCounterThreshold, TEXT("Default counter threshold for sleeping.[def:20]"));
@@ -2700,21 +2704,30 @@ namespace Chaos::Private
 		// Make all non-moving particles in the island go to sleep
 		SleepNonMovingParticles(Island, Dt);
 
+		bool bHaveParticlesBeenWokenUp = false;
 		if (CVars::bChaosSolverPartialSleepCollisionConstraintsOnly)
 		{
 			// If this island contains any constraint which don't support partial sleeping,
 			// we need to traverse the constraints to make sure that connected particles are either both sleeping or both awake.
-			RewakeConstraintsNotSupportingPartialSleep(Island);
-
-			// NOTE: If (ChaosSolverPostStepWakeThreshold > 0.0f), we might encounter situations where an active-sleeping pair is connected by a non-supporting constraint. 
-			// However, we tolerate this since the state should only last for one tick before RewakeConstraintsNotSupportingPartialSleep will sync again.
+			bHaveParticlesBeenWokenUp = RewakeConstraintsNotSupportingPartialSleep(Island);
 		}
 
-		// Deactivate all adjacent constraints between (a) 2 sleeping particles or (b) a sleeping and a non-dynamic particle.
-		// Processing non-sleeping constraints is sufficient because no particles were woken up since the last call to @SyncConstraintSleepStateToParticles
-		// (Except for particles which were re-woken but whose connected constraints are still active).
-		PropagateParticleSleepToConstraints(Island);
+		if (bHaveParticlesBeenWokenUp)
+		{
+			// After re-waking, we need to sync all constraints to the particle state.
+			// Note that it is not sufficient to call @SleepConstraintsBetweenNonAwakeParticles because more particles might be awake than before the solve.
+			SyncConstraintSleepStateToParticles(Island);
+		}
+		else
+		{
+			// Without re-waking:
+			// Deactivate all adjacent constraints between non-awake particles, i.e. (a) 2 sleeping particles or (b) a sleeping and a non-dynamic particle.
+			// Processing non-sleeping constraints is sufficient because no particles were woken up since the last call to @SyncConstraintSleepStateToParticles
+			SleepConstraintsBetweenNonAwakeParticles(Island);
+		}
 
+		// NOTE: If (ChaosSolverPostStepWakeThreshold > 0.0f), we might encounter situations where an active-sleeping pair is connected by a non-supporting constraint. 
+		// However, we tolerate this since the state should only last for one tick before RewakeConstraintsNotSupportingPartialSleep will sync again.
 		if (CVars::ChaosSolverPostStepWakeThreshold > 0.0f)
 		{
 			// Process all constraints connecting active-sleeping pairs.
@@ -2722,6 +2735,12 @@ namespace Chaos::Private
 			// Only one layer of sleeping particles can be woken up per tick (i.e. no propagation through layers).
 			PostStepPartialWake(Island);
 		}
+
+		// Debugging only: Validate that the constraint sleep state at the end of the step matches the particle sleep states
+		if (CVars::bChaosSolverValidateConstraintSleepState)
+		{
+			ValidateConstraintSleepState(Island);
+		}
 	}
 
 	void FPBDIslandManager::SleepNonMovingParticles(FPBDIsland* Island, const FRealSingle Dt)
@@ -2811,7 +2830,7 @@ namespace Chaos::Private
 		SwitchIslandSleepType(Island, MotionlessCount);
 	}
 
-	void FPBDIslandManager::RewakeConstraintsNotSupportingPartialSleep(FPBDIsland* Island)
+	bool FPBDIslandManager::RewakeConstraintsNotSupportingPartialSleep(FPBDIsland* Island)
 	{
 		// NOTE: We use node flags here since we want the updated state when processing.
 		const auto IsSleepingAndDynamic = [](const FPBDIslandParticle* Node) -> bool
@@ -2916,9 +2935,11 @@ namespace Chaos::Private
 		// Don't reset the sleep counter here since they have already been set in AddToWakeQueue.
 		WakeUpParticles(WakeUpQueue, /*bResetSleepCounter=*/ false);
 		ResetEdgeFlags(Island);
+
+		return !WakeUpQueue.IsEmpty();
 	}
 
-	void FPBDIslandManager::PropagateParticleSleepToConstraints(FPBDIsland* Island)
+	void FPBDIslandManager::SleepConstraintsBetweenNonAwakeParticles(FPBDIsland* Island)
 	{
 		for (TArray<FPBDIslandConstraint*>& IslandEdges : Island->ContainerEdges)
 		{
@@ -3093,6 +3114,68 @@ namespace Chaos::Private
 		}
 	}
 
+	void FPBDIslandManager::ValidateConstraintSleepState(const FPBDIsland* Island)
+	{
+		for (const TArray<FPBDIslandConstraint*>& IslandEdges : Island->ContainerEdges)
+		{
+			for (const FPBDIslandConstraint* IslandEdge : IslandEdges)
+			{
+				const FPBDIslandParticle* Node0 = IslandEdge->Nodes[0];
+				const FPBDIslandParticle* Node1 = IslandEdge->Nodes[1];
+				FConstraintHandle* Constraint = IslandEdge->GetConstraint();
+
+				if (Node0 != nullptr && Node1 != nullptr)
+				{
+					const bool bIsSleeping0 = IsParticleSleeping(Node0->GetParticle());
+					const bool bIsSleeping1 = IsParticleSleeping(Node1->GetParticle());
+					const bool bIsDynamic0 = IsParticleDynamic(Node0->GetParticle());
+					const bool bIsDynamic1 = IsParticleDynamic(Node1->GetParticle());
+
+					// Constraint should sleep if
+					// Case (a) Both particles are sleeping
+					// Case (b) One particle is sleeping and the other is not dynamic (i.e. kinematic, static, etc.)
+					const bool bAreBothSleeping = bIsSleeping0 && bIsSleeping1;
+					const bool bIsSleeping0Kinematic1 = bIsSleeping0 && !bIsDynamic1;
+					const bool bIsSleeping1Kinematic0 = bIsSleeping1 && !bIsDynamic0;
+					if (bAreBothSleeping || bIsSleeping0Kinematic1 || bIsSleeping1Kinematic0)
+					{
+						check(!!IslandEdge->Flags.bIsSleeping);
+						check(!Constraint->SupportsSleeping() || Constraint->IsSleeping());
+					}
+
+					// Constraint should not sleep if
+					// Case (a) at least one particle is awake
+					const bool bIsAwakeDynamic0 = !bIsSleeping0 && bIsDynamic0;
+					const bool bIsAwakeDynamic1 = !bIsSleeping1 && bIsDynamic1;
+					if (bIsAwakeDynamic0 || bIsAwakeDynamic1)
+					{
+						check(!IslandEdge->Flags.bIsSleeping);
+						check(!Constraint->SupportsSleeping() || !Constraint->IsSleeping());
+					}
+				}
+				else
+				{
+					// If one node is NULL, the edge sleep state should match the other node.
+					if (Node0 != nullptr) // && Node1 == nullptr
+					{
+						check(IslandEdge->Flags.bIsSleeping == Node0->Flags.bIsSleeping);
+						check(!Constraint->SupportsSleeping() || Constraint->IsSleeping() == Node0->Flags.bIsSleeping);
+					}
+					else if (Node1 != nullptr) // && Node0 == nullptr
+					{
+						check(IslandEdge->Flags.bIsSleeping == Node1->Flags.bIsSleeping);
+						check(!Constraint->SupportsSleeping() || Constraint->IsSleeping() == Node1->Flags.bIsSleeping);
+					}
+					else // Node0 == nullptr && Node1 == nullptr
+					{
+						check(!IslandEdge->Flags.bIsSleeping);
+						check(!Constraint->SupportsSleeping() || !Constraint->IsSleeping());
+					}
+				}
+			}
+		}
+	}
+
 	// If (bProcessEdgeBasedWake), we first process all updated nodes and triggered wake events, then we process all new edges and wake events.
 	// If (!bProcessEdgeBasedWake), we only process updated nodes and their wake events.
 	void FPBDIslandManager::ProcessPartialIslandWake(FPBDIsland* Island, const bool bProcessEdgeBasedWake)
diff --git a/Engine/Source/Runtime/Experimental/Chaos/Public/Chaos/Island/IslandManager.h b/Engine/Source/Runtime/Experimental/Chaos/Public/Chaos/Island/IslandManager.h
--- a/Engine/Source/Runtime/Experimental/Chaos/Public/Chaos/Island/IslandManager.h
+++ b/Engine/Source/Runtime/Experimental/Chaos/Public/Chaos/Island/IslandManager.h
@@ -627,10 +627,11 @@ namespace Chaos::Private
 		// Partial Island Sleeping
 		CHAOS_API void ProcessPartialIslandSleep(FPBDIsland* Island, const FRealSingle Dt);
 		CHAOS_API void SleepNonMovingParticles(FPBDIsland* Island, const FRealSingle Dt);
-		CHAOS_API void RewakeConstraintsNotSupportingPartialSleep(FPBDIsland* Island);
-		CHAOS_API void PropagateParticleSleepToConstraints(FPBDIsland* Island);
+		CHAOS_API bool RewakeConstraintsNotSupportingPartialSleep(FPBDIsland* Island);
+		CHAOS_API void SleepConstraintsBetweenNonAwakeParticles(FPBDIsland* Island);
 		CHAOS_API void PostStepPartialWake(FPBDIsland* Island);
 		CHAOS_API void SwitchIslandSleepType(FPBDIsland* Island, const int32 MotionlessCount);
+		CHAOS_API void ValidateConstraintSleepState(const FPBDIsland* Island);
 
 		// Partial Island Waking
 		CHAOS_API void ProcessPartialIslandWake(FPBDIsland* Island, const bool bProcessEdgeBasedWake);
