Fix compile breakage from CL#45484627 due to undeclared module dependency

diff --git a/Engine/Source/Runtime/Experimental/Chaos/Chaos.Build.cs b/Engine/Source/Runtime/Experimental/Chaos/Chaos.Build.cs
--- a/Engine/Source/Runtime/Experimental/Chaos/Chaos.Build.cs
+++ b/Engine/Source/Runtime/Experimental/Chaos/Chaos.Build.cs
@@ -27,7 +27,8 @@ namespace UnrealBuildTool.Rules
 			);
 			PrivateDependencyModuleNames.AddRange(
 				new string[] {
-				"Eigen"
+				"Eigen",
+				"MeshDescription"
 				}
 			);
 
