Introduce a CVar to run PushData in single thread

diff --git a/Engine/Source/Runtime/Experimental/Chaos/Private/PBDRigidsSolver.cpp b/Engine/Source/Runtime/Experimental/Chaos/Private/PBDRigidsSolver.cpp
--- a/Engine/Source/Runtime/Experimental/Chaos/Private/PBDRigidsSolver.cpp
+++ b/Engine/Source/Runtime/Experimental/Chaos/Private/PBDRigidsSolver.cpp
@@ -421,6 +421,9 @@ namespace Chaos
 		// Enable a couple bug fixes with temporary roll-back just in case
 		bool bRemoveParticleFromMovingKinematicsOnDisable = true;
 		FAutoConsoleVariableRef CVarChaosRemoveParticleFromMovingKinematicsOnDisable(TEXT("p.Chaos.RemoveParticleFromMovingKinematicsOnDisable"), bRemoveParticleFromMovingKinematicsOnDisable, TEXT(""));
+
+		bool bSingleThreadPushData = false;
+		FAutoConsoleVariableRef CVarChaosSingleThreadPushData(TEXT("p.Chaos.SingleThreadPushData"), bSingleThreadPushData, TEXT("Run push Data in Single Thread."));
 	}
 }
 
@@ -2608,7 +2611,7 @@ namespace Chaos
 	{
 		FPullPhysicsData* PullData = MarshallingManager.GetCurrentPullData_Internal();
 
-		if (!GSingleThreadedPhysics)
+		if (!(GSingleThreadedPhysics || CVars::bSingleThreadPushData))
 		{
 			using TSOA = typename TParticleView<FPBDRigidParticles>::TSOA;
 			using THandle = typename TSOA::THandleType;
@@ -2725,7 +2728,7 @@ namespace Chaos
 			RewindData->BufferPhysicsResults(PullData->DirtyRigidErrors);
 		}
 
-		if (!GSingleThreadedPhysics)
+		if (!(GSingleThreadedPhysics || CVars::bSingleThreadPushData))
 		{
 			TaskDispatcher->WaitTaskEndBufferResults();
 		}
