Chaos : simplify the chaos task scheduler for the physics buffer pull part for Geometry collections and cluster union
- remove the need to copy the data and make the tasks to directly write to the final buffer to garantee the index matches the proxy one
- made sure we wait on the schedule batch tasks

diff --git a/Engine/Source/Runtime/Experimental/Chaos/Private/TaskDispatcherSolver.cpp b/Engine/Source/Runtime/Experimental/Chaos/Private/TaskDispatcherSolver.cpp
--- a/Engine/Source/Runtime/Experimental/Chaos/Private/TaskDispatcherSolver.cpp
+++ b/Engine/Source/Runtime/Experimental/Chaos/Private/TaskDispatcherSolver.cpp
@@ -43,13 +43,8 @@ void FTaskDispatcherSolver::Initialize(TParticleView<FPBDRigidParticles>& DirtyP
 	ActiveClusterUnions.Reset();
 
 	PullDataDirtyRigidsArray.Reset();
-	PullDataDirtyGCArray.Reset();
-	PullDataDirtyClusterArray.Reset();
 
 	PreviousCopyRigidTask = UE::Tasks::FTask();
-	PreviousCopyGCTask = UE::Tasks::FTask();
-	PreviousCopyClusterTask = UE::Tasks::FTask();
-
 }
 
 void FTaskDispatcherSolver::BufferRigidResults(FPullPhysicsData* PullData)
@@ -102,36 +97,25 @@ void FTaskDispatcherSolver::BufferGCResults(FPullPhysicsData* PullData, FPBDRigi
 		const int32 GCByTask = FMath::Max(FMath::DivideAndRoundUp(NumGC, NumGCTask), MinGCPerTask);
 		const int32 NumGCBatches = FMath::DivideAndRoundUp(NumGC, GCByTask);
 
-		PullDataDirtyGCArray.SetNum(NumGCBatches);
 		//we only fill this once per frame
 		ensure(PullData->DirtyGeometryCollections.Num() == 0);
+		PullData->DirtyGeometryCollections.SetNum(NumGC);
 
+		GCBatchTasks.SetNum(NumGCBatches);
 		for (int32 BatchIndex = 0; BatchIndex < NumGCBatches; BatchIndex++)
 		{
-			TArray<FDirtyGeometryCollectionData>& PullDataDirtyGC = PullDataDirtyGCArray[BatchIndex];
-			PullDataDirtyGC.Reset();
-
 			const int32 StartIndex = BatchIndex * GCByTask;
 			int32 EndIndex = (BatchIndex + 1) * GCByTask;
 			EndIndex = FMath::Min(NumGC, EndIndex);
 
-			UE::Tasks::FTask GCPendingBufferTask = UE::Tasks::Launch(UE_SOURCE_LOCATION, [StartIndex, EndIndex, RigidsSolver, &ActiveGC = ActiveGC, &PullDataDirtyGC]()
+			GCBatchTasks[BatchIndex] = UE::Tasks::Launch(UE_SOURCE_LOCATION, [StartIndex, EndIndex, RigidsSolver, &ActiveGC = ActiveGC, &DirtyGeometryCollections = PullData->DirtyGeometryCollections]()
 			{
 				QUICK_SCOPE_CYCLE_COUNTER(STAT_BufferGCResultsTasks);
 				for (int32 Idx = StartIndex; Idx < EndIndex; ++Idx)
 				{
-					PullDataDirtyGC.AddDefaulted();
-					ActiveGC[Idx]->BufferPhysicsResults_Internal(RigidsSolver, PullDataDirtyGC.Last());
+					ActiveGC[Idx]->BufferPhysicsResults_Internal(RigidsSolver, DirtyGeometryCollections[Idx]);
 				}
 			});
-
-			// Link the following task with the one from the previous iteration to keep determinism
-			UE::Tasks::FTask GCCopyPendingTask = UE::Tasks::Launch(UE_SOURCE_LOCATION, [PullData, &PullDataDirtyGC]()
-			{
-				QUICK_SCOPE_CYCLE_COUNTER(STAT_BufferGCResultsCopyTasks);
-				PullData->DirtyGeometryCollections.Append(PullDataDirtyGC);
-			}, TArray<UE::Tasks::FTask>{GCPendingBufferTask, PreviousCopyGCTask});
-			PreviousCopyGCTask = GCCopyPendingTask;
 		}
 	}, DirtyGCAndClusterPendingTask);
 }
@@ -152,38 +136,27 @@ void FTaskDispatcherSolver::BufferClusterResults(FPullPhysicsData* PullData)
 		const int32 ClusterByTask = FMath::Max(FMath::DivideAndRoundUp(NumClusters, NumClusterTask), MinClusterPerTask);
 		const int32 NumClusterBatches = FMath::DivideAndRoundUp(NumClusters, ClusterByTask);
 
-		PullDataDirtyClusterArray.SetNum(NumClusterBatches);
 		//we only fill this once per frame
 		ensure(PullData->DirtyClusterUnions.Num() == 0);
+		PullData->DirtyClusterUnions.SetNum(NumClusters);
 
+		ClusterBatchTasks.SetNum(NumClusterBatches);
 		for (int32 BatchIndex = 0; BatchIndex < NumClusterBatches; BatchIndex++)
 		{
 			const int32 StartIndex = BatchIndex * ClusterByTask;
 			int32 EndIndex = (BatchIndex + 1) * ClusterByTask;
 			EndIndex = FMath::Min(NumClusters, EndIndex);
 
-			TArray<FDirtyClusterUnionData>& PullDataDirtyCluster = PullDataDirtyClusterArray[BatchIndex];
-			PullDataDirtyCluster.Reset();
-
-			UE::Tasks::FTask ClusterBufferTask = UE::Tasks::Launch(UE_SOURCE_LOCATION, [StartIndex, EndIndex, &ActiveClusterUnions = ActiveClusterUnions, &PullDataDirtyCluster]()
+			ClusterBatchTasks[BatchIndex] = UE::Tasks::Launch(UE_SOURCE_LOCATION, [StartIndex, EndIndex, &ActiveClusterUnions = ActiveClusterUnions, &DirtyClusterUnions = PullData->DirtyClusterUnions]()
 			{
 				QUICK_SCOPE_CYCLE_COUNTER(STAT_BufferClusterResultsTasks);
 				for (int32 Idx = StartIndex; Idx < EndIndex; ++Idx)
 				{
-					PullDataDirtyCluster.AddDefaulted();
-					ActiveClusterUnions[Idx]->BufferPhysicsResults_Internal(PullDataDirtyCluster.Last());
+					ActiveClusterUnions[Idx]->BufferPhysicsResults_Internal(DirtyClusterUnions[Idx]);
 				}
-			});
-
-			// Link the following task with the one from the previous iteration to keep determinism
-			UE::Tasks::FTask ClusterCopyPendingTask = UE::Tasks::Launch(UE_SOURCE_LOCATION, [PullData, &PullDataDirtyCluster]()
-			{
-				QUICK_SCOPE_CYCLE_COUNTER(STAT_BufferClusterResultsCopyTasks);
-				PullData->DirtyClusterUnions.Append(PullDataDirtyCluster);
-			}, TArray<UE::Tasks::FTask>{ClusterBufferTask, PreviousCopyClusterTask});
-			PreviousCopyClusterTask = ClusterCopyPendingTask;
+			}, GCBatchTasks);
 		}
-	}, PreviousCopyGCTask);
+		}, GCPendingTask);
 }
 
 void FTaskDispatcherSolver::WaitTaskEndBufferResults() const
@@ -197,7 +170,10 @@ void FTaskDispatcherSolver::WaitTaskEndBufferResults() const
 	// a CU and its GC the program doesn't create any race condition
 	UE::Tasks::Wait(TArray<UE::Tasks::FTask>{ClusterPendingTask});
 	// Only when those outer tasks are finished and so those inner task created we can wait for those inner tasks
-	UE::Tasks::Wait(TArray<UE::Tasks::FTask>{PreviousCopyRigidTask, PreviousCopyGCTask, PreviousCopyClusterTask});
+	TArray<UE::Tasks::FTask> TaskToWaitOn({ PreviousCopyRigidTask });
+	TaskToWaitOn.Append(GCBatchTasks);
+	TaskToWaitOn.Append(ClusterBatchTasks);
+	UE::Tasks::Wait(TaskToWaitOn);
 }
 
 }
\ No newline at end of file
diff --git a/Engine/Source/Runtime/Experimental/Chaos/Private/TaskDispatcherSolver.h b/Engine/Source/Runtime/Experimental/Chaos/Private/TaskDispatcherSolver.h
--- a/Engine/Source/Runtime/Experimental/Chaos/Private/TaskDispatcherSolver.h
+++ b/Engine/Source/Runtime/Experimental/Chaos/Private/TaskDispatcherSolver.h
@@ -40,14 +40,14 @@ private:
 
 	// Last task which copy in PullData the proxy for those three types of proxy
 	UE::Tasks::FTask PreviousCopyRigidTask;
-	UE::Tasks::FTask PreviousCopyGCTask;
-	UE::Tasks::FTask PreviousCopyClusterTask;
 
 	// PullData copies to buffer physics results by batch
 	TArray<TArray<FDirtyRigidParticleData>> PullDataDirtyRigidsArray;
-	TArray<TArray<FDirtyGeometryCollectionData>> PullDataDirtyGCArray;
-	TArray<TArray<FDirtyClusterUnionData>> PullDataDirtyClusterArray;
 	
+	// GC and cluster batches tasks launches by GCPendingTask & ClusterPendingTask
+	TArray<UE::Tasks::FTask> GCBatchTasks;
+	TArray<UE::Tasks::FTask> ClusterBatchTasks;
+
 	int32 NumRigidBatches;
 
 public:
