Geometry collection : fix issue with geometry collection nanite fallback enabled when nanite is disabled in editor
- because we were not testing if nanite was enabled globally , the scene proxy was being recreated every frame, causing the simulation to be recreated to its rest pose

diff --git a/Engine/Source/Runtime/Experimental/GeometryCollectionEngine/Private/GeometryCollection/GeometryCollectionComponent.cpp b/Engine/Source/Runtime/Experimental/GeometryCollectionEngine/Private/GeometryCollection/GeometryCollectionComponent.cpp
--- a/Engine/Source/Runtime/Experimental/GeometryCollectionEngine/Private/GeometryCollection/GeometryCollectionComponent.cpp
+++ b/Engine/Source/Runtime/Experimental/GeometryCollectionEngine/Private/GeometryCollection/GeometryCollectionComponent.cpp
@@ -4446,7 +4446,7 @@ void UGeometryCollectionComponent::UpdateRenderSystemsIfNeeded(bool bDynamicColl
 #if WITH_EDITOR
 	if (IsRegistered() && SceneProxy && RestCollection)
 	{
-		const bool bWantNanite = RestCollection->EnableNanite && GGeometryCollectionNanite != 0;
+		const bool bWantNanite = UseNanite(GMaxRHIShaderPlatform) && RestCollection->EnableNanite && GGeometryCollectionNanite != 0;
 		const bool bHaveNanite = SceneProxy->IsNaniteMesh();
 		bool bRecreateProxy = bWantNanite != bHaveNanite;
 		if (bRecreateProxy)
