Dataflow:
- Tidied up the private deprecations in FDataflowNode.
- Fixed typo on OutputArrayProperties.
- Prepared the override color properties for a move to the private section.

#rnx

diff --git a/Engine/Plugins/Experimental/ChaosFlesh/Source/ChaosFleshNodes/Private/Dataflow/ChaosFleshComputeMuscleActivationNode.cpp b/Engine/Plugins/Experimental/ChaosFlesh/Source/ChaosFleshNodes/Private/Dataflow/ChaosFleshComputeMuscleActivationNode.cpp
--- a/Engine/Plugins/Experimental/ChaosFlesh/Source/ChaosFleshNodes/Private/Dataflow/ChaosFleshComputeMuscleActivationNode.cpp
+++ b/Engine/Plugins/Experimental/ChaosFlesh/Source/ChaosFleshNodes/Private/Dataflow/ChaosFleshComputeMuscleActivationNode.cpp
@@ -401,7 +401,7 @@ bool FSetMuscleActivationParameterNode::ShouldInvalidateOnPropertyChanged(const
 	{
 		return false;
 	}
-	return true;
+	return Super::ShouldInvalidateOnPropertyChanged(InPropertyChangedEvent);
 }
 
 FReadSkeletalMeshCurvesDataflowNode::FReadSkeletalMeshCurvesDataflowNode(const UE::Dataflow::FNodeParameters& InParam, FGuid InGuid)
@@ -479,6 +479,10 @@ void FReadSkeletalMeshCurvesDataflowNode::Evaluate(UE::Dataflow::FContext& Conte
 bool FReadSkeletalMeshCurvesDataflowNode::ShouldInvalidateOnPropertyChanged(const FPropertyChangedEvent& InPropertyChangedEvent) const
 {
 	const FName ChangedPropertyName = InPropertyChangedEvent.GetMemberPropertyName();
-	return !(ChangedPropertyName == GET_MEMBER_NAME_CHECKED(FReadSkeletalMeshCurvesDataflowNode, CurveMuscleNameArray) ||
-		ChangedPropertyName == GET_MEMBER_NAME_CHECKED(FReadSkeletalMeshCurvesDataflowNode, ImportSKMCurveNames));
-}
\ No newline at end of file
+	if (ChangedPropertyName == GET_MEMBER_NAME_CHECKED(FReadSkeletalMeshCurvesDataflowNode, CurveMuscleNameArray) ||
+		ChangedPropertyName == GET_MEMBER_NAME_CHECKED(FReadSkeletalMeshCurvesDataflowNode, ImportSKMCurveNames))
+	{
+		return false;
+	}
+	return Super::ShouldInvalidateOnPropertyChanged(InPropertyChangedEvent);
+}
diff --git a/Engine/Plugins/Experimental/Dataflow/Source/DataflowEditor/Private/Dataflow/DataflowEditorCommands.cpp b/Engine/Plugins/Experimental/Dataflow/Source/DataflowEditor/Private/Dataflow/DataflowEditorCommands.cpp
--- a/Engine/Plugins/Experimental/Dataflow/Source/DataflowEditor/Private/Dataflow/DataflowEditorCommands.cpp
+++ b/Engine/Plugins/Experimental/Dataflow/Source/DataflowEditor/Private/Dataflow/DataflowEditorCommands.cpp
@@ -420,9 +420,7 @@ void FDataflowEditorCommands::OnPropertyValueChanged(UDataflow* OutDataflow, TSh
 						SelectedNode->UpdatePinsDefaultValuesFromNode();
 					}
 
-					// we do not want to invalidate when changing the color of the node as this is irrelevant to the output of the node
-					const bool bIsOverrideColorBeingChanged = (PropertyName == GET_MEMBER_NAME_CHECKED(FDataflowNode, bOverrideColor)) || (PropertyName == GET_MEMBER_NAME_CHECKED(FDataflowNode, OverrideColor));
-					if (!bIsOverrideColorBeingChanged && DataflowNode->ShouldInvalidateOnPropertyChanged(InPropertyChangedEvent))
+					if (DataflowNode->ShouldInvalidateOnPropertyChanged(InPropertyChangedEvent))
 					{
 						// Invalidate the node and reset the editor timestamp
 						DataflowNode->Invalidate();
diff --git a/Engine/Plugins/Experimental/Dataflow/Source/DataflowEditor/Private/Dataflow/DataflowSNode.cpp b/Engine/Plugins/Experimental/Dataflow/Source/DataflowEditor/Private/Dataflow/DataflowSNode.cpp
--- a/Engine/Plugins/Experimental/Dataflow/Source/DataflowEditor/Private/Dataflow/DataflowSNode.cpp
+++ b/Engine/Plugins/Experimental/Dataflow/Source/DataflowEditor/Private/Dataflow/DataflowSNode.cpp
@@ -138,10 +138,8 @@ TSharedPtr<SGraphPin> SDataflowEdNode::CreatePinWidget(UEdGraphPin* Pin) const
 						//
 						bIsOutputInvalid = !NumKeys ||
 							!CacheKeys.Contains(Output->CacheKey()) ||
-						PRAGMA_DISABLE_DEPRECATION_WARNINGS  // Until LastModifiedTimestamp becomes private
-							DataflowNode->LastModifiedTimestamp.IsInvalid() ||
-							!DataflowContext->IsCacheEntryAfterTimestamp(Output->CacheKey(), DataflowNode->LastModifiedTimestamp);
-						PRAGMA_ENABLE_DEPRECATION_WARNINGS
+							DataflowNode->GetTimestamp().IsInvalid() ||
+							!DataflowContext->IsCacheEntryAfterTimestamp(Output->CacheKey(), DataflowNode->GetTimestamp());
 					}
 				}
 			}
diff --git a/Engine/Plugins/Runtime/HairStrands/Source/HairStrandsDataflow/Private/GroomAssetTerminalNode.cpp b/Engine/Plugins/Runtime/HairStrands/Source/HairStrandsDataflow/Private/GroomAssetTerminalNode.cpp
--- a/Engine/Plugins/Runtime/HairStrands/Source/HairStrandsDataflow/Private/GroomAssetTerminalNode.cpp
+++ b/Engine/Plugins/Runtime/HairStrands/Source/HairStrandsDataflow/Private/GroomAssetTerminalNode.cpp
@@ -595,7 +595,7 @@ bool FGroomAssetTerminalDataflowNode_v2::ShouldInvalidateOnPropertyChanged(const
 		// we only changing the name of the inputs, no need to inavlidate the node
 		return false;
 	}
-	return true;
+	return Super::ShouldInvalidateOnPropertyChanged(InPropertyChangedEvent);
 }
 
 void FGroomAssetTerminalDataflowNode_v2::OnPropertyChanged(UE::Dataflow::FContext& Context, const FPropertyChangedEvent& InPropertyChangedEvent)
diff --git a/Engine/Source/Runtime/Experimental/Dataflow/Core/Private/Dataflow/DataflowNode.cpp b/Engine/Source/Runtime/Experimental/Dataflow/Core/Private/Dataflow/DataflowNode.cpp
--- a/Engine/Source/Runtime/Experimental/Dataflow/Core/Private/Dataflow/DataflowNode.cpp
+++ b/Engine/Source/Runtime/Experimental/Dataflow/Core/Private/Dataflow/DataflowNode.cpp
@@ -680,9 +680,7 @@ void FDataflowNode::Unfreeze(UE::Dataflow::FContext& Context)
 
 bool FDataflowNode::IsActive(bool bCheckFlagOnly) const
 {
-	PRAGMA_DISABLE_DEPRECATION_WARNINGS
 	return bActive && (bCheckFlagOnly || bDataflowEnableGraphEval);
-	PRAGMA_ENABLE_DEPRECATION_WARNINGS
 }
 
 void FDataflowNode::Invalidate(const UE::Dataflow::FTimestamp& InModifiedTimestamp)
@@ -695,7 +693,6 @@ void FDataflowNode::Invalidate(const UE::Dataflow::FTimestamp& InModifiedTimesta
 		}
 		return;
 	}
-PRAGMA_DISABLE_DEPRECATION_WARNINGS  // Until LastModifiedTimestamp becomes private
 	if (LastModifiedTimestamp < InModifiedTimestamp)
 	{
 		LastModifiedTimestamp = InModifiedTimestamp;
@@ -706,7 +703,6 @@ PRAGMA_DISABLE_DEPRECATION_WARNINGS  // Until LastModifiedTimestamp becomes priv
 		}
 
 		// propagate to downstream
-PRAGMA_ENABLE_DEPRECATION_WARNINGS
 		for (TPair<UE::Dataflow::FConnectionKey, FDataflowOutput*>& Elem : ExpandedOutputs)
 		{
 			FDataflowOutput* const Con = Elem.Value;
@@ -1037,7 +1033,7 @@ FDataflowOutput& FDataflowNode::RegisterOutputArrayConnectionInternal(const UE::
 	OutputParams.Name = GetPropertyFullName(PropertyChain, Reference.Index);
 	OutputParams.InnerOffset = UE::Dataflow::Private::GetArrayElementOffsetFromReference(OutputParams.ArrayProperty, Reference);
 
-	InputArrayProperties.Emplace(OutputParams.Offset, OutputParams.ArrayProperty);
+	OutputArrayProperties.Emplace(OutputParams.Offset, OutputParams.ArrayProperty);
 
 	FDataflowOutput* const Output = new FDataflowArrayOutput(Reference.Index, OutputParams);
 	AddOutput(Output);
@@ -1275,9 +1271,7 @@ bool FDataflowNode::ValidateConnections()
 				}
 			}
 
-PRAGMA_DISABLE_DEPRECATION_WARNINGS  // Until Outputs becomes private
 			for (const TPair<UE::Dataflow::FConnectionKey, FDataflowOutput*>& Output : ExpandedOutputs)
-PRAGMA_ENABLE_DEPRECATION_WARNINGS
 			{
 				FDataflowOutput* OutputVal = Output.Value;
 
@@ -1297,6 +1291,14 @@ PRAGMA_ENABLE_DEPRECATION_WARNINGS
 	return bHasValidConnections;
 }
 
+bool FDataflowNode::ShouldInvalidateOnPropertyChanged(const FPropertyChangedEvent& InPropertyChangedEvent) const
+{
+PRAGMA_DISABLE_DEPRECATION_WARNINGS  // UE_DEPRECATED(5.7, "Will be made private in 5.9, use IsColorOverriden() instead.")
+	return InPropertyChangedEvent.GetPropertyName() != GET_MEMBER_NAME_CHECKED(FDataflowNode, bOverrideColor) &&
+	InPropertyChangedEvent.GetPropertyName() != GET_MEMBER_NAME_CHECKED(FDataflowNode, OverrideColor);
+PRAGMA_ENABLE_DEPRECATION_WARNINGS
+}
+
 TUniquePtr<const FStructOnScope> FDataflowNode::NewStructOnScopeConst() const
 {
 	// New StructOnScopeConst is non-const and virtual , changing it would be quite difficult
diff --git a/Engine/Source/Runtime/Experimental/Dataflow/Core/Public/Dataflow/DataflowNode.h b/Engine/Source/Runtime/Experimental/Dataflow/Core/Public/Dataflow/DataflowNode.h
--- a/Engine/Source/Runtime/Experimental/Dataflow/Core/Public/Dataflow/DataflowNode.h
+++ b/Engine/Source/Runtime/Experimental/Dataflow/Core/Public/Dataflow/DataflowNode.h
@@ -56,33 +56,24 @@ PRAGMA_ENABLE_DEPRECATION_WARNINGS
 	friend class UE::Dataflow::FGraph;
 	friend struct FDataflowConnection;
 
-	UE_DEPRECATED(5.5, "Will be made private in 5.7, use Get/Set methods instead.")
-	FGuid Guid;
-	UE_DEPRECATED(5.5, "Will be made private in 5.7, use Get/Set methods instead.")
-	FName Name;
-	UE_DEPRECATED(5.5, "Will be made private in 5.7, use Get/Set methods instead.")
-	UE::Dataflow::FTimestamp LastModifiedTimestamp;
-
-	UE_DEPRECATED(5.5, "Inputs type has changed and has been made private (see ExpandedInputs)")
-	TMap< int, FDataflowInput* > Inputs;
 	UE_DEPRECATED(5.6, "Outputs type has changed and has been made private (see ExpandedOutputs)")
 	TMap< int, FDataflowOutput* > Outputs;
 
-	UE_DEPRECATED(5.6, "Will be made private in 5.8, use IsActive function instead.")
-	UPROPERTY(EditAnywhere, Category = "Dataflow", Meta = (DisplayAfter = "FreezeActions"))
-	bool bActive = true;
-
-	UPROPERTY(EditAnywhere, Category = "Dataflow", meta = (InlineEditConditionToggle))
+	UE_DEPRECATED(5.7, "Will be made private in 5.9, use IsColorOverriden() instead.")
+	UPROPERTY(EditAnywhere, Category = "Dataflow", meta = (InlineEditConditionToggle, DisplayAfter = "bActive"))
 	bool bOverrideColor = false;
 
+	UE_DEPRECATED(5.7, "Will be made private in 5.9, use GetOverrideColor() instead.")
 	UPROPERTY(EditAnywhere, Category = "Dataflow", meta = (EditCondition="bOverrideColor"))
 	FLinearColor OverrideColor = FLinearColor::White;
 
 	DATAFLOWCORE_API FDataflowNode();
 	DATAFLOWCORE_API FDataflowNode(const UE::Dataflow::FNodeParameters& Param, FGuid InGuid = FGuid::NewGuid());
 
-	PRAGMA_DISABLE_DEPRECATION_WARNINGS
+PRAGMA_DISABLE_DEPRECATION_WARNINGS
 	virtual ~FDataflowNode() { ClearInputs(); ClearOutputs(); }
+PRAGMA_ENABLE_DEPRECATION_WARNINGS
+
 	// Warning: FDataflowNodes aren't actually safe to copy/move yet. These are here to disable deprecation warnings from the implicit operators that were getting created anyways.
 	// (Deleting the operators would require tagging all derived classes since this is a USTRUCT, so also not doing that here).
 	FDataflowNode& operator=(const FDataflowNode& Other)
@@ -95,7 +86,6 @@ PRAGMA_ENABLE_DEPRECATION_WARNINGS
 	FName GetName() const { return Name; }
 	void SetName(FName InName) { Name = InName; }
 	UE::Dataflow::FTimestamp GetTimestamp() const { return LastModifiedTimestamp;  }
-	PRAGMA_ENABLE_DEPRECATION_WARNINGS
 	DATAFLOWCORE_API uint32 GetValueHash();
 
 	static FName StaticType() { return FName("FDataflowNode"); }
@@ -160,42 +150,17 @@ PRAGMA_ENABLE_DEPRECATION_WARNINGS
 
 	DATAFLOWCORE_API TArray<UE::Dataflow::FPin> GetPins() const;
 
-	UE_DEPRECATED(5.5, "Use AddPins method instead")
-	virtual UE::Dataflow::FPin AddPin() { return UE::Dataflow::FPin::InvalidPin; }
 	/** Override this function to add the AddOptionPin functionality to the node's context menu. */
 	virtual TArray<UE::Dataflow::FPin> AddPins()
 	{
-PRAGMA_DISABLE_DEPRECATION_WARNINGS
-		const UE::Dataflow::FPin DeprecatedAddPin = AddPin();
-PRAGMA_ENABLE_DEPRECATION_WARNINGS
-		if(DeprecatedAddPin == UE::Dataflow::FPin::InvalidPin)
-		{
-			return TArray<UE::Dataflow::FPin>();
-		}
-		return { DeprecatedAddPin };
+		return TArray<UE::Dataflow::FPin>();
 	}
 	/** Override this function to add the AddOptionPin functionality to the node's context menu. */
 	virtual bool CanAddPin() const { return false; }
-	UE_DEPRECATED(5.5, "Use GetPinsToRemove method instead")
-	virtual UE::Dataflow::FPin GetPinToRemove() const { return UE::Dataflow::FPin::InvalidPin; }
-	UE_DEPRECATED(5.4, "Use GetPinsToRemove and OnPinRemoved instead.")
-	virtual UE::Dataflow::FPin RemovePin()
-	{
-		PRAGMA_DISABLE_DEPRECATION_WARNINGS
-		return GetPinToRemove();
-		PRAGMA_ENABLE_DEPRECATION_WARNINGS
-	}
 	/** Override this function to add the RemoveOptionPin functionality to the node's context menu. OnPinRemoved will be called in this order.*/
 	virtual TArray<UE::Dataflow::FPin> GetPinsToRemove() const
 	{
-PRAGMA_DISABLE_DEPRECATION_WARNINGS
-		const UE::Dataflow::FPin DeprecatedRemovePin = GetPinToRemove();
-PRAGMA_ENABLE_DEPRECATION_WARNINGS
-		if (DeprecatedRemovePin == UE::Dataflow::FPin::InvalidPin)
-		{
-			return TArray<UE::Dataflow::FPin>();
-		}
-		return { DeprecatedRemovePin };
+		return TArray<UE::Dataflow::FPin>();
 	}
 	/** 
 	 * Override this to update any bookkeeping when a pin is being removed.
@@ -395,11 +360,6 @@ PRAGMA_ENABLE_DEPRECATION_WARNINGS
 	{
 		UnregisterInputConnection(GetKeyFromReference(Reference));
 	}
-	UE_DEPRECATED(5.5, "PropertyName is no longer needed to unregister connections")
-	void UnregisterInputConnection(const void* Reference, const FName& PropertyName)
-	{
-		return UnregisterInputConnection(GetKeyFromReference(Reference));
-	}
 	DATAFLOWCORE_API void UnregisterInputConnection(const UE::Dataflow::FConnectionKey& Key);
 	/** Unregister the connection if one exists matching this pin, then invalidate the graph. */
 	DATAFLOWCORE_API void UnregisterPinConnection(const UE::Dataflow::FPin& Pin);
@@ -434,12 +394,16 @@ PRAGMA_ENABLE_DEPRECATION_WARNINGS
 
 	bool IsColorOverriden() const
 	{
+PRAGMA_DISABLE_DEPRECATION_WARNINGS  // UE_DEPRECATED(5.7, "Will be made private in 5.9, use IsColorOverriden() instead.")
 		return bOverrideColor;
+PRAGMA_ENABLE_DEPRECATION_WARNINGS
 	}
 
 	const FLinearColor& GetOverrideColor() const
 	{
+PRAGMA_DISABLE_DEPRECATION_WARNINGS  // UE_DEPRECATED(5.7, "Will be made private in 5.9, use GetOverrideColor() instead.")
 		return OverrideColor;
+PRAGMA_ENABLE_DEPRECATION_WARNINGS
 	}
 
 	/**
@@ -447,9 +411,7 @@ PRAGMA_ENABLE_DEPRECATION_WARNINGS
 	 */
 	static const FName GetActivePropertyName()
 	{
-		PRAGMA_DISABLE_DEPRECATION_WARNINGS
 		return GET_MEMBER_NAME_CHECKED(FDataflowNode, bActive);
-		PRAGMA_ENABLE_DEPRECATION_WARNINGS
 	}
 
 	/** 
@@ -755,7 +717,8 @@ PRAGMA_ENABLE_DEPRECATION_WARNINGS
 	// if the node creates a connection, returns its pointer
 	virtual const FDataflowConnection* OnDropConnectionOnNode(const FDataflowConnection& DroppedConnection) { return nullptr; }
 
-	virtual bool ShouldInvalidateOnPropertyChanged(const FPropertyChangedEvent& InPropertyChangedEvent) const { return true; };
+	// Override when a property change musn't invalidate the node. Must return the base implementation when the property name isn't recognised.
+	DATAFLOWCORE_API virtual bool ShouldInvalidateOnPropertyChanged(const FPropertyChangedEvent& InPropertyChangedEvent) const;
 
 	// returns true to disable node properties in its detail view when the corresponding input is connected ( default is true )
 	virtual bool MakeConnectedPropertiesReadOnly() const { return true; }
@@ -892,6 +855,10 @@ private:
 
 	virtual TArray<UE::Dataflow::FRenderingParameter> GetRenderParametersImpl() const { return TArray<UE::Dataflow::FRenderingParameter>(); }
 
+	FGuid Guid;
+	FName Name;
+	UE::Dataflow::FTimestamp LastModifiedTimestamp;
+
 	bool bHasValidConnections = true;
 	TMap<UE::Dataflow::FConnectionKey, FDataflowInput*> ExpandedInputs;
 	TMap<uint32, const FArrayProperty*> InputArrayProperties; // Used to calculate ContainerElementOffsets
@@ -907,6 +874,10 @@ private:
 #endif
 #endif
 
+	UPROPERTY(EditAnywhere, Category = "Dataflow")
+	bool bActive = true;
+
+
 	UPROPERTY()
 	FInstancedPropertyBag FrozenProperties;  // Property bag for frozen outputs
 
