Automatically skip validating chunks if IASTool is not able to decompress the chunks compression format.

#rnx

- Some platforms use compression formats that are not avaliable out of the box and the validation for ChunkPlugin will not be able to decompress them.
- We now check to see if the compression format is one that we can read and skip the validation if we cannot.
- Note that FCompression::IsFormatValid will print a log line each time it fails to find a format to it is important that we only call it once per format.
- Support for additional compression formats may be added in the future.
- I plan to add more conditions to ::ShouldValidateChunk as there are other scenarios where we want to skip the validation. This is why we don't early out on success etc.

diff --git a/Engine/Source/Runtime/Experimental/IoStore/OnDemand/Private/Tool/ChunkPlugin.cpp b/Engine/Source/Runtime/Experimental/IoStore/OnDemand/Private/Tool/ChunkPlugin.cpp
--- a/Engine/Source/Runtime/Experimental/IoStore/OnDemand/Private/Tool/ChunkPlugin.cpp
+++ b/Engine/Source/Runtime/Experimental/IoStore/OnDemand/Private/Tool/ChunkPlugin.cpp
@@ -12,6 +12,7 @@
 #include "IO/IoStatus.h"
 #include "IO/IoStore.h"
 #include "IO/IoStoreOnDemand.h"
+#include "Misc/Compression.h"
 #include "Misc/DateTime.h"
 #include "Misc/FileHelper.h"
 #include "Misc/KeyChainUtilities.h"
@@ -232,9 +233,43 @@ static FIoStatus WriteChunk(const FString& Directory, FMemoryView Chunk, const F
 		<< TEXT("'");
 }
 
+static bool ShouldValidateChunk(FName CompressionFormat)
+{
+	static TMap<FName, bool> ValidatedFormats;
+
+	if (!CompressionFormat.IsNone())
+	{
+		if (bool* bCachedResult = ValidatedFormats.Find(CompressionFormat))
+		{
+			if (*bCachedResult == false)
+			{
+				return false;
+			}
+		}
+		else
+		{
+			// Note that if IsFormatValid returns false it will log a message about it. This is one of the
+			// reasons we must make sure to only check once per format to avoid  log spam. The other reason
+			// is that the check is fairly slow.
+			bool bValid = FCompression::IsFormatValid(CompressionFormat);
+			ValidatedFormats.Add(CompressionFormat, bValid);
+
+			if (bValid == false)
+			{
+				UE_LOG(LogIoStoreOnDemand, Display, TEXT("Skipping validation tests for chunks with the compression format '%s'"), *CompressionFormat.ToString())
+				return false;
+			}
+		}
+	}
+
+	return true;
+}
+
 static int32 ChunkContainer(const FString& ContainerUTocPath, const FChunkPluginSettings& Settings, const TMap<FGuid, FAES::FAESKey>& EncryptionKeys, FChunkWriterInterface* ChunkWriter, const FString& ChunksRelativeFolder, bool bIgnoreContainerHeader,
 	FOnDemandToc& OutOnDemandToc, TMap<FString, TArray<FOnDemandTocTagSetPackageList>>* OutFoundTagSets, TArray<FString>& OutContainerFiles, FChunkPluginStats& OutStats)
 {
+	TRACE_CPUPROFILER_EVENT_SCOPE(ChunkPluginCommand::ChunkContainer);
+
 	check(ChunkWriter != nullptr);
 
 	FIoStoreReader ContainerFileReader;
@@ -398,7 +433,7 @@ static int32 ChunkContainer(const FString& ContainerUTocPath, const FChunkPlugin
 		}
 
 		// Additional validation to make sure that the chunks can be decrypted and decoded
-#if 0
+		if (ShouldValidateChunk(CompressionFormat))
 		{
 			FIoChunkDecodingParams DecodingParams;
 			DecodingParams.CompressionFormat = CompressionFormat;
@@ -425,7 +460,6 @@ static int32 ChunkContainer(const FString& ContainerUTocPath, const FChunkPlugin
 				return -1;
 			}
 		}
-#endif
 
 		FIoStatus WriteStatus = ChunkWriter->WriteChunk(ChunksRelativeFolder, ReadResult.IoBuffer, ChunkHash);
 		if (WriteStatus.IsOk() == false)
@@ -478,6 +512,8 @@ static int32 ChunkContainer(const FString& ContainerUTocPath, const FChunkPlugin
 
 int32 WriteOnDemandToc(FOnDemandToc& OnDemandToc, const FString& ContainerName, const FString& OutputFolder)
 {
+	TRACE_CPUPROFILER_EVENT_SCOPE(ChunkPluginCommand::WriteOnDemandToc);
+
 	IFileManager& FileMgr = IFileManager::Get();
 
 	FString Filename = FPathViews::SetExtension(ContainerName, TEXT(".uondemandtoc"));
@@ -527,6 +563,7 @@ static FKeyChain LoadCryptoKeys(const FContext& Context)
 ////////////////////////////////////////////////////////////////////////////////
 static int32 ChunkPluginCommandEntry(const FContext& Context)
 {
+	TRACE_CPUPROFILER_EVENT_SCOPE(ChunkPluginCommand);
 	//while (!FPlatformMisc::IsDebuggerPresent());
 	//UE_DEBUG_BREAK();
 
@@ -746,6 +783,8 @@ static int32 ChunkPluginCommandEntry(const FContext& Context)
 
 	if (bDeleteContainerFiles)
 	{
+		TRACE_CPUPROFILER_EVENT_SCOPE(ChunkPluginCommand::DeleteContainerFiles);
+
 		UE_LOG(LogIoStoreOnDemand, Display, TEXT("Deleting chunked  containers..."));
 
 		for (const FString& Path : ContainerFiles)
@@ -764,6 +803,8 @@ static int32 ChunkPluginCommandEntry(const FContext& Context)
 	// Write additional file(s)
 	if (bIncludeSigPak)
 	{
+		TRACE_CPUPROFILER_EVENT_SCOPE(ChunkPluginCommand::IncludeSigPak);
+
 		const FStringView AllowedExt[]
 		{
 			TEXTVIEW("pak"),
@@ -853,6 +894,8 @@ static int32 ChunkPluginCommandEntry(const FContext& Context)
 	// Write dummy containers if necessary
 	if (bDeleteContainerFiles)
 	{
+		TRACE_CPUPROFILER_EVENT_SCOPE(ChunkPluginCommand::WriteDummyContainerFiles);
+
 		UE_LOG(LogIoStoreOnDemand, Display, TEXT("Replacing .utoc for chunked  containers..."));
 
 		for (const FString& SourcePath : ContainerFiles)
