The ChunkPlugin will not validate chunks that will decompress larger than 32MiB to avoid potential high memory peaks.

#rnx

diff --git a/Engine/Source/Runtime/Experimental/IoStore/OnDemand/Private/Tool/ChunkPlugin.cpp b/Engine/Source/Runtime/Experimental/IoStore/OnDemand/Private/Tool/ChunkPlugin.cpp
--- a/Engine/Source/Runtime/Experimental/IoStore/OnDemand/Private/Tool/ChunkPlugin.cpp
+++ b/Engine/Source/Runtime/Experimental/IoStore/OnDemand/Private/Tool/ChunkPlugin.cpp
@@ -233,7 +233,7 @@ static FIoStatus WriteChunk(const FString& Directory, FMemoryView Chunk, const F
 		<< TEXT("'");
 }
 
-static bool ShouldValidateChunk(FName CompressionFormat)
+static bool ShouldValidateChunk(FName CompressionFormat, uint64 ChunkRawSize)
 {
 	static TMap<FName, bool> ValidatedFormats;
 
@@ -262,6 +262,13 @@ static bool ShouldValidateChunk(FName CompressionFormat)
 		}
 	}
 
+	const uint64 MaxChunkSize = 32 * 1024 * 1024;	// We do not want to validate any chunk that will decompress to more than 32 MiB
+													// as we don't want to use too much memory.
+	if (ChunkRawSize >= MaxChunkSize)
+	{
+		return false;
+	}
+
 	return true;
 }
 
@@ -433,7 +440,7 @@ static int32 ChunkContainer(const FString& ContainerUTocPath, const FChunkPlugin
 		}
 
 		// Additional validation to make sure that the chunks can be decrypted and decoded
-		if (ShouldValidateChunk(CompressionFormat))
+		if (ShouldValidateChunk(CompressionFormat, RawChunkSize))
 		{
 			FIoChunkDecodingParams DecodingParams;
 			DecodingParams.CompressionFormat = CompressionFormat;
