ChaosFlesh: fix backward compatibility for VertexBoneWeightsFacade

diff --git a/Engine/Source/Runtime/Experimental/Chaos/Private/Chaos/Deformable/ChaosDeformableSolver.cpp b/Engine/Source/Runtime/Experimental/Chaos/Private/Chaos/Deformable/ChaosDeformableSolver.cpp
--- a/Engine/Source/Runtime/Experimental/Chaos/Private/Chaos/Deformable/ChaosDeformableSolver.cpp
+++ b/Engine/Source/Runtime/Experimental/Chaos/Private/Chaos/Deformable/ChaosDeformableSolver.cpp
@@ -1383,7 +1383,7 @@ namespace Chaos::Softs
 								typedef GeometryCollection::Facades::FVertexBoneWeightsFacade FWeightsFacade;
 								bool bParticleTouched = false;
 								FWeightsFacade WeightsFacade(Rest);
-								if (WeightsFacade.IsValid())
+								if (WeightsFacade.HasValidBoneIndicesAndWeights())
 								{
 									int32 NumObjectVertices = Rest.NumElements(FGeometryCollection::VerticesGroup);
 									int32 ObjectVertexIndex = SolverParticleToObjectVertexIndex(Index);
diff --git a/Engine/Source/Runtime/Experimental/Chaos/Private/GeometryCollection/Facades/CollectionVertexBoneWeightsFacade.cpp b/Engine/Source/Runtime/Experimental/Chaos/Private/GeometryCollection/Facades/CollectionVertexBoneWeightsFacade.cpp
--- a/Engine/Source/Runtime/Experimental/Chaos/Private/GeometryCollection/Facades/CollectionVertexBoneWeightsFacade.cpp
+++ b/Engine/Source/Runtime/Experimental/Chaos/Private/GeometryCollection/Facades/CollectionVertexBoneWeightsFacade.cpp
@@ -24,8 +24,7 @@ namespace GeometryCollection::Facades
 	const FName FVertexBoneWeightsFacade::GeometryLODAttributeName = "GeometryLOD";
 
 	// Deprecated
-	const FName FVertexBoneWeightsFacade::DeprecatedBoneWeightsAttributeName = "BoneWeight";
-	const FName FVertexBoneWeightsFacade::DeprecatedBoneIndicesAttributeName = "BoneWeightIndex";
+	const FName FVertexBoneWeightsFacade::DeprecatedBoneIndicesAttributeName = "BoneWeightsIndex";
 	const FName FVertexBoneWeightsFacade::DeprecatedKinematicFlagAttributeName = "Kinematic";
 	
 	FVertexBoneWeightsFacade::FVertexBoneWeightsFacade(FManagedArrayCollection& InCollection, const bool bInInternalWeights)
@@ -93,6 +92,12 @@ namespace GeometryCollection::Facades
 		return BoneIndicesAttribute.IsValid() && BoneWeightsAttribute.IsValid() && KinematicWeightAttribute.IsValid() && SkeletalMeshAttribute.IsValid() && GeometryLODAttribute.IsValid();
 	}
 
+	bool FVertexBoneWeightsFacade::HasValidBoneIndicesAndWeights() const
+	{
+		return (ConstCollection.HasAttribute(DeprecatedBoneIndicesAttributeName, FGeometryCollection::VerticesGroup) 
+			|| BoneIndicesAttribute.IsValid()) && BoneWeightsAttribute.IsValid();
+	}
+
 	void FVertexBoneWeightsFacade::ModifyGeometryBinding(const int32 GeometryIndex, const TObjectPtr<UObject>& SkeletalMesh, const int32 GeometryLOD)
 	{
 		TManagedArray< TObjectPtr<UObject> >& SkeletalMeshes = SkeletalMeshAttribute.Modify();
@@ -231,26 +236,12 @@ namespace GeometryCollection::Facades
 
 	const TManagedArray< TArray<float> >* FVertexBoneWeightsFacade::FindBoneWeights()  const
 	{
-		if(ConstCollection.HasAttribute(DeprecatedBoneWeightsAttributeName, FGeometryCollection::VerticesGroup))
-		{
-			return ConstCollection.FindAttribute<TArray<float>>(DeprecatedBoneWeightsAttributeName, FGeometryCollection::VerticesGroup);
-		}
-		else
-		{
-			return BoneWeightsAttribute.Find();
-		}
+		return BoneWeightsAttribute.Find();
 	}
 	
 	const TManagedArray< TArray<float> >& FVertexBoneWeightsFacade::GetBoneWeights() const
 	{
-		if(ConstCollection.HasAttribute(DeprecatedBoneWeightsAttributeName, FGeometryCollection::VerticesGroup))
-		{
-			return ConstCollection.GetAttribute<TArray<float>>(DeprecatedBoneWeightsAttributeName, FGeometryCollection::VerticesGroup);
-		}
-		else
-		{
-			return BoneWeightsAttribute.Get();
-		}
+		return BoneWeightsAttribute.Get();
 	}
 
 	//
diff --git a/Engine/Source/Runtime/Experimental/Chaos/Public/GeometryCollection/Facades/CollectionVertexBoneWeightsFacade.h b/Engine/Source/Runtime/Experimental/Chaos/Public/GeometryCollection/Facades/CollectionVertexBoneWeightsFacade.h
--- a/Engine/Source/Runtime/Experimental/Chaos/Public/GeometryCollection/Facades/CollectionVertexBoneWeightsFacade.h
+++ b/Engine/Source/Runtime/Experimental/Chaos/Public/GeometryCollection/Facades/CollectionVertexBoneWeightsFacade.h
@@ -54,6 +54,9 @@ namespace GeometryCollection::Facades
 		/** Is the Facade defined on the collection? */
 		CHAOS_API bool IsValid() const;
 
+		/** Does the collection have bone weight and indices defined? */
+		CHAOS_API bool HasValidBoneIndicesAndWeights() const;
+
 		/** Add bone weight based on the kinematic bindings. */
 		CHAOS_API void AddBoneWeightsFromKinematicBindings();
 
@@ -119,7 +122,6 @@ namespace GeometryCollection::Facades
 
 		/** Deprecated names for attributes */
 		static const FName DeprecatedBoneIndicesAttributeName;
-		static const FName DeprecatedBoneWeightsAttributeName;
 		static const FName DeprecatedKinematicFlagAttributeName;
 		
 		const FManagedArrayCollection& ConstCollection;
