[ChaosVD] Fix for an issue where particles are not being correctly unloaded when they are destroyed.

There was a bug in how we processed the particle destroyed event when we collapsed delta frames data. This is the process done each time you scrub to another frame manually with a delta of +-1 frame.

Now instead of trying to keep a collapsed set of particle destroyed id, we process the particle destruction events in place. this is more costly, but this work is done mostly during load in a worker thread. The same code runs in the game thread during scrubbing but the number of frames it needs to process is low as we generate "fake" key frames during load as an optimization, so the perf impact cost is mitigated.

Also, this CL fixes an issue where hidden particles were still being draw by the particle data visualizer (the one that handles all debug draw like visualizations for particle data).

diff --git a/Engine/Plugins/ChaosVD/Source/ChaosVD/Private/Visualizers/ChaosVDParticleDataComponentVisualizer.cpp b/Engine/Plugins/ChaosVD/Source/ChaosVD/Private/Visualizers/ChaosVDParticleDataComponentVisualizer.cpp
--- a/Engine/Plugins/ChaosVD/Source/ChaosVD/Private/Visualizers/ChaosVDParticleDataComponentVisualizer.cpp
+++ b/Engine/Plugins/ChaosVD/Source/ChaosVD/Private/Visualizers/ChaosVDParticleDataComponentVisualizer.cpp
@@ -257,7 +257,10 @@ void FChaosVDParticleDataComponentVisualizer::DrawVisualizationForParticleData(c
 		return;
 	}
 
-
+	if (!InParticleInstance.IsVisible())
+	{
+		return;
+	}
 
 	FBox ParticleBounds = InParticleInstance.GetBoundingBox();
 
diff --git a/Engine/Source/Runtime/Experimental/ChaosVDData/Private/ChaosVDRecording.cpp b/Engine/Source/Runtime/Experimental/ChaosVDData/Private/ChaosVDRecording.cpp
--- a/Engine/Source/Runtime/Experimental/ChaosVDData/Private/ChaosVDRecording.cpp
+++ b/Engine/Source/Runtime/Experimental/ChaosVDData/Private/ChaosVDRecording.cpp
@@ -170,7 +170,7 @@ int32 FChaosVDRecording::FindFirstSolverKeyFrameNumberFromFrame_AssumesLocked(in
 		const int32 IndexFound = Algo::LowerBound(KeyFrameNumbers, StartFrameNumber);
 		
 		// If StartFrameNumber is larger than the last keyframe recorded
-		// IndexFound will be outside of the array's bounds. In that case we want to use the last key frame available
+		// IndexFound will be outside the array's bounds. In that case we want to use the last key frame available
 		if (IndexFound >= KeyFrameNumbers.Num())
 		{
 			return KeyFrameNumbers.Last();
@@ -394,25 +394,28 @@ void FChaosVDRecording::CollapseSolverFramesRange_AssumesLocked(int32 SolverID,
 	
 	for (int32 CurrentFrameNumber = StartFrame; CurrentFrameNumber <= EndFrame; CurrentFrameNumber++)
 	{
+		auto ProcessDestroyedParticles = [](const TSet<int32>& InDestroyedParticleIDs, TMap<int32, TSharedPtr<FChaosVDParticleDataWrapper>>& InOutParticlesMap)
+		{
+			for (int32 DestroyedPArticleID : InDestroyedParticleIDs)
+			{
+				InOutParticlesMap.Remove(DestroyedPArticleID);
+			}
+		};
+	
 		const bool bRequestingKeyFrameOnly = CurrentFrameNumber == StartFrame;
 		if (const FChaosVDSolverFrameData* SolverFrameData = GetSolverFrameData_AssumesLocked(SolverID, CurrentFrameNumber, bRequestingKeyFrameOnly))
 		{
-			OutCollapsedFrameData.ParticlesDestroyedIDs.Append(SolverFrameData->ParticlesDestroyedIDs);
-
 			if (SolverFrameData->SolverSteps.Num() > 0)
 			{
-				for (const FChaosVDFrameStageData& StepData : SolverFrameData->SolverSteps)
+				for (const FChaosVDFrameStageData& SolverStageData : SolverFrameData->SolverSteps)
 				{
-					for (const TSharedPtr<FChaosVDParticleDataWrapper>& ParticleData : StepData.RecordedParticlesData)
+					for (const TSharedPtr<FChaosVDParticleDataWrapper>& ParticleData : SolverStageData.RecordedParticlesData)
 					{
 						if (!ParticleData)
 						{
 							continue;
 						}
 
-						// The index could have been re-used after the particle was destroyed
-						OutCollapsedFrameData.ParticlesDestroyedIDs.Remove(ParticleData->ParticleIndex);
-
 						if (TSharedPtr<FChaosVDParticleDataWrapper>* FoundParticleData = ParticlesOnCurrentGeneratedKeyframe.Find(ParticleData->ParticleIndex))
 						{
 							(*FoundParticleData) = ParticleData;
@@ -421,10 +424,16 @@ void FChaosVDRecording::CollapseSolverFramesRange_AssumesLocked(int32 SolverID,
 						{
 							ParticlesOnCurrentGeneratedKeyframe.Add(ParticleData->ParticleIndex, ParticleData);
 						}
+
+						// Remove any particle data that was remoted in this stage
+						ProcessDestroyedParticles(SolverStageData.ParticlesDestroyedIDs, ParticlesOnCurrentGeneratedKeyframe);	
 					}
 				}
 			}
 
+			// Remove any particle data that was remove outside a recorded stage
+			ProcessDestroyedParticles(SolverFrameData->ParticlesDestroyedIDs, ParticlesOnCurrentGeneratedKeyframe);
+
 			// If this is the end frame, Copy all the "metadata" for the generated frame, and generate the Solver step
 			if (CurrentFrameNumber == EndFrame)
 			{
