[ChaosVD] Fixing an unitialized property in FRelayConnectionInfo

diff --git a/Engine/Source/Runtime/Experimental/ChaosSolverEngine/Public/Chaos/ChaosVDTraceRelayTransport.h b/Engine/Source/Runtime/Experimental/ChaosSolverEngine/Public/Chaos/ChaosVDTraceRelayTransport.h
--- a/Engine/Source/Runtime/Experimental/ChaosSolverEngine/Public/Chaos/ChaosVDTraceRelayTransport.h
+++ b/Engine/Source/Runtime/Experimental/ChaosSolverEngine/Public/Chaos/ChaosVDTraceRelayTransport.h
@@ -20,7 +20,7 @@ struct FRelayConnectionInfo
 
 	/** Relay port */
 	UPROPERTY()
-	uint16 Port;
+	uint16 Port = 0;
 
 	/** Encoded SSL certificate */
 	UPROPERTY()
