Geometry collection : fix local field not properly working with partial clusters when using minimal resolution

diff --git a/Engine/Source/Runtime/Experimental/Chaos/Private/PhysicsProxy/GeometryCollectionPhysicsProxy.cpp b/Engine/Source/Runtime/Experimental/Chaos/Private/PhysicsProxy/GeometryCollectionPhysicsProxy.cpp
--- a/Engine/Source/Runtime/Experimental/Chaos/Private/PhysicsProxy/GeometryCollectionPhysicsProxy.cpp
+++ b/Engine/Source/Runtime/Experimental/Chaos/Private/PhysicsProxy/GeometryCollectionPhysicsProxy.cpp
@@ -2325,16 +2325,38 @@ void FGeometryCollectionPhysicsProxy::GetRelevantParticleHandles(
 
 		for (FClusterHandle* ClusterHandle : SolverParticleHandles)
 		{
-			if (ClusterHandle && !ClusterHandle->Disabled())
+			if (ClusterHandle)
 			{
-				Handles.Add(ClusterHandle);
-				if (ClusterHandle->ClusterIds().NumChildren)
+				FClusterHandle* HandleToAdd = nullptr;
+				if (ClusterHandle->Disabled())
+				{
+					if (FClusterHandle* ParentHandle = ClusterHandle->Parent())
+					{
+						bool bHasPartialClusterParent = ParentHandle->InternalCluster() && !ParentHandle->Disabled() && (ParentHandle->PhysicsProxy() == ClusterHandle->PhysicsProxy());
+						if (bHasPartialClusterParent)
+						{
+							HandleToAdd = ParentHandle;
+						}
+					}
+				}
+				else
+				{
+					HandleToAdd = ClusterHandle;
+				}
+				if (HandleToAdd)
 				{
-					if (ClusterMap.Contains(ClusterHandle))
+					Handles.Add(HandleToAdd);
+					if (HandleToAdd->ClusterIds().NumChildren)
 					{
-						for (Chaos::TPBDRigidParticleHandle<Chaos::FReal, 3> * Child : ClusterMap[ClusterHandle])
+						if (ClusterMap.Find(HandleToAdd))
 						{
-							Handles.Add(Child);
+							for (Chaos::FPBDRigidParticleHandle* ChildHandle : ClusterMap[HandleToAdd])
+							{
+								if (ChildHandle->Disabled())
+								{
+									Handles.Add(ChildHandle);
+								}
+							}
 						}
 					}
 				}
