Geometry collection : make replication code more robust to invalid indices

diff --git a/Engine/Source/Runtime/Experimental/Chaos/Public/PhysicsProxy/GeometryCollectionPhysicsProxy.h b/Engine/Source/Runtime/Experimental/Chaos/Public/PhysicsProxy/GeometryCollectionPhysicsProxy.h
--- a/Engine/Source/Runtime/Experimental/Chaos/Public/PhysicsProxy/GeometryCollectionPhysicsProxy.h
+++ b/Engine/Source/Runtime/Experimental/Chaos/Public/PhysicsProxy/GeometryCollectionPhysicsProxy.h
@@ -316,10 +316,13 @@ public:
 
 	FClusterHandle* GetParticle_Internal(int32 Index) const
 	{
-		const int32 ParticleIndex = FromTransformToParticleIndex[Index];
-		if (ParticleIndex != INDEX_NONE)
+		if (FromTransformToParticleIndex.IsValidIndex(Index))
 		{
-			return SolverParticleHandles[ParticleIndex];
+			const int32 ParticleIndex = FromTransformToParticleIndex[Index];
+			if (SolverParticleHandles.IsValidIndex(ParticleIndex))
+			{
+				return SolverParticleHandles[ParticleIndex];
+			}
 		}
 		return nullptr;
 	}
diff --git a/Engine/Source/Runtime/Experimental/GeometryCollectionEngine/Private/GeometryCollection/GeometryCollectionComponent.cpp b/Engine/Source/Runtime/Experimental/GeometryCollectionEngine/Private/GeometryCollection/GeometryCollectionComponent.cpp
--- a/Engine/Source/Runtime/Experimental/GeometryCollectionEngine/Private/GeometryCollection/GeometryCollectionComponent.cpp
+++ b/Engine/Source/Runtime/Experimental/GeometryCollectionEngine/Private/GeometryCollection/GeometryCollectionComponent.cpp
@@ -2810,9 +2810,7 @@ namespace
 			}
 #endif
 
-			FPBDRigidParticleHandle* OneOff = PhysicsProxy->GetParticle_Internal(ActivatedCluster.ActivatedIndex);
-
-			if (ensure(OneOff))
+			if (FPBDRigidParticleHandle* OneOff = PhysicsProxy->GetParticle_Internal(ActivatedCluster.ActivatedIndex))
 			{
 				if (FPBDRigidClusteredParticleHandle* ClusterParticle = OneOff->CastToClustered())
 				{
