Add concept of "turbulence ratio" to Chaos Cloth Aerodynamics to better simulate lift and drag underwater.

Lift/drag forces are now split with

TotalForce = (1 - TurbulenceRatio) * LinearForce + TurblenceRatio * QuadraticForce

where QuadraticForce is the previously defined force which scales like Velocity^2, and LinearForce is a new force which is just QuadraticForce / Speed.

This is a simple, artist controllable way to approximate the differences between low Reynolds number (slower, denser fluids) laminar flow vs high Reynolds number turbulent flow. The default values continue to be set for typical air flow scenarios (AirDensity = 1.225, TurbulenceRatio = 1). To simulate underwater under slower movement conditions, should set AirDensity = 1000, TurbulenceRatio = 0.

diff --git a/Engine/Plugins/ChaosClothAssetEditor/Source/ChaosClothAssetDataflowNodes/Private/ChaosClothAsset/SimulationAerodynamicsConfigNode.cpp b/Engine/Plugins/ChaosClothAssetEditor/Source/ChaosClothAssetDataflowNodes/Private/ChaosClothAsset/SimulationAerodynamicsConfigNode.cpp
--- a/Engine/Plugins/ChaosClothAssetEditor/Source/ChaosClothAssetDataflowNodes/Private/ChaosClothAsset/SimulationAerodynamicsConfigNode.cpp
+++ b/Engine/Plugins/ChaosClothAssetEditor/Source/ChaosClothAssetDataflowNodes/Private/ChaosClothAsset/SimulationAerodynamicsConfigNode.cpp
@@ -29,6 +29,7 @@ void FChaosClothAssetSimulationAerodynamicsConfigNode::AddProperties(FPropertyHe
 	PropertyHelper.SetProperty(this, &FluidDensity);
 	PropertyHelper.SetPropertyEnum(this, &WindVelocitySpace);
 	PropertyHelper.SetProperty(this, &WindVelocity);
+	PropertyHelper.SetProperty(this, &TurbulenceRatio);
 	
 	PropertyHelper.SetSolverPropertyWeighted(FName(TEXT("Drag")), Drag, [](
 				const UE::Chaos::ClothAsset::FCollectionClothFacade& ClothFacade)-> float
diff --git a/Engine/Plugins/ChaosClothAssetEditor/Source/ChaosClothAssetDataflowNodes/Public/ChaosClothAsset/SimulationAerodynamicsConfigNode.h b/Engine/Plugins/ChaosClothAssetEditor/Source/ChaosClothAssetDataflowNodes/Public/ChaosClothAsset/SimulationAerodynamicsConfigNode.h
--- a/Engine/Plugins/ChaosClothAssetEditor/Source/ChaosClothAssetDataflowNodes/Public/ChaosClothAsset/SimulationAerodynamicsConfigNode.h
+++ b/Engine/Plugins/ChaosClothAssetEditor/Source/ChaosClothAssetDataflowNodes/Public/ChaosClothAsset/SimulationAerodynamicsConfigNode.h
@@ -21,7 +21,7 @@ public:
 	 * The fluid density is given in kg/m^3.
 	 * Air density is considered to be around 1.225 kg/m^3 in average atmospheric conditions.
 	 */
-	UPROPERTY(EditAnywhere, Category = "Aerodynamics", Meta = (UIMin = "0", ClampMin = "0", ClampMax = "10", InteractorName = "FluidDensity"))
+	UPROPERTY(EditAnywhere, Category = "Aerodynamics", Meta = (UIMin = "0", UIMax = "10", ClampMin = "0", ClampMax = "10000", InteractorName = "FluidDensity"))
 	float FluidDensity = 1.225f;
 
 	/**
@@ -37,6 +37,14 @@ public:
 	UPROPERTY(EditAnywhere, Category = "Aerodynamics", Meta = (UIMin = "0", UIMax = "10", InteractorName = "WindVelocity"))
 	FVector3f WindVelocity = { 0.f, 0.f, 0.f };
 
+	/**
+	 * Ratio of aerodynamic forces applied as turbulent rather than laminar flow.
+	 * Typically, you will want to use turbulent flow for wind and laminar flow for swimming underwater (unless the water is flowing very rapidly).
+	 * Turbulent forces scale like velocity squared. Laminar forces scale like velocity.
+	 */
+	UPROPERTY(EditAnywhere, Category = "Aerodynamics", Meta = (ClampMin = "0", ClampMax = "1", InteractorName = "TurbulenceRatio"))
+	float TurbulenceRatio = 1.f;
+
 	/**
 	 * The aerodynamic coefficient of drag applying on each particle.
 	 * When "Outer Drag" is enabled, this acts as the "Inner Drag", i.e., drag applied when the air velocity is
diff --git a/Engine/Source/Runtime/Experimental/Chaos/Private/Chaos/VelocityField.cpp b/Engine/Source/Runtime/Experimental/Chaos/Private/Chaos/VelocityField.cpp
--- a/Engine/Source/Runtime/Experimental/Chaos/Private/Chaos/VelocityField.cpp
+++ b/Engine/Source/Runtime/Experimental/Chaos/Private/Chaos/VelocityField.cpp
@@ -137,6 +137,11 @@ void FVelocityAndPressureField::SetProperties(
 		Rho = (FSolverReal)FMath::Max(GetFluidDensity(PropertyCollection), 0.f) / FMath::Cube(WorldScale);
 	}
 
+	if (IsTurbulenceRatioMutable(PropertyCollection))
+	{
+		TurbulenceRatio = FMath::Clamp(GetTurbulenceRatio(PropertyCollection), 0.f, 1.f);
+	}
+
 	if (IsPressureMutable(PropertyCollection))
 	{
 		const FSolverVec2 WeightedValue(GetWeightedFloatPressure(PropertyCollection));
@@ -222,6 +227,7 @@ void FVelocityAndPressureField::SetGeometry(
 	OuterLiftIndex = FOuterLiftIndex(PropertyCollection);
 	FluidDensityIndex = FFluidDensityIndex(PropertyCollection);
 	PressureIndex = FPressureIndex(PropertyCollection);
+	TurbulenceRatioIndex = FTurbulenceRatioIndex(PropertyCollection);
 
 	// Reset geometry, properties, and weight maps
 	SetGeometry(TriangleMesh);
@@ -239,6 +245,7 @@ void FVelocityAndPressureField::SetGeometry(
 		(FSolverReal)GetFluidDensity(PropertyCollection, 0.f),
 		FSolverVec2(GetWeightedFloatPressure(PropertyCollection, 0.f)),
 		WorldScale);
+	TurbulenceRatio = FMath::Clamp(GetTurbulenceRatio(PropertyCollection, DefaultTurbulenceRatio), 0.f, 1.f);
 	SetMultipliers(PropertyCollection, Weightmaps);
 	ResetColor();
 }
@@ -441,6 +448,7 @@ void FVelocityAndPressureField::UpdateForces(const FSolverParticles& InParticles
 					(const ispc::FVector3f*)InParticles.XArray().GetData(),
 					(const ispc::FVector3f&)Velocity,
 					QuarterRho,
+					TurbulenceRatio,
 					(FSolverReal)Drag,
 					(FSolverReal)OuterDrag,
 					(FSolverReal)Lift,
@@ -457,6 +465,7 @@ void FVelocityAndPressureField::UpdateForces(const FSolverParticles& InParticles
 					(const ispc::FVector3f*)InParticles.XArray().GetData(),
 					(const ispc::FVector3f&)Velocity,
 					QuarterRho,
+					TurbulenceRatio,
 					(FSolverReal)Drag,
 					(FSolverReal)OuterDrag,
 					(FSolverReal)Lift,
@@ -509,6 +518,7 @@ void FVelocityAndPressureField::UpdateForces(const FSolverParticles& InParticles
 					(const ispc::FVector3f*)InParticles.XArray().GetData(),
 					(const ispc::FVector3f&)Velocity,
 					QuarterRho,
+					TurbulenceRatio,
 					bDragHasMap,
 					reinterpret_cast<const ispc::FVector2f&>(Drag.GetOffsetRange()),
 					bDragHasMap ? Drag.GetMapValues().GetData() : nullptr,
@@ -535,6 +545,7 @@ void FVelocityAndPressureField::UpdateForces(const FSolverParticles& InParticles
 					(const ispc::FVector3f*)InParticles.XArray().GetData(),
 					(const ispc::FVector3f&)Velocity,
 					QuarterRho,
+					TurbulenceRatio,
 					bDragHasMap,
 					reinterpret_cast<const ispc::FVector2f&>(Drag.GetOffsetRange()),
 					bDragHasMap ? Drag.GetMapValues().GetData() : nullptr,
@@ -617,6 +628,7 @@ void FVelocityAndPressureField::Apply(FSolverParticlesRange& InParticles, const
 						(const ispc::FVector3f*)InParticles.XArray().GetData(),
 						(const ispc::FVector3f&)Velocity,
 						QuarterRho,
+						TurbulenceRatio,
 						(FSolverReal)Drag,
 						(FSolverReal)OuterDrag,
 						(FSolverReal)Lift,
@@ -639,6 +651,7 @@ void FVelocityAndPressureField::Apply(FSolverParticlesRange& InParticles, const
 						(const ispc::FVector3f*)InParticles.XArray().GetData(),
 						(const ispc::FVector3f&)Velocity,
 						QuarterRho,
+						TurbulenceRatio,
 						(FSolverReal)Drag,
 						(FSolverReal)OuterDrag,
 						(FSolverReal)Lift,
@@ -704,6 +717,7 @@ void FVelocityAndPressureField::Apply(FSolverParticlesRange& InParticles, const
 						(const ispc::FVector3f*)InParticles.XArray().GetData(),
 						(const ispc::FVector3f&)Velocity,
 						QuarterRho,
+						TurbulenceRatio,
 						bDragHasMap,
 						reinterpret_cast<const ispc::FVector2f&>(Drag.GetOffsetRange()),
 						bDragHasMap ? Drag.GetMapValues().GetData() : nullptr,
@@ -736,6 +750,7 @@ void FVelocityAndPressureField::Apply(FSolverParticlesRange& InParticles, const
 						(const ispc::FVector3f*)InParticles.XArray().GetData(),
 						(const ispc::FVector3f&)Velocity,
 						QuarterRho,
+						TurbulenceRatio,
 						bDragHasMap,
 						reinterpret_cast<const ispc::FVector2f&>(Drag.GetOffsetRange()),
 						bDragHasMap ? Drag.GetMapValues().GetData() : nullptr,
diff --git a/Engine/Source/Runtime/Experimental/Chaos/Private/Chaos/VelocityField.ispc b/Engine/Source/Runtime/Experimental/Chaos/Private/Chaos/VelocityField.ispc
--- a/Engine/Source/Runtime/Experimental/Chaos/Private/Chaos/VelocityField.ispc
+++ b/Engine/Source/Runtime/Experimental/Chaos/Private/Chaos/VelocityField.ispc
@@ -19,6 +19,7 @@ export void UpdateField(uniform FVector3f Forces[],
 								const uniform FVector3f X[],
 								const uniform FVector3f& Velocity,
 								const uniform float QuarterRho,
+								const uniform float TurbulenceRatio,
 								const uniform float CdI,
 								const uniform float CdO,
 								const uniform float ClI,
@@ -61,6 +62,7 @@ export void UpdateField(uniform FVector3f Forces[],
 		// Set the aerodynamic forces
 		const float VDotN = VectorDot(RelVelocity, N);
 		const float VSquare = VectorDot(RelVelocity, RelVelocity);
+		const float TurbulenceFactor = VSquare > FLOAT_SMALL_NUMBER ? ((1 - TurbulenceRatio) / sqrt(VSquare) + TurbulenceRatio) : 1;
 
 #if HW_GATHER_SUPPORTED == 1
 		if(ElementIndex + programCount < NumElements)
@@ -72,7 +74,7 @@ export void UpdateField(uniform FVector3f Forces[],
 		const FVector3f EvenForce = (CdI - ClI) * VDotN * RelVelocity + ClI * VSquare * N;
 		const FVector3f OddForce = (ClO - CdO) * VDotN * RelVelocity - ClO * VSquare * N;
 
-		const FVector3f Force = QuarterRho * DoubleArea * VectorSelect(VDotN >= 0, EvenForce, OddForce) + DoubleArea * 0.5 * Cp * N;
+		const FVector3f Force = QuarterRho * DoubleArea * TurbulenceFactor * VectorSelect(VDotN >= 0, EvenForce, OddForce) + DoubleArea * 0.5 * Cp * N;
 
 		VectorStore(&Forces[Offset], Force);
 
@@ -86,6 +88,7 @@ export void UpdateFieldWithWeightMaps(uniform FVector3f Forces[],
 								const uniform FVector3f X[],
 								const uniform FVector3f& Velocity,
 								const uniform float QuarterRho,
+								const uniform float TurbulenceRatio,
 								const uniform bool bDragHasMap,
 								const uniform FVector2f& DragOffsetRange,
 								const uniform float DragMapValues[],
@@ -137,6 +140,7 @@ export void UpdateFieldWithWeightMaps(uniform FVector3f Forces[],
 		// Set the aerodynamic forces
 		const float VDotN = VectorDot(RelVelocity, N);
 		const float VSquare = VectorDot(RelVelocity, RelVelocity);
+		const float TurbulenceFactor = VSquare > FLOAT_SMALL_NUMBER ? ((1 - TurbulenceRatio) / sqrt(VSquare) + TurbulenceRatio) : 1;
 
 		const float CdI = bDragHasMap ? DragOffsetRange.V[0] + DragMapValues[ElementIndex] * DragOffsetRange.V[1] : DragOffsetRange.V[0];
 		const float CdO = bOuterDragHasMap ? OuterDragOffsetRange.V[0] + OuterDragMapValues[ElementIndex] * OuterDragOffsetRange.V[1] : OuterDragOffsetRange.V[0];
@@ -153,8 +157,8 @@ export void UpdateFieldWithWeightMaps(uniform FVector3f Forces[],
 
 		const FVector3f EvenForce = (CdI - ClI) * VDotN * RelVelocity + ClI * VSquare * N;
 		const FVector3f OddForce = (ClO - CdO) * VDotN * RelVelocity - ClO * VSquare * N;
-
-		const FVector3f Force = QuarterRho * DoubleArea * VectorSelect(VDotN >= 0, EvenForce, OddForce) + DoubleArea * 0.5 * Cp * N;
+		
+		const FVector3f Force = QuarterRho * DoubleArea * TurbulenceFactor * VectorSelect(VDotN >= 0, EvenForce, OddForce) + DoubleArea * 0.5 * Cp * N;
 
 		VectorStore(&Forces[Offset], Force);
 
@@ -168,6 +172,7 @@ export void UpdateFieldAndClampVelocity(uniform FVector3f Forces[],
 								const uniform FVector3f X[],
 								const uniform FVector3f& Velocity,
 								const uniform float QuarterRho,
+								const uniform float TurbulenceRatio,
 								const uniform float CdI,
 								const uniform float CdO,
 								const uniform float ClI,
@@ -218,6 +223,7 @@ export void UpdateFieldAndClampVelocity(uniform FVector3f Forces[],
 		// Set the aerodynamic forces
 		const float VDotN = VectorDot(RelVelocity, N);
 		const float VSquare = VectorDot(RelVelocity, RelVelocity);
+		const float TurbulenceFactor = VSquare > FLOAT_SMALL_NUMBER ? ((1 - TurbulenceRatio) / sqrt(VSquare) + TurbulenceRatio) : 1;
 
 #if HW_GATHER_SUPPORTED == 1
 		if(ElementIndex + programCount < NumElements)
@@ -229,7 +235,7 @@ export void UpdateFieldAndClampVelocity(uniform FVector3f Forces[],
 		const FVector3f EvenForce = (CdI - ClI) * VDotN * RelVelocity + ClI * VSquare * N;
 		const FVector3f OddForce = (ClO - CdO) * VDotN * RelVelocity - ClO * VSquare * N;
 		
-		const FVector3f Force = QuarterRho * DoubleArea * VectorSelect(VDotN >= 0, EvenForce, OddForce) + DoubleArea * 0.5 * Cp * N;
+		const FVector3f Force = QuarterRho * DoubleArea * TurbulenceFactor * VectorSelect(VDotN >= 0, EvenForce, OddForce) + DoubleArea * 0.5 * Cp * N;
 
 		VectorStore(&Forces[Offset], Force);
 
@@ -243,6 +249,7 @@ export void UpdateFieldWithWeightMapsAndClampVelocity(uniform FVector3f Forces[]
 								const uniform FVector3f X[],
 								const uniform FVector3f& Velocity,
 								const uniform float QuarterRho,
+								const uniform float TurbulenceRatio,
 								const uniform bool bDragHasMap,
 								const uniform FVector2f& DragOffsetRange,
 								const uniform float DragMapValues[],
@@ -302,6 +309,7 @@ export void UpdateFieldWithWeightMapsAndClampVelocity(uniform FVector3f Forces[]
 		// Set the aerodynamic forces
 		const float VDotN = VectorDot(RelVelocity, N);
 		const float VSquare = VectorDot(RelVelocity, RelVelocity);
+		const float TurbulenceFactor = VSquare > FLOAT_SMALL_NUMBER ? ((1 - TurbulenceRatio) / sqrt(VSquare) + TurbulenceRatio) : 1;
 		
 		const float CdI = bDragHasMap ? DragOffsetRange.V[0] + DragMapValues[ElementIndex] * DragOffsetRange.V[1] : DragOffsetRange.V[0];
 		const float CdO = bOuterDragHasMap ? OuterDragOffsetRange.V[0] + OuterDragMapValues[ElementIndex] * OuterDragOffsetRange.V[1] : OuterDragOffsetRange.V[0];
@@ -318,8 +326,8 @@ export void UpdateFieldWithWeightMapsAndClampVelocity(uniform FVector3f Forces[]
 
 		const FVector3f EvenForce = (CdI - ClI) * VDotN * RelVelocity + ClI * VSquare * N;
 		const FVector3f OddForce = (ClO - CdO) * VDotN * RelVelocity - ClO * VSquare * N;
-
-		const FVector3f Force = QuarterRho * DoubleArea * VectorSelect(VDotN >= 0, EvenForce, OddForce) + DoubleArea * 0.5 * Cp * N;
+		
+		const FVector3f Force = QuarterRho * DoubleArea * TurbulenceFactor * VectorSelect(VDotN >= 0, EvenForce, OddForce) + DoubleArea * 0.5 * Cp * N;
 
 		VectorStore(&Forces[Offset], Force);
 
@@ -334,6 +342,7 @@ export void UpdateAndApplyVelocityField(uniform FVector3f Acceleration[],
 								const uniform FVector3f X[],
 								const uniform FVector3f& Velocity,
 								const uniform float QuarterRho,
+								const uniform float TurbulenceRatio,
 								const uniform float CdI,
 								const uniform float CdO,
 								const uniform float ClI,
@@ -368,11 +377,12 @@ export void UpdateAndApplyVelocityField(uniform FVector3f Acceleration[],
 		// Set the aerodynamic forces
 		const float VDotN = VectorDot(RelVelocity, N);
 		const float VSquare = VectorDot(RelVelocity, RelVelocity);
+		const float TurbulenceFactor = VSquare > FLOAT_SMALL_NUMBER ? ((1 - TurbulenceRatio) / sqrt(VSquare) + TurbulenceRatio) : 1;
 		
 		const FVector3f EvenForce = (CdI - ClI) * VDotN * RelVelocity + ClI * VSquare * N;
 		const FVector3f OddForce = (ClO - CdO) * VDotN * RelVelocity - ClO * VSquare * N;
 		
-		const FVector3f Force = QuarterRho * DoubleArea * VectorSelect(VDotN >= 0, EvenForce, OddForce) + DoubleArea * 0.5 * Cp * N;
+		const FVector3f Force = QuarterRho * DoubleArea * TurbulenceFactor * VectorSelect(VDotN >= 0, EvenForce, OddForce) + DoubleArea * 0.5 * Cp * N;
 		
 		#pragma ignore warning(perf)
 		const float InvM0 = InvM[Index0];
@@ -398,6 +408,7 @@ export void UpdateAndClampVelocityAndApplyVelocityField(uniform FVector3f Accele
 								const uniform FVector3f X[],
 								const uniform FVector3f& Velocity,
 								const uniform float QuarterRho,
+								const uniform float TurbulenceRatio,
 								const uniform float CdI,
 								const uniform float CdO,
 								const uniform float ClI,
@@ -443,8 +454,9 @@ export void UpdateAndClampVelocityAndApplyVelocityField(uniform FVector3f Accele
 		
 		const FVector3f EvenForce = (CdI - ClI) * VDotN * RelVelocity + ClI * VSquare * N;
 		const FVector3f OddForce = (ClO - CdO) * VDotN * RelVelocity - ClO * VSquare * N;
+		const float TurbulenceFactor = VSquare > FLOAT_SMALL_NUMBER ? ((1 - TurbulenceRatio) / sqrt(VSquare) + TurbulenceRatio) : 1;
 		
-		const FVector3f Force = QuarterRho * DoubleArea * VectorSelect(VDotN >= 0, EvenForce, OddForce) + DoubleArea * 0.5 * Cp * N;
+		const FVector3f Force = QuarterRho * DoubleArea * TurbulenceFactor * VectorSelect(VDotN >= 0, EvenForce, OddForce) + DoubleArea * 0.5 * Cp * N;
 		
 		#pragma ignore warning(perf)
 		const float InvM0 = InvM[Index0];
@@ -471,6 +483,7 @@ export void UpdateAndApplyVelocityFieldWithWeightMaps(uniform FVector3f Accelera
 								const uniform FVector3f X[],
 								const uniform FVector3f& Velocity,
 								const uniform float QuarterRho,
+								const uniform float TurbulenceRatio,
 								const uniform bool bDragHasMap,
 								const uniform FVector2f& DragOffsetRange,
 								const uniform float DragMapValues[],
@@ -515,6 +528,7 @@ export void UpdateAndApplyVelocityFieldWithWeightMaps(uniform FVector3f Accelera
 		// Set the aerodynamic forces
 		const float VDotN = VectorDot(RelVelocity, N);
 		const float VSquare = VectorDot(RelVelocity, RelVelocity);
+		const float TurbulenceFactor = VSquare > FLOAT_SMALL_NUMBER ? ((1 - TurbulenceRatio) / sqrt(VSquare) + TurbulenceRatio) : 1;
 		
 		const float CdI = bDragHasMap ? DragOffsetRange.V[0] + DragMapValues[ElementIndex] * DragOffsetRange.V[1] : DragOffsetRange.V[0];
 		const float CdO = bOuterDragHasMap ? OuterDragOffsetRange.V[0] + OuterDragMapValues[ElementIndex] * OuterDragOffsetRange.V[1] : OuterDragOffsetRange.V[0];
@@ -524,8 +538,8 @@ export void UpdateAndApplyVelocityFieldWithWeightMaps(uniform FVector3f Accelera
 		
 		const FVector3f EvenForce = (CdI - ClI) * VDotN * RelVelocity + ClI * VSquare * N;
 		const FVector3f OddForce = (ClO - CdO) * VDotN * RelVelocity - ClO * VSquare * N;
-
-		const FVector3f Force = QuarterRho * DoubleArea * VectorSelect(VDotN >= 0, EvenForce, OddForce) + DoubleArea * 0.5 * Cp * N;
+		
+		const FVector3f Force = QuarterRho * DoubleArea * TurbulenceFactor * VectorSelect(VDotN >= 0, EvenForce, OddForce) + DoubleArea * 0.5 * Cp * N;
 		
 		#pragma ignore warning(perf)
 		const float InvM0 = InvM[Index0];
@@ -551,6 +565,7 @@ export void UpdateAndClampVelocityAndApplyVelocityFieldWithWeightMaps(uniform FV
 								const uniform FVector3f X[],
 								const uniform FVector3f& Velocity,
 								const uniform float QuarterRho,
+								const uniform float TurbulenceRatio,
 								const uniform bool bDragHasMap,
 								const uniform FVector2f& DragOffsetRange,
 								const uniform float DragMapValues[],
@@ -603,6 +618,7 @@ export void UpdateAndClampVelocityAndApplyVelocityFieldWithWeightMaps(uniform FV
 		// Set the aerodynamic forces
 		const float VDotN = VectorDot(RelVelocity, N);
 		const float VSquare = VectorDot(RelVelocity, RelVelocity);
+		const float TurbulenceFactor = VSquare > FLOAT_SMALL_NUMBER ? ((1 - TurbulenceRatio) / sqrt(VSquare) + TurbulenceRatio) : 1;
 		
 		const float CdI = bDragHasMap ? DragOffsetRange.V[0] + DragMapValues[ElementIndex] * DragOffsetRange.V[1] : DragOffsetRange.V[0];
 		const float CdO = bOuterDragHasMap ? OuterDragOffsetRange.V[0] + OuterDragMapValues[ElementIndex] * OuterDragOffsetRange.V[1] : OuterDragOffsetRange.V[0];
@@ -613,7 +629,7 @@ export void UpdateAndClampVelocityAndApplyVelocityFieldWithWeightMaps(uniform FV
 		const FVector3f EvenForce = (CdI - ClI) * VDotN * RelVelocity + ClI * VSquare * N;
 		const FVector3f OddForce = (ClO - CdO) * VDotN * RelVelocity - ClO * VSquare * N;
 
-		const FVector3f Force = QuarterRho * DoubleArea * VectorSelect(VDotN >= 0, EvenForce, OddForce) + DoubleArea * 0.5 * Cp * N;
+		const FVector3f Force = QuarterRho * DoubleArea * TurbulenceFactor * VectorSelect(VDotN >= 0, EvenForce, OddForce) + DoubleArea * 0.5 * Cp * N;
 		
 		#pragma ignore warning(perf)
 		const float InvM0 = InvM[Index0];
diff --git a/Engine/Source/Runtime/Experimental/Chaos/Public/Chaos/VelocityField.h b/Engine/Source/Runtime/Experimental/Chaos/Public/Chaos/VelocityField.h
--- a/Engine/Source/Runtime/Experimental/Chaos/Public/Chaos/VelocityField.h
+++ b/Engine/Source/Runtime/Experimental/Chaos/Public/Chaos/VelocityField.h
@@ -22,6 +22,7 @@ public:
 	static constexpr FSolverReal MinCoefficient = (FSolverReal)0.;   // Applies to both drag and lift
 	static constexpr FSolverReal MaxCoefficient = (FSolverReal)10.;  //
 	static constexpr EChaosSoftsSimulationSpace DefaultWindVelocitySpace = EChaosSoftsSimulationSpace::WorldSpace;
+	static constexpr FSolverReal DefaultTurbulenceRatio = (FSolverReal)1.f;
 
 	static bool IsEnabled(const FCollectionPropertyConstFacade& PropertyCollection)
 	{
@@ -41,12 +42,14 @@ public:
 		, Pressure(GetWeightedFloatPressure(PropertyCollection, (FSolverReal)0.))
 		, Rho(FMath::Max(GetFluidDensity(PropertyCollection, (FSolverReal)0.), (FSolverReal)0.))
 		, QuarterRho(Rho * (FSolverReal)0.25f)
+		, TurbulenceRatio(FMath::Clamp(GetTurbulenceRatio(PropertyCollection, DefaultTurbulenceRatio), 0.f, 1.f))
 		, DragIndex(PropertyCollection)
 		, OuterDragIndex(PropertyCollection)
 		, LiftIndex(PropertyCollection)
 		, OuterLiftIndex(PropertyCollection)
 		, FluidDensityIndex(PropertyCollection)
 		, PressureIndex(PropertyCollection)
+		, TurbulenceRatioIndex(PropertyCollection)
 		, WindVelocityIndex(PropertyCollection)
 		, WindVelocitySpaceIndex(PropertyCollection)
 	{
@@ -65,12 +68,14 @@ public:
 		, Pressure(GetWeightedFloatPressure(PropertyCollection, (FSolverReal)0.)/WorldScale)
 		, Rho(FMath::Max(GetFluidDensity(PropertyCollection, (FSolverReal)0.)/FMath::Cube(WorldScale), (FSolverReal)0.))
 		, QuarterRho(Rho* (FSolverReal)0.25f)
+		, TurbulenceRatio(FMath::Clamp(GetTurbulenceRatio(PropertyCollection, DefaultTurbulenceRatio), 0.f, 1.f))
 		, DragIndex(PropertyCollection)
 		, OuterDragIndex(PropertyCollection)
 		, LiftIndex(PropertyCollection)
 		, OuterLiftIndex(PropertyCollection)
 		, FluidDensityIndex(PropertyCollection)
 		, PressureIndex(PropertyCollection)
+		, TurbulenceRatioIndex(PropertyCollection)
 		, WindVelocityIndex(PropertyCollection)
 		, WindVelocitySpaceIndex(PropertyCollection)
 	{
@@ -90,12 +95,14 @@ public:
 		, Pressure(FSolverVec2(0.))
 		, Rho((FSolverReal)0.)
 		, QuarterRho(Rho* (FSolverReal)0.25f)
+		, TurbulenceRatio(DefaultTurbulenceRatio)
 		, DragIndex(ForceInit)
 		, OuterDragIndex(ForceInit)
 		, LiftIndex(ForceInit)
 		, OuterLiftIndex(ForceInit)
 		, FluidDensityIndex(ForceInit)
 		, PressureIndex(ForceInit)
+		, TurbulenceRatioIndex(ForceInit)
 		, WindVelocityIndex(ForceInit)
 		, WindVelocitySpaceIndex(ForceInit)
 	{
@@ -145,17 +152,6 @@ public:
 		const FSolverVec2& Pressure = FSolverVec2::ZeroVector,
 		FSolverReal WorldScale = 1.f);
 
-	UE_DEPRECATED(5.5, "Use the SetProperties that includes OuterDrag and OuterLift")
-	void SetProperties(
-		const FSolverVec2& InDrag,
-		const FSolverVec2& InLift,
-		const FSolverReal InFluidDensity,
-		const FSolverVec2& InPressure = FSolverVec2::ZeroVector,
-		FSolverReal WorldScale = 1.f)
-	{
-		SetProperties(InDrag, InDrag, InLift, InLift, InFluidDensity, InPressure, WorldScale);
-	}
-
 	bool IsActive() const 
 	{ 
 		return Pressure.GetLow() != (FSolverReal)0. || Pressure.GetHigh() != (FSolverReal)0. ||
@@ -206,6 +202,7 @@ public:
 	UE_CHAOS_DECLARE_INDEXLESS_PROPERTYCOLLECTION_NAME(OuterLift, float);
 	UE_CHAOS_DECLARE_INDEXLESS_PROPERTYCOLLECTION_NAME(FluidDensity, float);
 	UE_CHAOS_DECLARE_INDEXLESS_PROPERTYCOLLECTION_NAME(Pressure, float);
+	UE_CHAOS_DECLARE_INDEXLESS_PROPERTYCOLLECTION_NAME(TurbulenceRatio, float);
 	UE_CHAOS_DECLARE_INDEXLESS_PROPERTYCOLLECTION_NAME(WindVelocity, FVector3f);
 	UE_CHAOS_DECLARE_INDEXLESS_PROPERTYCOLLECTION_NAME(WindVelocitySpace, int32);
 
@@ -249,7 +246,9 @@ private:
 		const FSolverReal VDotN = FSolverVec3::DotProduct(V, N);
 		const FSolverReal VSquare = FSolverVec3::DotProduct(V, V);
 
-		return QuarterRho * DoubleArea * (VDotN >= (FSolverReal)0. ?  // The flow can hit either side of the triangle, so the normal might need to be reversed
+		const FSolverReal TurbulenceFactor = VSquare > UE_SMALL_NUMBER ? ((1.f - TurbulenceRatio) / FMath::Sqrt(VSquare) + TurbulenceRatio) : 1.f;
+
+		return QuarterRho * DoubleArea * TurbulenceFactor * (VDotN >= (FSolverReal)0. ?  // The flow can hit either side of the triangle, so the normal might need to be reversed
 			(CdI - ClI) * VDotN * V + ClI * VSquare * N :
 			(ClO - CdO) * VDotN * V - ClO * VSquare * N) + DoubleArea * (FSolverReal)0.5 * Cp * N;
 	}
@@ -288,8 +287,9 @@ private:
 		// Set the aerodynamic forces
 		const FSolverReal VDotN = FSolverVec3::DotProduct(V, N);
 		const FSolverReal VSquare = FSolverVec3::DotProduct(V, V);
+		const FSolverReal TurbulenceFactor = VSquare > UE_SMALL_NUMBER ? ((1.f - TurbulenceRatio) / FMath::Sqrt(VSquare) + TurbulenceRatio) : 1.f;
 
-		return QuarterRho * DoubleArea * (VDotN >= (FSolverReal)0. ?  // The flow can hit either side of the triangle, so the normal might need to be reversed
+		return QuarterRho * DoubleArea * TurbulenceFactor * (VDotN >= (FSolverReal)0. ?  // The flow can hit either side of the triangle, so the normal might need to be reversed
 			(CdI - ClI) * VDotN * V + ClI * VSquare * N :
 			(ClO - CdO) * VDotN * V - ClO * VSquare * N) + DoubleArea * (FSolverReal)0.5 * Cp * N;
 	}
@@ -317,6 +317,7 @@ private:
 	FSolverVec3 Velocity;
 	FSolverReal Rho;
 	FSolverReal QuarterRho;
+	FSolverReal TurbulenceRatio = 1.f;
 
 	UE_CHAOS_DECLARE_INDEXED_PROPERTYCOLLECTION_NAME(Drag, float);
 	UE_CHAOS_DECLARE_INDEXED_PROPERTYCOLLECTION_NAME(OuterDrag, float);
@@ -324,6 +325,7 @@ private:
 	UE_CHAOS_DECLARE_INDEXED_PROPERTYCOLLECTION_NAME(OuterLift, float);
 	UE_CHAOS_DECLARE_INDEXED_PROPERTYCOLLECTION_NAME(FluidDensity, float);
 	UE_CHAOS_DECLARE_INDEXED_PROPERTYCOLLECTION_NAME(Pressure, float);
+	UE_CHAOS_DECLARE_INDEXED_PROPERTYCOLLECTION_NAME(TurbulenceRatio, float);
 	UE_CHAOS_DECLARE_INDEXED_PROPERTYCOLLECTION_NAME(WindVelocity, FVector3f);
 	UE_CHAOS_DECLARE_INDEXED_PROPERTYCOLLECTION_NAME(WindVelocitySpace, int32);
 };
