Minor improvements to HTTP I/O dispatcher cancellation.
- updated FHttpQueue::Cancel to not assume that a request must exist in a queued batch for a certain chunk hash. Existing ensures
is simply not correct.
- Updated CancelRequest to append CancelRequested flag and not require the flags to match. The current behavior can prevent removing the shared cancellation token.
- Updated ProcessHttpRequests to dequeue the next pending request AFTER sending the current request and entering the Tick loop. This can prevent sending an HTTP request that is already cancelled due to long response times or stalls in the tick loop.

#rnx

diff --git a/Engine/Source/Runtime/Experimental/IoStore/OnDemand/Private/OnDemandHttpIoDispacher.cpp b/Engine/Source/Runtime/Experimental/IoStore/OnDemand/Private/OnDemandHttpIoDispacher.cpp
--- a/Engine/Source/Runtime/Experimental/IoStore/OnDemand/Private/OnDemandHttpIoDispacher.cpp
+++ b/Engine/Source/Runtime/Experimental/IoStore/OnDemand/Private/OnDemandHttpIoDispacher.cpp
@@ -370,17 +370,14 @@ bool FHttpQueue::Cancel(FHttpRequest* Request)
 		return false;
 	}
 
+	// Note: Even if a batch exists for the given chunk hash, it doesn't necessarily mean this request is included in that batch.
 	FHttpRequestList* List = ByChunkKey.Find(Request->ChunkHash);
-	if (List == nullptr)
+	if ((List == nullptr) || (List->Remove(Request) == false))
 	{
 		ensure(Heap.Find(Request) == INDEX_NONE);
 		return false;
 	}
 
-	const bool bRemoved = List->Remove(Request);
-	ensure(bRemoved);
-	ensure(Request->Next == nullptr);
-
 	// If the request was the first for this chunk we push any of the remaining requests to the heap
 	if (int32 Index = Heap.Find(Request); Index != INDEX_NONE)
 	{
@@ -1016,11 +1013,7 @@ void FHttpDispatcher::CancelRequest(FIoHttpRequestHandle Handle)
 	{
 		return;
 	}
-
-	if (Request->Flags.TrySet(CurrentFlags | EHttpRequestFlags::CancelRequested) == false)
-	{
-		return;
-	}
+	Request->Flags.Add(EHttpRequestFlags::CancelRequested);
 
 	UE_LOG(LogHttpIoDispatcher, Verbose, TEXT("Cancelling request, SeqNo=%u"), Request->SeqNo);
 
@@ -1201,7 +1194,7 @@ bool FHttpDispatcher::TryReadFromCache(FHttpRequest* Request)
 				FOnDemandIoBackendStats::Get()->OnCacheGet(Request->Buffer.GetSize());
 				CompleteRequest(*Request, Request->Buffer, FIoHttpRange(), 200, true);
 			}
-			else if (Request->CacheStatus == EIoErrorCode::Cancelled)
+			else if ((Request->CacheStatus == EIoErrorCode::Cancelled) || Request->Flags.HasAny(EHttpRequestFlags::CancelRequested))
 			{
 				CompleteRequest(*Request, Request->Buffer, FIoHttpRange(), 0, false);
 			}
@@ -1267,7 +1260,7 @@ bool FHttpDispatcher::ProcessHttpRequests()
 		while (NextHttpRequest)
 		{
 			FHttpRequest* HttpRequest	= NextHttpRequest;
-			NextHttpRequest				= GetNextRequest();
+			NextHttpRequest				= nullptr;
 
 			ensure(!HttpRequest->Flags.HasAny(EHttpRequestFlags::HttpQueued | EHttpRequestFlags::CacheInflight));
 
@@ -1280,6 +1273,9 @@ bool FHttpDispatcher::ProcessHttpRequests()
 				// order to correctly reduce the pending count.
 				FOnDemandIoBackendStats::Get()->OnHttpCancel(HttpRequest->GetHttpRequestType());
 				CompleteRequest(*HttpRequest, HttpRequest->Buffer, FIoHttpRange(), 0, false);
+
+				check(NextHttpRequest == nullptr);
+				NextHttpRequest = GetNextRequest();
 				continue;
 			}
 
@@ -1362,10 +1358,8 @@ bool FHttpDispatcher::ProcessHttpRequests()
 				}
 			}
 
-			if (!NextHttpRequest)
-			{
-				NextHttpRequest = GetNextRequest();
-			}
+			check(NextHttpRequest == nullptr);
+			NextHttpRequest = GetNextRequest();
 		} // Inner
 
 		{
