Disable HTTP I/O dispatcher.

#rnx

diff --git a/Engine/Source/Runtime/Experimental/IoStore/OnDemand/Private/OnDemandIoStore.cpp b/Engine/Source/Runtime/Experimental/IoStore/OnDemand/Private/OnDemandIoStore.cpp
--- a/Engine/Source/Runtime/Experimental/IoStore/OnDemand/Private/OnDemandIoStore.cpp
+++ b/Engine/Source/Runtime/Experimental/IoStore/OnDemand/Private/OnDemandIoStore.cpp
@@ -53,7 +53,7 @@
 #endif
 
 ///////////////////////////////////////////////////////////////////////////////
-bool GHttpIoDispatcherEnabled = true;
+bool GHttpIoDispatcherEnabled = false;
 static FAutoConsoleVariableRef CVar_HttpIoDispatcherEnabled(
 	TEXT("iax.HttpIoDispatcherEnabled"),
 	GHttpIoDispatcherEnabled,
