Fix IAD cache not handling multiple entries for the same chunk.  This caused chunks in the cache to be counted multiple times.
Fix IAD cache not consitently masking timestamps and sometimes returning negative timestamps.

diff --git a/Engine/Source/Runtime/Experimental/IoStore/OnDemand/Private/OnDemandContentInstaller.cpp b/Engine/Source/Runtime/Experimental/IoStore/OnDemand/Private/OnDemandContentInstaller.cpp
--- a/Engine/Source/Runtime/Experimental/IoStore/OnDemand/Private/OnDemandContentInstaller.cpp
+++ b/Engine/Source/Runtime/Experimental/IoStore/OnDemand/Private/OnDemandContentInstaller.cpp
@@ -900,17 +900,25 @@ void FOnDemandContentInstaller::ExecuteInstallRequest(FRequest& Request)
 	if (CVars::EnableConcurentInstallRequests() == false)
 	{
 		// Make sure we have enough space in the cache
-
-		TSet<FIoHash> ChunksToInstall;
-		for (FChunkHttpRequestHandle& HttpRequest : InstallRequest.HttpRequestHandles)
+		uint64 BytesToInstall = 0;
 		{
-			FSharedOnDemandContainer& Container		= InstallRequest.ResolvedChunks[HttpRequest.ContainerIndex].Container;
-			const FOnDemandChunkEntry& ChunkEntry	= Container->ChunkEntries[HttpRequest.EntryIndex];
+			TSet<FIoHash> ChunksToInstall;
+			ChunksToInstall.Reserve(InstallRequest.HttpRequestHandles.Num());
+			for (FChunkHttpRequestHandle& HttpRequest : InstallRequest.HttpRequestHandles)
+			{
+				FSharedOnDemandContainer& Container = InstallRequest.ResolvedChunks[HttpRequest.ContainerIndex].Container;
+				const FOnDemandChunkEntry& ChunkEntry = Container->ChunkEntries[HttpRequest.EntryIndex];
 
-			ChunksToInstall.Add(ChunkEntry.Hash);
+				bool bAlreadyInSet = false;
+				ChunksToInstall.Add(ChunkEntry.Hash, &bAlreadyInSet);
+				if (!bAlreadyInSet)
+				{
+					BytesToInstall += ChunkEntry.GetDiskSize();
+				}
+			}
 		}
 
-		if (Request.Result = IoStore.InstallCache->Purge(MoveTemp(ChunksToInstall)); Request.Result.HasError())
+		if (Request.Result = IoStore.InstallCache->Purge(BytesToInstall); Request.Result.HasError())
 		{
 			return CompleteInstallRequest(Request);
 		}
diff --git a/Engine/Source/Runtime/Experimental/IoStore/OnDemand/Private/OnDemandInstallCache.cpp b/Engine/Source/Runtime/Experimental/IoStore/OnDemand/Private/OnDemandInstallCache.cpp
--- a/Engine/Source/Runtime/Experimental/IoStore/OnDemand/Private/OnDemandInstallCache.cpp
+++ b/Engine/Source/Runtime/Experimental/IoStore/OnDemand/Private/OnDemandInstallCache.cpp
@@ -152,7 +152,7 @@ struct FCasLocation
 	static const FCasLocation Invalid;
 
 	FCasBlockId	BlockId;
-	uint32		BlockOffset = MAX_uint32; 
+	uint32		BlockOffset = MAX_uint32;
 };
 
 const FCasLocation FCasLocation::Invalid = FCasLocation();
@@ -616,7 +616,7 @@ uint64 FCas::GetBlockInfo(FCasBlockInfoMap& OutBlockInfo)
 			BlockInfo.Add(BlockId, FCasBlockInfo
 			{
 				.FileSize = uint64(FileSize),
-				.LastAccess = UtcTicks != nullptr ? *UtcTicks : TimestampForMissingLastAccess
+				.LastAccess = UtcTicks ? (*UtcTicks & ~FCas::DirtyTimestampMask) : TimestampForMissingLastAccess
 			});
 			TotalSize += uint64(FileSize);
 
@@ -687,7 +687,7 @@ FResult FCas::Verify(TArray<FCasAddr>& OutAddrs)
 				
 #if UE_ONDEMANDINSTALLCACHE_USE_MODTIME
 				FDateTime ModTime = Ipf.GetTimeStamp(*Filename);
-				if (ModTime != FDateTime::MinValue())
+				if (ensure(ModTime > FDateTime::MinValue()))
 				{
 					if (ModTime > Now)
 					{
@@ -1303,7 +1303,7 @@ struct FCasSnapshot
 		}
 
 		FCasBlockId BlockId;
-		int64		LastAccess;
+		int64		LastAccess = 0;
 	};
 
 	using FChunkLocation = TPair<FCasAddr, FCasLocation>;
@@ -1739,7 +1739,7 @@ public:
 													FOnDemandContentHandle ContentHandle,
 													TArray<int32>& OutMissing) override;
 	virtual FResult								PutChunk(FIoBuffer&& Chunk, const FIoHash& ChunkHash) override;
-	virtual FResult								Purge(TSet<FIoHash>&& ChunksToInstall) override;
+	virtual FResult								Purge(uint64 BytesToInstall) override;
 	virtual FResult								PurgeAllUnreferenced(bool bDefrag, const uint64* BytesToPurge = nullptr) override;
 	virtual FResult								DefragAll(const uint64* BytesToFree = nullptr) override;
 	virtual FResult								Verify() override;
@@ -1752,13 +1752,12 @@ private:
 	void										RegisterConsoleCommands();
 	FResult										Reset();
 	FResult										InitialVerify();
-	uint64										AddReferencesToBlocks(
+	void										AddReferencesToBlocks(
 													const TArray<FSharedOnDemandContainer>& Containers, 
 													const TArray<TBitArray<>>& ChunkEntryIndices,
-													const TSet<FIoHash>& ChunksToInstall,
 													FCasBlockInfoMap& BlockInfoMap,
 													uint64& OutTotalReferencedBytes) const;
-	FResult										Purge(FCasBlockInfoMap& BlockInfo, uint64 TotalBytesToPurge, uint64& OutTotalPurgedBytes);
+	FResult										PurgeInternal(FCasBlockInfoMap& BlockInfo, uint64 TotalBytesToPurge, uint64& OutTotalPurgedBytes);
 	FResult										Defrag(
 													const TArray<FSharedOnDemandContainer>& Containers,
 													const TArray<TBitArray<>>& ChunkEntryIndices,
@@ -2353,7 +2352,7 @@ FResult FOnDemandInstallCache::PutChunk(FIoBuffer&& Chunk, const FIoHash& ChunkH
 	return MakeValue();
 }
 
-FResult FOnDemandInstallCache::Purge(TSet<FIoHash>&& ChunksToInstall)
+FResult FOnDemandInstallCache::Purge(const uint64 BytesToInstall)
 {
 	using namespace UE::UnifiedError::IoStoreOnDemand;
 
@@ -2373,7 +2372,7 @@ FResult FOnDemandInstallCache::Purge(TSet<FIoHash>&& ChunksToInstall)
 	uint64 TotalReferencedBlockBytes = 0;
 	int64 OldestBlockAccess = FDateTime::MaxValue().GetTicks();
 
-	const uint64 TotalUncachedBytes = AddReferencesToBlocks(Containers, ChunkEntryIndices, ChunksToInstall, BlockInfo, ReferencedBytes);
+	AddReferencesToBlocks(Containers, ChunkEntryIndices, BlockInfo, ReferencedBytes);
 	for (const TPair<FCasBlockId, FCasBlockInfo>& Kv : BlockInfo)
 	{
 		const FCasBlockInfo& Info = Kv.Value;
@@ -2408,6 +2407,7 @@ FResult FOnDemandInstallCache::Purge(TSet<FIoHash>&& ChunksToInstall)
 	FOnDemandInstallCacheStats::OnCacheUsage(
 		MaxCacheSize, TotalCachedBytes, TotalReferencedBlockBytes, ReferencedBytes, FragmentedBytes, OldestBlockAccess);
 
+	const uint64 TotalUncachedBytes = BytesToInstall + PendingChunks.TotalSize;
 	const uint64 TotalRequiredBytes = TotalCachedBytes + TotalUncachedBytes;
 	if (TotalRequiredBytes <= MaxCacheSize)
 	{
@@ -2423,7 +2423,7 @@ FResult FOnDemandInstallCache::Purge(TSet<FIoHash>&& ChunksToInstall)
 
 	const uint64	TotalBytesToPurge	= TotalRequiredBytes - MaxCacheSize;
 	uint64			TotalPurgedBytes	= 0;
-	FResult			Result				= Purge(BlockInfo, TotalBytesToPurge, TotalPurgedBytes);
+	FResult			Result				= PurgeInternal(BlockInfo, TotalBytesToPurge, TotalPurgedBytes);
 
 	if (TotalPurgedBytes > 0)
 	{
@@ -2473,7 +2473,7 @@ FResult FOnDemandInstallCache::PurgeAllUnreferenced(bool bDefrag, const uint64*
 	check(Containers.Num() == ChunkEntryIndices.Num());
 
 	uint64 ReferencedBytes = 0;
-	AddReferencesToBlocks(Containers, ChunkEntryIndices, {}, BlockInfo, ReferencedBytes);
+	AddReferencesToBlocks(Containers, ChunkEntryIndices, BlockInfo, ReferencedBytes);
 
 	const uint64 TotalReferencedBytes = Algo::TransformAccumulate(BlockInfo,
 		[](const TPair<FCasBlockId, FCasBlockInfo>& Kv) { return (Kv.Value.RefSize > 0) ? Kv.Value.FileSize : uint64(0); },
@@ -2484,7 +2484,7 @@ FResult FOnDemandInstallCache::PurgeAllUnreferenced(bool bDefrag, const uint64*
 
 	const uint64	TotalBytesToPurge	= BytesToPurge ? *BytesToPurge : MaxCacheSize;
 	uint64			TotalPurgedBytes	= 0;
-	FResult			Result				= Purge(BlockInfo, TotalBytesToPurge, TotalPurgedBytes);
+	FResult			Result				= PurgeInternal(BlockInfo, TotalBytesToPurge, TotalPurgedBytes);
 
 	if (TotalPurgedBytes > 0)
 	{
@@ -2534,7 +2534,7 @@ FResult FOnDemandInstallCache::DefragAll(const uint64* BytesToFree /*= nullptr*/
 	check(Containers.Num() == ChunkEntryIndices.Num());
 
 	uint64 ReferencedBytes = 0;
-	AddReferencesToBlocks(Containers, ChunkEntryIndices, {}, BlockInfo, ReferencedBytes);
+	AddReferencesToBlocks(Containers, ChunkEntryIndices, BlockInfo, ReferencedBytes);
 
 	const uint64 TotalReferencedBlockBytes = Algo::TransformAccumulate(BlockInfo,
 		[](const TPair<FCasBlockId, FCasBlockInfo>& Kv) { return (Kv.Value.RefSize > 0) ? Kv.Value.FileSize : uint64(0); },
@@ -2833,7 +2833,7 @@ FResult FOnDemandInstallCache::InitialVerify()
 			UE_LOG(LogIoStoreOnDemand, Warning, TEXT("Cache size is greater than disk quota - Purging install cache, MaxCacheSize=%.2lf MiB, TotalSize=%.2lf MiB, TotalBytesToPurge=%.2lf MiB"),
 				ToMiB(MaxCacheSize), ToMiB(CacheSize), ToMiB(TotalBytesToPurge));
 
-			if (FResult Result = Purge(BlockInfo, TotalBytesToPurge, TotalPurgedBytes); Result.HasError())
+			if (FResult Result = PurgeInternal(BlockInfo, TotalBytesToPurge, TotalPurgedBytes); Result.HasError())
 			{
 				const FString ErrorMessage = FString::Printf(TEXT("Failed to purge %.2lf MiB from install cache reason '%s'"),
 					ToMiB(TotalBytesToPurge), *LexToString(Result.GetError()));
@@ -2859,83 +2859,86 @@ FResult FOnDemandInstallCache::InitialVerify()
 	return MakeValue();
 }
 
-uint64 FOnDemandInstallCache::AddReferencesToBlocks(
-	const TArray<FSharedOnDemandContainer>& Containers,
-	const TArray<TBitArray<>>& ChunkEntryIndices,
-	const TSet<FIoHash>& ChunksToInstall,
-	FCasBlockInfoMap& BlockInfoMap,
+void FOnDemandInstallCache::AddReferencesToBlocks(
+	const TArray<FSharedOnDemandContainer>& Containers, 
+	const TArray<TBitArray<>>& ChunkEntryIndices, 
+	FCasBlockInfoMap& BlockInfoMap, 
 	uint64& OutTotalReferencedBytes) const
 {
-	uint64 TotalUncachedBytes = PendingChunks.TotalSize;
-	OutTotalReferencedBytes	= 0;
+	OutTotalReferencedBytes = 0;
+
+	// TODO: this could use the CAS address to be smaller
+	TSet<FIoHash> VisitedReferencedChunks;
+	{
+		int32 ReserveSize = 0;
+		for (int32 Index = 0; FSharedOnDemandContainer Container : Containers)
+		{
+			const TBitArray<>& IsReferenced = ChunkEntryIndices[Index++];
+			ReserveSize += IsReferenced.CountSetBits();
+		}
+
+		VisitedReferencedChunks.Reserve(ReserveSize);
+	}
 
 	for (int32 Index = 0; FSharedOnDemandContainer Container : Containers)
 	{
 		const TBitArray<>& IsReferenced = ChunkEntryIndices[Index++];
 		for (int32 EntryIndex = 0; const FOnDemandChunkEntry& Entry : Container->ChunkEntries)
 		{
-			const bool bToInstall = ChunksToInstall.Contains(Entry.Hash);
 			const bool bIsReferenced = IsReferenced[EntryIndex];
+			if (!bIsReferenced)
+			{
+				EntryIndex++;
+				continue;
+			}
 
-			const uint64 ChunkDiskSize = Entry.GetDiskSize();
-
-			if (bIsReferenced)
+			bool bAlreadyVisited = false;
+			VisitedReferencedChunks.Add(Entry.Hash, &bAlreadyVisited);
+			if (bAlreadyVisited)
 			{
-				OutTotalReferencedBytes += ChunkDiskSize;
+				EntryIndex++;
+				continue;
 			}
 
-			FCasBlockInfo* BlockInfo = nullptr;
-			if (bToInstall || bIsReferenced)
+			const uint64 ChunkDiskSize = Entry.GetDiskSize();
+			OutTotalReferencedBytes += ChunkDiskSize;
+
+			if (FCasLocation Loc = Cas.FindChunk(Entry.Hash); Loc.IsValid())
 			{
-				if (FCasLocation Loc = Cas.FindChunk(Entry.Hash); Loc.IsValid())
+				FCasBlockInfo* BlockInfo = BlockInfoMap.Find(Loc.BlockId);
+				if (!BlockInfo)
 				{
-					BlockInfo = BlockInfoMap.Find(Loc.BlockId);
-					if (!BlockInfo)
-					{
-						UE_CLOG(bIsReferenced, LogIoStoreOnDemand, Error, TEXT("Failed to find CAS block info for referenced chunk, ChunkId='%s', Container='%s'"),
-							*LexToString(Container->ChunkIds[EntryIndex]), *Container->Name);
-					}
+					UE_LOG(LogIoStoreOnDemand, Error, TEXT("Failed to find CAS block info for referenced chunk, ChunkId='%s', Container='%s'"),
+						*LexToString(Container->ChunkIds[EntryIndex]), *Container->Name);
 				}
 				else
 				{
-					// Check pending list
-					bool bChunkIsPending = false;
-					{
-						TSharedLock SharedLock(PendingChunks.SharedMutex);
-						bChunkIsPending = PendingChunks.ContainsChunk(Entry.Hash);
-					}
-
-					UE_CLOG(bIsReferenced && !bChunkIsPending, LogIoStoreOnDemand, Error, TEXT("Failed to find CAS location or pending chunk for chunk reference, ChunkId='%s', Container='%s'"),
-						*LexToString(Container->ChunkIds[EntryIndex]), *Container->Name);
+					BlockInfo->RefSize += ChunkDiskSize;
 				}
 			}
-
-			// If a chunk is in the list it WILL be installed even if already cached so must be counted in the uncached size.
-			if (bToInstall)
+			else
 			{
-				TotalUncachedBytes += ChunkDiskSize;
-
-				if (bIsReferenced)
+				// Check pending list
+				bool bChunkIsPending = false;
 				{
-					UE_LOGFMT(LogIoStoreOnDemand, Error, "Installing a chunk but its already cached! ChunkId={ChunkId}, Container={Container}",
-						LexToString(Container->ChunkIds[EntryIndex]), Container->Name);
-					ensure(false);
+					TSharedLock SharedLock(PendingChunks.SharedMutex);
+					bChunkIsPending = PendingChunks.ContainsChunk(Entry.Hash);
 				}
-			}
-			
-			if (bIsReferenced && BlockInfo)
-			{
-				BlockInfo->RefSize += ChunkDiskSize;
+
+				UE_CLOG(!bChunkIsPending, LogIoStoreOnDemand, Error, TEXT("Failed to find CAS location or pending chunk for chunk reference, ChunkId='%s', Container='%s'"),
+					*LexToString(Container->ChunkIds[EntryIndex]), *Container->Name);
+
+				ensure(bChunkIsPending);
 			}
 
 			EntryIndex++;
 		}
 	}
 
-	return TotalUncachedBytes;
+	return;
 }
 
-FResult FOnDemandInstallCache::Purge(FCasBlockInfoMap& BlockInfo, const uint64 TotalBytesToPurge, uint64& OutTotalPurgedBytes)
+FResult FOnDemandInstallCache::PurgeInternal(FCasBlockInfoMap& BlockInfo, const uint64 TotalBytesToPurge, uint64& OutTotalPurgedBytes)
 {
 	BlockInfo.ValueSort([](const FCasBlockInfo& LHS, const FCasBlockInfo& RHS)
 	{
@@ -3186,14 +3189,18 @@ FResult FOnDemandInstallCache::Defrag(
 			{
 				if (FDefragBlock* DefragBlock = Algo::FindBy(BlocksToDefrag, Loc.BlockId, &FDefragBlock::BlockId))
 				{
-					DefragBlock->ReferencedChunks.Add(FDefragBlockReferencedChunk
+					// TODO: Should this be a map?
+					if (nullptr == Algo::FindBy(DefragBlock->ReferencedChunks, Entry.Hash, &FDefragBlockReferencedChunk::Hash))
 					{
-						.Hash = Entry.Hash,
-						.ChunkId = ChunkId,
-						.BlockOffset = Loc.BlockOffset,
-						.DiskSize = Entry.GetDiskSize(),
-					});
-					TotalRefBytes += Entry.GetDiskSize();
+						DefragBlock->ReferencedChunks.Add(FDefragBlockReferencedChunk
+						{
+							.Hash = Entry.Hash,
+							.ChunkId = ChunkId,  // May not be unique, just return the first found for debugging purposes
+							.BlockOffset = Loc.BlockOffset,
+							.DiskSize = Entry.GetDiskSize(),
+						});
+						TotalRefBytes += Entry.GetDiskSize();
+					}
 				}
 			}
 		}
@@ -3281,7 +3288,7 @@ FResult FOnDemandInstallCache::Defrag(
 					return Result = MakeError(InstallCacheDefragError::MakeError(
 						FChunkHashMismatchErrorContext
 						{
-							.ChunkId		= ReffedChunk.ChunkId,
+							.ChunkId		= ReffedChunk.ChunkId, // May not be unique, just return the first found for debugging purposes
 							.ExpectedHash	= ReffedChunk.Hash,
 							.ActualHash		= ChunkHash
 						}));
@@ -3418,7 +3425,7 @@ FOnDemandInstallCacheUsage FOnDemandInstallCache::GetCacheUsage()
 	check(Containers.Num() == ChunkEntryIndices.Num());
 
 	uint64 ReferencedBytes = 0;
-	AddReferencesToBlocks(Containers, ChunkEntryIndices, {}, BlockInfo, ReferencedBytes);
+	AddReferencesToBlocks(Containers, ChunkEntryIndices, BlockInfo, ReferencedBytes);
 
 	uint64 FragmentedBytes = 0;
 	uint64 ReferencedBlockBytes = 0;
@@ -3548,7 +3555,7 @@ FResult FOnDemandInstallCache::FlushPendingChunksImpl(const FPendingChunks& Chun
 		}
 
 		TConstArrayView<FIoHash> ChunkHashes = MakeArrayView(ChunkHashBegin, Offsets.Num());
-		const int64 BytesToWrite = uint64(Ar.Tell());
+		const int64 BytesToWrite = Ar.Tell();
 
 		if (BytesToWrite > 0)
 		{
diff --git a/Engine/Source/Runtime/Experimental/IoStore/OnDemand/Private/OnDemandInstallCache.h b/Engine/Source/Runtime/Experimental/IoStore/OnDemand/Private/OnDemandInstallCache.h
--- a/Engine/Source/Runtime/Experimental/IoStore/OnDemand/Private/OnDemandInstallCache.h
+++ b/Engine/Source/Runtime/Experimental/IoStore/OnDemand/Private/OnDemandInstallCache.h
@@ -36,7 +36,7 @@ public:
 													FOnDemandContentHandle ContentHandle,
 													TArray<int32>& OutMissing) = 0;
 	virtual FResult								PutChunk(FIoBuffer&& Chunk, const FIoHash& ChunkHash) = 0;
-	virtual FResult								Purge(TSet<FIoHash>&& ChunksToInstall) = 0;
+	virtual FResult								Purge(uint64 BytesToInstall) = 0;
 	virtual FResult								PurgeAllUnreferenced(bool bDefrag, const uint64* BytesToPurge = nullptr) = 0;
 	virtual FResult								DefragAll(const uint64* BytesToFree = nullptr) = 0;
 	virtual FResult								Verify() = 0;
diff --git a/Engine/Source/Runtime/Experimental/IoStore/OnDemand/Private/Statistics.cpp b/Engine/Source/Runtime/Experimental/IoStore/OnDemand/Private/Statistics.cpp
--- a/Engine/Source/Runtime/Experimental/IoStore/OnDemand/Private/Statistics.cpp
+++ b/Engine/Source/Runtime/Experimental/IoStore/OnDemand/Private/Statistics.cpp
@@ -1587,7 +1587,16 @@ void FOnDemandInstallCacheStats::OnBlockDeleted(int64 LastAccessTicks, bool bFro
 {
 	const FDateTime Now = FDateTime::UtcNow();
 	const FDateTime BlockAccess(LastAccessTicks);
-	const FTimespan BlockAge = ensure(Now >= BlockAccess) ? (Now - BlockAccess) : FTimespan::MaxValue();
+
+	// Theoretically all of the following checks shouldn't be necessary, as they should be 
+	// impossible. Emperical data is not fun :(
+	FTimespan BlockAge = FTimespan::MaxValue();
+	if (ensure(BlockAccess.GetTicks() >= 0) &&
+		ensure(Now.GetTicks() >= 0) &&
+		ensure(Now >= BlockAccess))
+	{
+		BlockAge = Now - BlockAccess;
+	}
 
 	SendImmediateAnalytic(FOnDemandImmediateAnalytic
 	{
