Fix Division by zero

diff --git a/Engine/Source/Runtime/Experimental/Chaos/Public/Chaos/TaskDispatcherEvolution.h b/Engine/Source/Runtime/Experimental/Chaos/Public/Chaos/TaskDispatcherEvolution.h
--- a/Engine/Source/Runtime/Experimental/Chaos/Public/Chaos/TaskDispatcherEvolution.h
+++ b/Engine/Source/Runtime/Experimental/Chaos/Public/Chaos/TaskDispatcherEvolution.h
@@ -50,7 +50,7 @@ public:
 		for (int32 ViewIndex = 0; ViewIndex < NumView; ++ViewIndex)
 		{
 			const int32 NumKinParticles = KienamticParticles.SOAViews[ViewIndex].Size();
-			const int32 ParticleByTask = FMath::Max(FMath::DivideAndRoundUp(NumKinParticles, NumTasks), MinParticlePerTask);
+			const int32 ParticleByTask = FMath::Max(FMath::DivideAndRoundUp(NumKinParticles, FMath::Max(NumTasks, 1)), MinParticlePerTask);
 			NumKinBatches += FMath::DivideAndRoundUp(NumKinParticles, ParticleByTask);
 		}
 	}
