Constraints: evaluate non-interacting rigs when they are related to an object/control being manipulated (follow-up to 46732089)

#rnx

diff --git a/Engine/Plugins/Animation/ControlRig/Source/ControlRigEditor/Private/EditMode/ControlRigEditMode.cpp b/Engine/Plugins/Animation/ControlRig/Source/ControlRigEditor/Private/EditMode/ControlRigEditMode.cpp
--- a/Engine/Plugins/Animation/ControlRig/Source/ControlRigEditor/Private/EditMode/ControlRigEditMode.cpp
+++ b/Engine/Plugins/Animation/ControlRig/Source/ControlRigEditor/Private/EditMode/ControlRigEditMode.cpp
@@ -2389,6 +2389,7 @@ bool FControlRigEditMode::InputDelta(FEditorViewportClient* InViewportClient, FV
 				}
 			}
 		}
+		FConstraintsManagerController::Get(WorldPtr).FlushEvaluationGraph();
 		
 		TickManipulatableObjects(RigsToTick);
 	}
diff --git a/Engine/Source/Runtime/Experimental/Animation/Constraints/Private/ConstraintsEvaluationGraph.cpp b/Engine/Source/Runtime/Experimental/Animation/Constraints/Private/ConstraintsEvaluationGraph.cpp
--- a/Engine/Source/Runtime/Experimental/Animation/Constraints/Private/ConstraintsEvaluationGraph.cpp
+++ b/Engine/Source/Runtime/Experimental/Animation/Constraints/Private/ConstraintsEvaluationGraph.cpp
@@ -5,6 +5,7 @@
 #include "ConstraintsManager.h"
 #include "ConstraintSubsystem.h"
 #include "Algo/TopologicalSort.h"
+#include "Transform/TransformableHandleUtils.h"
 
 namespace ConstraintsEvaluationGraph
 {
@@ -201,7 +202,7 @@ void FConstraintsEvaluationGraph::Evaluate(FConstraintNode* InNode)
 	
 	if (Constraint->IsFullyActive() && InNode->ConstraintTick->IsTickFunctionRegistered() && InNode->ConstraintTick->IsTickFunctionEnabled())
 	{
-		Constraint->Evaluate();
+		Constraint->Evaluate(TransformableHandleUtils::SkipTicking());
 	}
 	InNode->bMarkedForEvaluation = false;
 
diff --git a/Engine/Source/Runtime/Experimental/Animation/Constraints/Private/Transform/TransformConstraintUtil.cpp b/Engine/Source/Runtime/Experimental/Animation/Constraints/Private/Transform/TransformConstraintUtil.cpp
--- a/Engine/Source/Runtime/Experimental/Animation/Constraints/Private/Transform/TransformConstraintUtil.cpp
+++ b/Engine/Source/Runtime/Experimental/Animation/Constraints/Private/Transform/TransformConstraintUtil.cpp
@@ -188,13 +188,13 @@ bool FConstraintsInteractionCache::HasAnyActiveConstraint(UObject* InTarget, UWo
 
 bool FConstraintsInteractionCache::HasAnyDependency(UObject* InChild, UObject* InParent, UWorld* InWorld)
 {
-	TSet< UObject* > Visited({InChild});
+	TSet< UObject* > Visited;
 	return HasAnyDependencyInternal(InChild, InParent, InWorld, Visited);
 }
 
 bool FConstraintsInteractionCache::HasAnyDependencyInternal(UObject* InChild, UObject* InParent, UWorld* InWorld, TSet< UObject* >& Visited)
 {
-	if (!InChild || !InParent || !InWorld || InChild == InParent || Visited.Contains(InParent))
+	if (!InChild || !InParent || !InWorld || InChild == InParent)
 	{
 		return false;
 	}
@@ -227,12 +227,16 @@ bool FConstraintsInteractionCache::HasAnyDependencyInternal(UObject* InChild, UO
 	{
 		return true;
 	}
-	
-	Visited.Add(InParent);
+
+	if (Visited.Contains(InChild))
+	{
+		return false;
+	}
+	Visited.Add(InChild);
 	
 	for (UObject* Parent: Parents)
 	{
-		if (HasAnyDependencyInternal(InChild, Parent, InWorld, Visited))
+		if (HasAnyDependencyInternal(Parent, InParent, InWorld, Visited))
 		{
 			return true;
 		}
@@ -354,6 +358,11 @@ namespace Private
 // we suppose that both InParentHandle and InChildHandle are safe to use
 bool HasConstraintDependencyWith(UWorld* InWorld, const UTransformableHandle* InParentHandle, const UTransformableHandle* InChildHandle)
 {
+	if (!InWorld || !InParentHandle || !InChildHandle)
+	{
+		return false;
+	}
+	
 	using ConstraintPtr = TWeakObjectPtr<UTickableConstraint>;
 	using HandlePtr = TObjectPtr<UTransformableHandle>;
 
@@ -685,7 +694,7 @@ TOptional<FTransform> GetRelativeTransform(UWorld* InWorld, const uint32 InHandl
 	}
 
 	// get relative transform
-	const UTickableTransformConstraint* Constraint = Cast<UTickableTransformConstraint>(Constraints[LastActiveIndex]);
+	const UTickableTransformConstraint* Constraint = Cast<UTickableTransformConstraint>(Constraints[LastActiveIndex].Get());
 	const FTransform ChildLocal = Constraint->GetChildLocalTransform();
 	const FTransform ChildGlobal = Constraint->GetChildGlobalTransform();
 	
@@ -710,7 +719,7 @@ TOptional<FTransform> GetConstraintsRelativeTransform(
 
 	// get relative transform
 	// if that constraint handles the entire transform then return the relative transform directly
-	const UTickableTransformConstraint* Constraint = Cast<UTickableTransformConstraint>(InConstraints[LastActiveIndex]);
+	const UTickableTransformConstraint* Constraint = Cast<UTickableTransformConstraint>(InConstraints[LastActiveIndex].Get());
 	if (EnumHasAllFlags(Constraint->GetChannelsToKey(), EMovieSceneTransformChannel::AllTransform))
 	{
 		const FTransform ParentGlobal = Constraint->GetParentGlobalTransform();
@@ -757,7 +766,7 @@ TOptional<FTransform> GetConstraintsRelativeTransform(
 	FTransform ChildLocal = InChildLocal;
 	for (const int32 SubConstraintIndex: SubTransformIndices)
 	{
-		const UTickableTransformConstraint* SubConstraint = Cast<UTickableTransformConstraint>(InConstraints[SubConstraintIndex]);
+		const UTickableTransformConstraint* SubConstraint = Cast<UTickableTransformConstraint>(InConstraints[SubConstraintIndex].Get());
 		const FTransform ParentGlobal = SubConstraint->GetParentGlobalTransform();
 		ChildLocal = ComputeRelativeTransform(ChildLocal, InChildWorld, ParentGlobal, SubConstraint);
 	}
