Constraints: recursive anim evaluation
- ensure that the parent handle is marked for evaluation when evaluating a component
- mark components and evaluate them recursively

#rnx

diff --git a/Engine/Plugins/Animation/ControlRig/Source/ControlRig/Private/Constraints/ControlRigTransformableHandle.cpp b/Engine/Plugins/Animation/ControlRig/Source/ControlRig/Private/Constraints/ControlRigTransformableHandle.cpp
--- a/Engine/Plugins/Animation/ControlRig/Source/ControlRig/Private/Constraints/ControlRigTransformableHandle.cpp
+++ b/Engine/Plugins/Animation/ControlRig/Source/ControlRig/Private/Constraints/ControlRigTransformableHandle.cpp
@@ -264,7 +264,11 @@ FTransform UTransformableControlHandle::GetGlobalTransform() const
 		{
 			if (!ControlRig->IsEvaluating())
 			{
-				TransformableHandleUtils::EvaluateComponent( SkeletalMeshComponent, GetEvaluationTask() );
+				const UE::Anim::FAnimationEvaluator& Evaluator = TransformableHandleUtils::EvaluateComponent( SkeletalMeshComponent, GetEvaluationTask() );
+				if (BoundComponent->GetAttachParent())
+				{
+					return Cache * Evaluator.GetGlobalTransform(NAME_None);
+				}
 			}
 			return Cache * ComponentTransform;
 		}
@@ -534,9 +538,10 @@ const UE::Anim::FAnimationEvaluationTask& UTransformableControlHandle::GetEvalua
 		const int32 CtrlIndex = Hierarchy->GetIndex(Control->GetKey());
 		EvaluationTask.PostEvaluationFunction = [&CacheRef = Cache, CtrlIndex, WeakHierarchy = TWeakObjectPtr<URigHierarchy>(Hierarchy)]()
 		{
-			if (WeakHierarchy.IsValid())
+			const URigHierarchy* Hierarchy = WeakHierarchy.IsValid() ? WeakHierarchy.Get() : nullptr;
+			if (Hierarchy && CtrlIndex > INDEX_NONE)
 			{
-				CacheRef = WeakHierarchy->GetGlobalTransform(CtrlIndex);
+				CacheRef = Hierarchy->GetGlobalTransform(CtrlIndex);
 			}
 		};
 	}
@@ -844,7 +849,7 @@ bool UTransformableControlHandle::AddTransformKeys(const TArray<FFrameNumber>& I
 //for control rig need to check to see if the control rig is different then we may need to update it based upon what we are now bound to
 void UTransformableControlHandle::ResolveBoundObjects(FMovieSceneSequenceID LocalSequenceID, TSharedRef<UE::MovieScene::FSharedPlaybackState> SharedPlaybackState, UObject* SubObject)
 {
-	if (const UControlRig* InControlRig = Cast<UControlRig>(SubObject))
+	if (UControlRig* InControlRig = Cast<UControlRig>(SubObject))
 	{
 		// nothing to do
 		if (ControlRig == InControlRig)
@@ -866,7 +871,7 @@ void UTransformableControlHandle::ResolveBoundObjects(FMovieSceneSequenceID Loca
 			if (InControlRig->GetObjectBinding() && InControlRig->GetObjectBinding()->GetBoundObject() == Bindable)
 			{
 				UnregisterDelegates();
-				ControlRig = const_cast<UControlRig*>(InControlRig);
+				ControlRig = InControlRig;
 				RegisterDelegates();
 			}
 		}
diff --git a/Engine/Plugins/Animation/ControlRig/Source/ControlRig/Private/Sequencer/MovieSceneControlRigParameterTemplate.cpp b/Engine/Plugins/Animation/ControlRig/Source/ControlRig/Private/Sequencer/MovieSceneControlRigParameterTemplate.cpp
--- a/Engine/Plugins/Animation/ControlRig/Source/ControlRig/Private/Sequencer/MovieSceneControlRigParameterTemplate.cpp
+++ b/Engine/Plugins/Animation/ControlRig/Source/ControlRig/Private/Sequencer/MovieSceneControlRigParameterTemplate.cpp
@@ -24,6 +24,7 @@
 #include "Rigs/RigHierarchy.h"
 #include "Transform/TransformConstraint.h"
 #include "Transform/TransformableHandle.h"
+#include "Transform/TransformableHandleUtils.h"
 #include "AnimationCoreLibrary.h"
 #include "Constraints/ControlRigTransformableHandle.h"
 #include "Misc/ScopeLock.h"
@@ -1468,12 +1469,22 @@ struct FControlRigParameterExecutionToken : IMovieSceneExecutionToken
 
 						if (ConstraintValue.Constraint.IsValid())
 						{
-							if (UTickableTransformConstraint* TransformConstraint = Cast<UTickableTransformConstraint>(ConstraintValue.Constraint))
+							UTickableTransformConstraint* TransformConstraint = Cast<UTickableTransformConstraint>(ConstraintValue.Constraint);
+							if (TransformConstraint)
 							{
 								TransformConstraint->InitConstraint(BoundObjectWorld);
 							}
 							ConstraintValue.Constraint->ResolveBoundObjects(Operand.SequenceID, SharedPlaybackState, ControlRig);
 							ConstraintValue.Constraint->SetActive(ConstraintValue.Value);
+
+							if (TransformConstraint && ConstraintValue.Value && TransformableHandleUtils::SkipTicking())
+							{
+								if (UTransformableComponentHandle* ParentHandle = Cast<UTransformableComponentHandle>(TransformConstraint->ParentTRSHandle))
+								{
+									constexpr bool bRecursive = true;
+									TransformableHandleUtils::MarkComponentForEvaluation(Cast<USkeletalMeshComponent>(ParentHandle->GetTarget()), bRecursive);
+								}
+							}
 						}
 					}
 
diff --git a/Engine/Plugins/Animation/ControlRig/Source/ControlRig/Private/Sequencer/MovieSceneControlRigSystem.cpp b/Engine/Plugins/Animation/ControlRig/Source/ControlRig/Private/Sequencer/MovieSceneControlRigSystem.cpp
--- a/Engine/Plugins/Animation/ControlRig/Source/ControlRig/Private/Sequencer/MovieSceneControlRigSystem.cpp
+++ b/Engine/Plugins/Animation/ControlRig/Source/ControlRig/Private/Sequencer/MovieSceneControlRigSystem.cpp
@@ -37,6 +37,7 @@
 #include "Rigs/RigHierarchy.h"
 #include "Transform/TransformConstraint.h"
 #include "Transform/TransformableHandle.h"
+#include "Transform/TransformableHandleUtils.h"
 #include "Constraints/ControlRigTransformableHandle.h"
 
 #if WITH_EDITOR
@@ -1153,15 +1154,25 @@ struct FEvaluateBaseControlRigs
 			ConstraintAndActiveChannel.ActiveChannel.Evaluate(RoundTime, Value);
 			CreateConstraintIfNeeded(BoundObject->GetWorld(), ConstraintAndActiveChannel.GetConstraint().Get(), InOutBaseEvalData.Section);
 
-			if (ConstraintAndActiveChannel.GetConstraint())
+			if (TObjectPtr<UTickableConstraint> Constraint = ConstraintAndActiveChannel.GetConstraint())
 			{
 				//For Control Rig we may need to explicitly set the control rig
-				if (UTickableTransformConstraint* TransformConstraint = Cast<UTickableTransformConstraint>(ConstraintAndActiveChannel.GetConstraint()))
+				UTickableTransformConstraint* TransformConstraint = Cast<UTickableTransformConstraint>(Constraint);
+				if (TransformConstraint)
 				{
 					TransformConstraint->InitConstraint(BoundObject->GetWorld());
 				}
-				ConstraintAndActiveChannel.GetConstraint()->ResolveBoundObjects(SequenceID, SharedPlaybackState, ControlRig);
-				ConstraintAndActiveChannel.GetConstraint()->SetActive(Value);
+				Constraint->ResolveBoundObjects(SequenceID, SharedPlaybackState, ControlRig);
+				Constraint->SetActive(Value);
+
+				if (TransformConstraint && Value && TransformableHandleUtils::SkipTicking())
+				{
+					if (UTransformableComponentHandle* ParentHandle = Cast<UTransformableComponentHandle>(TransformConstraint->ParentTRSHandle))
+					{
+						constexpr bool bRecursive = true;
+						TransformableHandleUtils::MarkComponentForEvaluation(Cast<USkeletalMeshComponent>(ParentHandle->GetTarget()), bRecursive);
+					}
+				}
 			}
 		}
 
diff --git a/Engine/Source/Runtime/Experimental/Animation/Constraints/Private/Transform/AnimationEvaluation.cpp b/Engine/Source/Runtime/Experimental/Animation/Constraints/Private/Transform/AnimationEvaluation.cpp
--- a/Engine/Source/Runtime/Experimental/Animation/Constraints/Private/Transform/AnimationEvaluation.cpp
+++ b/Engine/Source/Runtime/Experimental/Animation/Constraints/Private/Transform/AnimationEvaluation.cpp
@@ -9,6 +9,7 @@
 #include "Engine/SkeletalMesh.h"
 #include "GameFramework/Actor.h"
 #include "Logging/LogMacros.h"
+#include "UObject/FastReferenceCollector.h"
 
 namespace UE::Anim
 {
@@ -88,18 +89,73 @@ FTransform FAnimationEvaluator::GetGlobalTransform(const FName InSocketName) con
 	{
 		return FTransform::Identity;
 	}
-	
-	if (IsValid() && InSocketName != NAME_None)
+
+	USceneComponent* Parent = RawSkeletalMeshComponent->GetAttachParent();
+	FName AttachSocket = RawSkeletalMeshComponent->GetAttachSocketName();
+	if (!Parent)
+	{
+		return GetRelativeTransform(InSocketName);
+	}
+
+	FTransform Global = GetRelativeTransform(InSocketName);
+	while (Parent)
 	{
-		// TODO socket
-		const int32 BoneIndex = Context.SkeletalMesh->GetRefSkeleton().FindBoneIndex(InSocketName);
-		if (ensure(BoneIndex != INDEX_NONE && Context.ComponentSpaceTransforms.IsValidIndex(BoneIndex)))
+		if (USkeletalMeshComponent* ParentSkeletalMeshComponent = Cast<USkeletalMeshComponent>(Parent))
 		{
-			return Context.ComponentSpaceTransforms[BoneIndex] * RawSkeletalMeshComponent->GetComponentTransform();
+			const FAnimationEvaluator& ParentEvaluator = FAnimationEvaluationCache::Get().GetEvaluator(ParentSkeletalMeshComponent);
+			Global = Global * ParentEvaluator.GetRelativeTransform(AttachSocket);
 		}
+		else
+		{
+			Global = Global * Parent->GetRelativeTransform();	
+		}
+		
+		AttachSocket = Parent->GetAttachSocketName();
+		Parent = Parent->GetAttachParent();
+	}
+
+	return Global;
+}
+
+FTransform FAnimationEvaluator::GetRelativeTransform(const FName InSocketName) const
+{
+	USkeletalMeshComponent* RawSkeletalMeshComponent = SkeletalMeshComponent.Get();
+	if (!RawSkeletalMeshComponent)
+	{
+		return FTransform::Identity;
 	}
+
+	
+	USkeletalMesh* SkeletalMesh = IsValid() ? Context.SkeletalMesh : RawSkeletalMeshComponent->GetSkeletalMeshAsset();
+	auto GetSocketTransform = [InSocketName, RawSkeletalMeshComponent, SkeletalMesh](const TArray<FTransform>& ComponentSpaceTransforms)
+	{
+		FTransform SocketTransform;
+		int32 BoneIndex = INDEX_NONE;
+
+		if (RawSkeletalMeshComponent->GetSocketInfoByName(InSocketName, SocketTransform, BoneIndex))
+		{
+			if (ensure(BoneIndex != INDEX_NONE && ComponentSpaceTransforms.IsValidIndex(BoneIndex)))
+			{
+				return SocketTransform * ComponentSpaceTransforms[BoneIndex];
+			}
+		}
+
+		if (SkeletalMesh)
+		{
+			BoneIndex = SkeletalMesh->GetRefSkeleton().FindBoneIndex(InSocketName);
+			if (ensure(BoneIndex != INDEX_NONE && ComponentSpaceTransforms.IsValidIndex(BoneIndex)))
+			{
+				return ComponentSpaceTransforms[BoneIndex];
+			}
+		}
+
+		return FTransform::Identity;
+	};
 	
-	return RawSkeletalMeshComponent->GetSocketTransform(InSocketName);
+	const FTransform& RelativeTransform = RawSkeletalMeshComponent->GetRelativeTransform();
+	const TArray<FTransform>& ComponentSpaceTransforms = IsValid() ?
+		Context.ComponentSpaceTransforms : RawSkeletalMeshComponent->GetComponentSpaceTransforms();
+	return InSocketName == NAME_None ? RelativeTransform : GetSocketTransform(ComponentSpaceTransforms) * RelativeTransform;
 }
 
 bool FAnimationEvaluator::AddPostEvaluationTask(const FAnimationEvaluationTask& InTask)
@@ -295,18 +351,9 @@ void FAnimationEvaluator::BoneTransformsFinalized()
 		}
 		else
 		{
-			if (bLogNewEvaluation)
-			{
-				UE_LOG(LogTemp, Warning, TEXT("FinalizeTask %s"), *PostEvaluationTask.Guid.ToString() );
-			}
 			PostEvaluationTask.PostEvaluationFunction();
 		}
 	}
-	
-	if (bLogNewEvaluation && bFinalized)
-	{
-		UE_LOG(LogTemp, Warning, TEXT("Finalized %s's evaluator."), *GetComponentName(SkeletalMeshComponent.Get()));
-	}
 }
 	
 FAnimationEvaluationCache& FAnimationEvaluationCache::Get()
@@ -463,4 +510,4 @@ const FAnimationEvaluator& EvaluateComponent(USceneComponent* InSceneComponent)
 	return Invalid;
 }
 
-}
\ No newline at end of file
+}
diff --git a/Engine/Source/Runtime/Experimental/Animation/Constraints/Private/Transform/TransformableHandleUtils.cpp b/Engine/Source/Runtime/Experimental/Animation/Constraints/Private/Transform/TransformableHandleUtils.cpp
--- a/Engine/Source/Runtime/Experimental/Animation/Constraints/Private/Transform/TransformableHandleUtils.cpp
+++ b/Engine/Source/Runtime/Experimental/Animation/Constraints/Private/Transform/TransformableHandleUtils.cpp
@@ -93,11 +93,21 @@ void TickSkeletalMeshComponent(USkeletalMeshComponent* InSkeletalMeshComponent)
 	InSkeletalMeshComponent->MarkRenderDynamicDataDirty();
 }
 	
-void MarkComponentForEvaluation(const USceneComponent* InSceneComponent)
+void MarkComponentForEvaluation(const USceneComponent* InSceneComponent, const bool bRecursive)
 {
 	if (const USkeletalMeshComponent* SkeletalMeshComponent = Cast<USkeletalMeshComponent>(InSceneComponent))
 	{
 		UE::Anim::FAnimationEvaluationCache::Get().MarkForEvaluation(SkeletalMeshComponent);
+		
+		if (bRecursive)
+		{
+			USceneComponent* Parent = SkeletalMeshComponent->GetAttachParent();
+            while (Parent)
+            {
+            	MarkComponentForEvaluation(Parent);
+            	Parent = Parent->GetAttachParent();
+            }
+        }
 	}
 }
 
diff --git a/Engine/Source/Runtime/Experimental/Animation/Constraints/Public/Transform/AnimationEvaluation.h b/Engine/Source/Runtime/Experimental/Animation/Constraints/Public/Transform/AnimationEvaluation.h
--- a/Engine/Source/Runtime/Experimental/Animation/Constraints/Public/Transform/AnimationEvaluation.h
+++ b/Engine/Source/Runtime/Experimental/Animation/Constraints/Public/Transform/AnimationEvaluation.h
@@ -49,7 +49,7 @@ struct FAnimationEvaluator
 	bool IsValid() const;
 	
 	/** Returns InSocketName component space transform composed with the skeletal mesh component's global transform. */
-	FTransform GetGlobalTransform(const FName InSocketName) const;
+	CONSTRAINTS_API FTransform GetGlobalTransform(const FName InSocketName) const;
 
 	/** Adds a post-evaluation task to this evaluator if it has not already been added (returns true if added). */
 	bool AddPostEvaluationTask(const FAnimationEvaluationTask& InTask);
@@ -69,6 +69,9 @@ private:
 
 	/** Updates the bones transforms once the skeletal mesh is finally up to date. */
 	void BoneTransformsFinalized();
+
+	/** Returns InSocketName component space transform composed with the skeletal mesh component's relative transform. */
+	FTransform GetRelativeTransform(const FName InSocketName) const;
 	
 	TWeakObjectPtr<USkeletalMeshComponent> SkeletalMeshComponent = nullptr;
 	FAnimationEvaluationContext Context;
diff --git a/Engine/Source/Runtime/Experimental/Animation/Constraints/Public/Transform/TransformableHandleUtils.h b/Engine/Source/Runtime/Experimental/Animation/Constraints/Public/Transform/TransformableHandleUtils.h
--- a/Engine/Source/Runtime/Experimental/Animation/Constraints/Public/Transform/TransformableHandleUtils.h
+++ b/Engine/Source/Runtime/Experimental/Animation/Constraints/Public/Transform/TransformableHandleUtils.h
@@ -19,8 +19,11 @@ namespace TransformableHandleUtils
 	/** Force ticking InSkeletalMeshComponent. */
 	CONSTRAINTS_API void TickSkeletalMeshComponent(USkeletalMeshComponent* InSkeletalMeshComponent);
 
-	/** Mark InSceneComponent for animation evaluation. */
-	CONSTRAINTS_API void MarkComponentForEvaluation(const USceneComponent* InSceneComponent);
+	/**
+	 * Mark InSceneComponent for animation evaluation.
+	 * If bRecursive is true, then all skeletal mesh components up the hierarchy will be marked for evaluation.
+	 */
+	CONSTRAINTS_API void MarkComponentForEvaluation(const USceneComponent* InSceneComponent, const bool bRecursive = false);
 	
 	/** Returns an updated version of InSceneComponent's animation evaluator. */
 	CONSTRAINTS_API const UE::Anim::FAnimationEvaluator& EvaluateComponent(USceneComponent* InSceneComponent);
diff --git a/Engine/Source/Runtime/MovieSceneTracks/Private/Systems/MovieSceneConstraintSystem.cpp b/Engine/Source/Runtime/MovieSceneTracks/Private/Systems/MovieSceneConstraintSystem.cpp
--- a/Engine/Source/Runtime/MovieSceneTracks/Private/Systems/MovieSceneConstraintSystem.cpp
+++ b/Engine/Source/Runtime/MovieSceneTracks/Private/Systems/MovieSceneConstraintSystem.cpp
@@ -16,6 +16,7 @@
 #include "Systems/MovieSceneComponentTransformSystem.h"
 #include "Transform/TransformConstraint.h"
 #include "Transform/TransformableHandle.h"
+#include "Transform/TransformableHandleUtils.h"
 #include "Sections/MovieSceneConstrainedSection.h"
 
 #include UE_INLINE_GENERATED_CPP_BY_NAME(MovieSceneConstraintSystem)
@@ -166,9 +167,9 @@ struct FEvaluateConstraintChannels
 		const FSequenceInstance& TargetInstance = InstanceRegistry->GetInstance(InstanceHandle);
 		TSharedRef<FSharedPlaybackState> SharedPlaybackState = TargetInstance.GetSharedPlaybackState();
 
-		bool Result = false;
-		ConstraintAndActiveChannel->ActiveChannel.Evaluate(FrameTime, Result);
-		Constraint->SetActive(Result);
+		bool Value = false;
+		ConstraintAndActiveChannel->ActiveChannel.Evaluate(FrameTime, Value);
+		Constraint->SetActive(Value);
 		
 		if (UTickableTransformConstraint* TransformConstraint = Cast< UTickableTransformConstraint>(Constraint))
 		{
@@ -176,14 +177,23 @@ struct FEvaluateConstraintChannels
 
 			// this has to be done once the constraint initialized
 			TransformConstraint->ResolveBoundObjects(TargetInstance.GetSequenceID(), SharedPlaybackState);
+
+			if (Value && TransformableHandleUtils::SkipTicking())
+			{
+				if (UTransformableComponentHandle* ParentHandle = Cast<UTransformableComponentHandle>(TransformConstraint->ParentTRSHandle))
+				{
+					constexpr bool bRecursive = true;
+					TransformableHandleUtils::MarkComponentForEvaluation(Cast<USkeletalMeshComponent>(ParentHandle->GetTarget()), bRecursive);
+				}
+			}
 			
-			if (UTransformableComponentHandle* ComponentHandle = Cast<UTransformableComponentHandle>(TransformConstraint->ChildTRSHandle))
+			if (UTransformableComponentHandle* ChildHandle = Cast<UTransformableComponentHandle>(TransformConstraint->ChildTRSHandle))
 			{
 				//bound component may change so need to update constraint
-				ComponentHandle->Component = Cast<USceneComponent>(BoundObject);
+				ChildHandle->Component = Cast<USceneComponent>(BoundObject);
 				UMovieSceneConstraintSystem::FUpdateHandleForConstraint UpdateHandle;
 				UpdateHandle.Constraint = TransformConstraint;
-				UpdateHandle.TransformHandle = ComponentHandle;
+				UpdateHandle.TransformHandle = ChildHandle;
 				System->DynamicOffsets.Add(UpdateHandle);
 				TransformConstraint->EnsurePrimaryDependency(World);
 			}
