Fix crash in on-demand HTTP client due to response buffer not being set in all cases.

The originating change was introduced to handle CMS not modified content. This change
reverts the behavior as of 37.XX.

#rnx

diff --git a/Engine/Source/Runtime/Experimental/IoStore/OnDemand/Private/OnDemandHttpClient.cpp b/Engine/Source/Runtime/Experimental/IoStore/OnDemand/Private/OnDemandHttpClient.cpp
--- a/Engine/Source/Runtime/Experimental/IoStore/OnDemand/Private/OnDemandHttpClient.cpp
+++ b/Engine/Source/Runtime/Experimental/IoStore/OnDemand/Private/OnDemandHttpClient.cpp
@@ -435,14 +435,8 @@ FMultiEndpointHttpClient::FHttpTicketId FMultiEndpointHttpClient::IssueRequest(F
 					Request.ResponseMessage = FString(HttpResponse.GetStatusMessage().TrimEnd());
 				}
 
-				if (HttpResponse.GetContentLength() > 0)
-				{
-					HttpResponse.SetDestination(Request.bIsChunkedTransfer == false ? &Request.Body : &Request.Chunk);
-				}
-				else
-				{
-					CompleteRequest(MoveTemp(Request), TicketStatus);
-				}
+				HttpResponse.SetDestination(Request.bIsChunkedTransfer == false ? &Request.Body : &Request.Chunk);
+
 				break;
 			}
 			case FTicketStatus::EId::Content:
