Constraints: fix misuse of invalid worlds
- use static notification functions when the world is not necessary to avoid updating the controller with irrelevant data
- scope the current controller's world so that it acts on the correct one (the constraint subsystem could change the controller's world when cleaning gcd worlds/constraints)

#rnx

diff --git a/Engine/Plugins/Animation/ControlRig/Source/ControlRigEditor/Private/ControlRigSequencerEditorLibrary.cpp b/Engine/Plugins/Animation/ControlRig/Source/ControlRigEditor/Private/ControlRigSequencerEditorLibrary.cpp
--- a/Engine/Plugins/Animation/ControlRig/Source/ControlRigEditor/Private/ControlRigSequencerEditorLibrary.cpp
+++ b/Engine/Plugins/Animation/ControlRig/Source/ControlRigEditor/Private/ControlRigSequencerEditorLibrary.cpp
@@ -517,7 +517,7 @@ UTickableConstraint* UControlRigSequencerEditorLibrary::AddConstraint(UWorld* Wo
 	else
 	{
 		FConstraintsManagerController& Controller = FConstraintsManagerController::Get(World);
-		Controller.StaticConstraintCreated(World, Constraint);
+		Controller.StaticConstraintCreated(Constraint);
 	}
 	return Constraint;
 }
diff --git a/Engine/Plugins/Animation/ControlRig/Source/ControlRigEditor/Private/EditMode/ControlRigEditMode.cpp b/Engine/Plugins/Animation/ControlRig/Source/ControlRigEditor/Private/EditMode/ControlRigEditMode.cpp
--- a/Engine/Plugins/Animation/ControlRig/Source/ControlRigEditor/Private/EditMode/ControlRigEditMode.cpp
+++ b/Engine/Plugins/Animation/ControlRig/Source/ControlRigEditor/Private/EditMode/ControlRigEditMode.cpp
@@ -522,7 +522,7 @@ void FControlRigEditMode::Enter()
 
 	if (IsInLevelEditor())
 	{
-		ConstraintsCache.RegisterNotifications(WorldPtr);
+		ConstraintsCache.RegisterNotifications();
 		Keyframer.Initialize();
 	}
 }
@@ -547,7 +547,7 @@ static void ClearOutAnyActiveTools()
 
 void FControlRigEditMode::Exit()
 {
-	ConstraintsCache.UnregisterNotifications(WorldPtr);
+	ConstraintsCache.UnregisterNotifications();
 	ConstraintsCache.Reset();
 	
 	if (WeakGizmoContext.IsValid())
diff --git a/Engine/Plugins/Animation/ControlRig/Source/ControlRigEditor/Private/Editor/Constraints/SConstraintsWidget.cpp b/Engine/Plugins/Animation/ControlRig/Source/ControlRigEditor/Private/Editor/Constraints/SConstraintsWidget.cpp
--- a/Engine/Plugins/Animation/ControlRig/Source/ControlRigEditor/Private/Editor/Constraints/SConstraintsWidget.cpp
+++ b/Engine/Plugins/Animation/ControlRig/Source/ControlRigEditor/Private/Editor/Constraints/SConstraintsWidget.cpp
@@ -284,7 +284,7 @@ namespace UE::Private
 					else
 					{
 						FConstraintsManagerController& Controller = FConstraintsManagerController::Get(World);
-						Controller.StaticConstraintCreated(World, Constraint);
+						Controller.StaticConstraintCreated(Constraint);
 					}		
 				}
 			}
@@ -998,33 +998,62 @@ void SEditableConstraintItem::Construct(
 
 FBaseConstraintListWidget::~FBaseConstraintListWidget()
 {
-	UnregisterSelectionChanged();
+	UnregisterNotifications();
 }
 
-void FBaseConstraintListWidget::RegisterSelectionChanged()
+void FBaseConstraintListWidget::RegisterNotifications()
 {
-	FLevelEditorModule& LevelEditor = FModuleManager::GetModuleChecked<FLevelEditorModule>("LevelEditor");
-	FLevelEditorModule::FActorSelectionChangedEvent& ActorSelectionChangedEvent = LevelEditor.OnActorSelectionChanged();
-
-	// unregister previous one
-	if (OnSelectionChangedHandle.IsValid())
+	// selection changed
 	{
-		ActorSelectionChangedEvent.Remove(OnSelectionChangedHandle);
-		OnSelectionChangedHandle.Reset();
-	}
+		FLevelEditorModule& LevelEditor = FModuleManager::GetModuleChecked<FLevelEditorModule>("LevelEditor");
+		FLevelEditorModule::FActorSelectionChangedEvent& ActorSelectionChangedEvent = LevelEditor.OnActorSelectionChanged();
 
-	// register
-	OnSelectionChangedHandle = ActorSelectionChangedEvent.AddRaw(this, &FBaseConstraintListWidget::OnActorSelectionChanged);
+		// unregister previous one
+		if (OnSelectionChangedHandle.IsValid())
+		{
+			ActorSelectionChangedEvent.Remove(OnSelectionChangedHandle);
+			OnSelectionChangedHandle.Reset();
+		}
+		
+		// register
+		OnSelectionChangedHandle = ActorSelectionChangedEvent.AddRaw(this, &FBaseConstraintListWidget::OnActorSelectionChanged);
+	}
+	
+	// constraints controller updated
+ 	{
+		if (ConstraintsNotificationHandle.IsValid())
+		{
+			FConstraintsManagerController::GetNotifyDelegate().Remove(ConstraintsNotificationHandle);
+			ConstraintsNotificationHandle.Reset();
+		}
+		
+		ConstraintsNotificationHandle =
+			FConstraintsManagerController::GetNotifyDelegate().AddLambda([this](EConstraintsManagerNotifyType InNotifyType, UObject* InObject)
+		{
+			if (InNotifyType == EConstraintsManagerNotifyType::GraphUpdated)
+			{
+				InvalidateConstraintList();
+			}
+		});
+	}
 }
 
-void FBaseConstraintListWidget::UnregisterSelectionChanged()
+void FBaseConstraintListWidget::UnregisterNotifications()
 {
 	if (OnSelectionChangedHandle.IsValid())
 	{
-		FLevelEditorModule& LevelEditor = FModuleManager::GetModuleChecked<FLevelEditorModule>("LevelEditor");
-		LevelEditor.OnActorSelectionChanged().Remove(OnSelectionChangedHandle);
+		if (FLevelEditorModule* LevelEditor = FModuleManager::GetModulePtr<FLevelEditorModule>("LevelEditor"))
+		{
+			LevelEditor->OnActorSelectionChanged().Remove(OnSelectionChangedHandle);
+		}
 		OnSelectionChangedHandle.Reset();
 	}
+	
+	if (ConstraintsNotificationHandle.IsValid())
+	{
+		FConstraintsManagerController::GetNotifyDelegate().Remove(ConstraintsNotificationHandle);
+		ConstraintsNotificationHandle.Reset();
+	}
 }
 
 void FBaseConstraintListWidget::PostUndo(bool bSuccess)
@@ -1512,7 +1541,7 @@ void SConstraintsEditionWidget::Construct(const FArguments& InArgs)
 	];
 
 	RefreshConstraintList();
-	RegisterSelectionChanged();
+	RegisterNotifications();
 }
 
 bool SConstraintsEditionWidget::SequencerTimeChanged() 
@@ -2149,7 +2178,7 @@ void SConstraintBakeWidget::Construct(const FArguments& InArgs)
 		]
 	];
 	RefreshConstraintList();
-	RegisterSelectionChanged();
+	RegisterNotifications();
 }
 
 TSharedRef<ITableRow> SConstraintBakeWidget::OnGenerateWidgetForItem(
diff --git a/Engine/Plugins/Animation/ControlRig/Source/ControlRigEditor/Private/Editor/Constraints/SConstraintsWidget.h b/Engine/Plugins/Animation/ControlRig/Source/ControlRigEditor/Private/Editor/Constraints/SConstraintsWidget.h
--- a/Engine/Plugins/Animation/ControlRig/Source/ControlRigEditor/Private/Editor/Constraints/SConstraintsWidget.h
+++ b/Engine/Plugins/Animation/ControlRig/Source/ControlRigEditor/Private/Editor/Constraints/SConstraintsWidget.h
@@ -263,8 +263,8 @@ protected:
 	using ConstraintItemListView = SListView< ItemSharedPtr >;
 
 
-	UE_API void RegisterSelectionChanged();
-	UE_API void UnregisterSelectionChanged();
+	void RegisterNotifications();
+	void UnregisterNotifications();
 
 	/** List view that shows constraint types*/
 	TSharedPtr< ConstraintItemListView > ListView;
@@ -276,6 +276,7 @@ protected:
 	bool bNeedsRefresh = false;
 
 	FDelegateHandle OnSelectionChangedHandle;
+	FDelegateHandle ConstraintsNotificationHandle;
 
 public:
 
diff --git a/Engine/Plugins/Animation/ControlRig/Source/ControlRigEditor/Private/Sequencer/ControlRigParameterTrackEditor.cpp b/Engine/Plugins/Animation/ControlRig/Source/ControlRigEditor/Private/Sequencer/ControlRigParameterTrackEditor.cpp
--- a/Engine/Plugins/Animation/ControlRig/Source/ControlRigEditor/Private/Sequencer/ControlRigParameterTrackEditor.cpp
+++ b/Engine/Plugins/Animation/ControlRig/Source/ControlRigEditor/Private/Sequencer/ControlRigParameterTrackEditor.cpp
@@ -644,13 +644,11 @@ void FControlRigParameterTrackEditor::ObjectImplicitlyRemoved(UObject* InObject)
 
 void FControlRigParameterTrackEditor::OnRelease()
 {
-	UWorld* World = GCurrentLevelEditingViewportClient ? GCurrentLevelEditingViewportClient->GetWorld() : nullptr;
-	FConstraintsManagerController& Controller = FConstraintsManagerController::Get(World);
 	for (FDelegateHandle& Handle : ConstraintHandlesToClear)
 	{
 		if (Handle.IsValid())
 		{
-			Controller.GetNotifyDelegate().Remove(Handle);
+			FConstraintsManagerController::GetNotifyDelegate().Remove(Handle);
 		}
 	}
 	ConstraintHandlesToClear.Reset();
@@ -729,7 +727,7 @@ void FControlRigParameterTrackEditor::OnRelease()
 
 		if (OnMovieSceneBindingsChangeHandle.IsValid())
 		{
-			Sequencer->OnChannelChanged().Remove(OnMovieSceneBindingsChangeHandle);
+			Sequencer->OnMovieSceneBindingsChanged().Remove(OnMovieSceneBindingsChangeHandle);
 			OnMovieSceneBindingsChangeHandle.Reset();
 		}
 
@@ -3655,12 +3653,8 @@ void FControlRigParameterTrackEditor::ClearOutAllSpaceAndConstraintDelegates(con
 
 				if (CRSection->OnConstraintRemovedHandle.IsValid())
 				{
-					if (const UControlRig* ControlRig = CRSection->GetControlRig())
-					{
-						FConstraintsManagerController& Controller = FConstraintsManagerController::Get(ControlRig->GetWorld());
-						Controller.GetNotifyDelegate().Remove(CRSection->OnConstraintRemovedHandle);
-						CRSection->OnConstraintRemovedHandle.Reset();
-					}
+					FConstraintsManagerController::GetNotifyDelegate().Remove(CRSection->OnConstraintRemovedHandle);
+					CRSection->OnConstraintRemovedHandle.Reset();
 				}
 			}
 		}
@@ -3827,16 +3821,18 @@ void FControlRigParameterTrackEditor::HandleConstraintKeyMoved(
 void FControlRigParameterTrackEditor::HandleConstraintRemoved(IMovieSceneConstrainedSection* InSection) 
 {
 	UMovieSceneControlRigParameterSection* Section = Cast<UMovieSceneControlRigParameterSection>(InSection);
+	if (!Section)
+	{
+		return;
+	}
 
 	if (const UControlRig* ControlRig = Section->GetControlRig())
 	{
-		FConstraintsManagerController& Controller = FConstraintsManagerController::Get(ControlRig->GetWorld());
-
-		if (!InSection->OnConstraintRemovedHandle.IsValid())
+		if (!Section->OnConstraintRemovedHandle.IsValid())
 		{
-			InSection->OnConstraintRemovedHandle =
-			Controller.GetNotifyDelegate().AddLambda(
-				[InSection, Section, this](EConstraintsManagerNotifyType InNotifyType, UObject *InObject)
+			Section->OnConstraintRemovedHandle =
+			FConstraintsManagerController::GetNotifyDelegate().AddLambda(
+				[WeakSection = MakeWeakObjectPtr(Section), this](EConstraintsManagerNotifyType InNotifyType, UObject *InObject)
 			{
 				switch (InNotifyType)
 				{
@@ -3846,12 +3842,13 @@ void FControlRigParameterTrackEditor::HandleConstraintRemoved(IMovieSceneConstra
 					case EConstraintsManagerNotifyType::ConstraintRemovedWithCompensation:
 						{
 							const UTickableConstraint* Constraint = Cast<UTickableConstraint>(InObject);
-							if (!IsValid(Constraint))
+							UMovieSceneControlRigParameterSection* Section = WeakSection.IsValid() ? WeakSection.Get() : nullptr;
+							if (!IsValid(Constraint) || !Section)
 							{
 								return;
 							}
 
-							const FConstraintAndActiveChannel* ConstraintChannel = InSection->GetConstraintChannel(Constraint->ConstraintID);
+							const FConstraintAndActiveChannel* ConstraintChannel = Section->GetConstraintChannel(Constraint->ConstraintID);
 							if (!ConstraintChannel || ConstraintChannel->GetConstraint().Get() != Constraint)
 							{
 								return;
@@ -3871,7 +3868,7 @@ void FControlRigParameterTrackEditor::HandleConstraintRemoved(IMovieSceneConstra
 							   }
 							}
 
-							InSection->RemoveConstraintChannel(Constraint);
+							Section->RemoveConstraintChannel(Constraint);
 
 							if (Sequencer)
 							{
@@ -3880,14 +3877,17 @@ void FControlRigParameterTrackEditor::HandleConstraintRemoved(IMovieSceneConstra
 						}
 						break;
 					case EConstraintsManagerNotifyType::ManagerUpdated:
-						InSection->OnConstraintsChanged();
+						if(UMovieSceneControlRigParameterSection* Section = WeakSection.IsValid() ? WeakSection.Get() : nullptr)
+						{
+							Section->OnConstraintsChanged();
+						}
 						break;
 					case EConstraintsManagerNotifyType::GraphUpdated:
 						break;
 				}
 			});
 
-			ConstraintHandlesToClear.Add(InSection->OnConstraintRemovedHandle);
+			ConstraintHandlesToClear.Add(Section->OnConstraintRemovedHandle);
 		}
 	}
 }
diff --git a/Engine/Plugins/Experimental/CineCameraRigs/Source/CineCameraRigs/Private/CineCameraAttachMount.cpp b/Engine/Plugins/Experimental/CineCameraRigs/Source/CineCameraRigs/Private/CineCameraAttachMount.cpp
--- a/Engine/Plugins/Experimental/CineCameraRigs/Source/CineCameraRigs/Private/CineCameraAttachMount.cpp
+++ b/Engine/Plugins/Experimental/CineCameraRigs/Source/CineCameraRigs/Private/CineCameraAttachMount.cpp
@@ -139,7 +139,7 @@ void ACineCameraAttachMount::CreateConstraint()
 		Constraint->bDynamicOffset = true;
 		Cast<UTickableParentConstraint>(Constraint)->OffsetTransform = FTransform::Identity;
 		FConstraintsManagerController& Controller = FConstraintsManagerController::Get(World);
-		Controller.StaticConstraintCreated(World, Constraint);
+		Controller.StaticConstraintCreated(Constraint);
 	}
 	bConstraintUpdated = true;
 }
diff --git a/Engine/Source/Editor/MovieSceneTools/Private/TrackEditors/TransformTrackEditor.cpp b/Engine/Source/Editor/MovieSceneTools/Private/TrackEditors/TransformTrackEditor.cpp
--- a/Engine/Source/Editor/MovieSceneTools/Private/TrackEditors/TransformTrackEditor.cpp
+++ b/Engine/Source/Editor/MovieSceneTools/Private/TrackEditors/TransformTrackEditor.cpp
@@ -147,9 +147,7 @@ void F3DTransformTrackEditor::OnRelease()
 
 	if (OnSceneComponentConstrainedHandle.IsValid())
 	{
-		UWorld* World = GCurrentLevelEditingViewportClient ? GCurrentLevelEditingViewportClient->GetWorld() : nullptr;
-		FConstraintsManagerController& Controller = FConstraintsManagerController::Get(World);
-		Controller.OnSceneComponentConstrained().Remove(OnSceneComponentConstrainedHandle);
+		FConstraintsManagerController::OnSceneComponentConstrained().Remove(OnSceneComponentConstrainedHandle);
 		OnSceneComponentConstrainedHandle.Reset();
 	}
 
@@ -1662,12 +1660,10 @@ void F3DTransformTrackEditor::HandleOnConstraintAdded(IMovieSceneConstrainedSect
 		return;
 	}
 	// handle scene component changes so we can update the gizmo location when it changes
-	UWorld* World = GCurrentLevelEditingViewportClient ? GCurrentLevelEditingViewportClient->GetWorld() : nullptr;
-	FConstraintsManagerController& Controller = FConstraintsManagerController::Get(World);
-	if (!Controller.OnSceneComponentConstrained().IsBound())
+	if (!FConstraintsManagerController::OnSceneComponentConstrained().IsBound())
 	{
 		OnSceneComponentConstrainedHandle =
-			Controller.OnSceneComponentConstrained().AddLambda([](USceneComponent* InSceneComponent)
+			FConstraintsManagerController::OnSceneComponentConstrained().AddLambda([](USceneComponent* InSceneComponent)
 				{
 					//only update gizmo if it or it's actor is seleced, otherwise can be wasteful
 					GUnrealEd->UpdatePivotLocationForSelection();
@@ -1801,14 +1797,16 @@ void F3DTransformTrackEditor::HandleConstraintKeyMoved(IMovieSceneConstrainedSec
 
 void F3DTransformTrackEditor::HandleConstraintRemoved(IMovieSceneConstrainedSection* InSection) 
 {
-	UWorld* World = GCurrentLevelEditingViewportClient ? GCurrentLevelEditingViewportClient->GetWorld() : nullptr;
 	UMovieScene3DTransformSection* Section = Cast<UMovieScene3DTransformSection>(InSection);
+	if (!Section)
+	{
+		return;
+	}
 
-	FConstraintsManagerController& Controller = FConstraintsManagerController::Get(World);
-	if (!InSection->OnConstraintRemovedHandle.IsValid())
+	if (!Section->OnConstraintRemovedHandle.IsValid())
 	{
-		InSection->OnConstraintRemovedHandle =
-			Controller.GetNotifyDelegate().AddLambda([InSection,Section, this](EConstraintsManagerNotifyType InNotifyType, UObject *InObject)
+		Section->OnConstraintRemovedHandle =
+			FConstraintsManagerController::GetNotifyDelegate().AddLambda([WeakSection = MakeWeakObjectPtr(Section), this](EConstraintsManagerNotifyType InNotifyType, UObject *InObject)
 				{
 					switch (InNotifyType)
 					{
@@ -1818,18 +1816,18 @@ void F3DTransformTrackEditor::HandleConstraintRemoved(IMovieSceneConstrainedSect
 						case EConstraintsManagerNotifyType::ConstraintRemovedWithCompensation:
 							{
 								const UTickableConstraint* Constraint = Cast<UTickableConstraint>(InObject);
-								if (!IsValid(Constraint))
+								UMovieScene3DTransformSection* Section = WeakSection.IsValid() ? WeakSection.Get() : nullptr;
+								if (!IsValid(Constraint) || !Section)
 								{
 									return;
 								}
 
-								const FConstraintAndActiveChannel* ConstraintChannel = InSection->GetConstraintChannel(Constraint->ConstraintID);
+								const FConstraintAndActiveChannel* ConstraintChannel = Section->GetConstraintChannel(Constraint->ConstraintID);
 								if (!ConstraintChannel || ConstraintChannel->GetConstraint().Get() != Constraint)
 								{
 									return;
 								}
 
-
 								TSharedPtr<ISequencer> Sequencer = GetSequencer();
 								
 								const bool bCompensate = (InNotifyType == EConstraintsManagerNotifyType::ConstraintRemovedWithCompensation);
@@ -1842,7 +1840,7 @@ void F3DTransformTrackEditor::HandleConstraintRemoved(IMovieSceneConstrainedSect
 										Section);
 								}
 
-								InSection->RemoveConstraintChannel(Constraint);
+								Section->RemoveConstraintChannel(Constraint);
 
 								if (Sequencer)
 								{
@@ -1851,14 +1849,17 @@ void F3DTransformTrackEditor::HandleConstraintRemoved(IMovieSceneConstrainedSect
 							}
 							break;
 						case EConstraintsManagerNotifyType::ManagerUpdated:
-							InSection->OnConstraintsChanged();
+							if (UMovieScene3DTransformSection* Section = WeakSection.IsValid() ? WeakSection.Get() : nullptr)
+							{
+								Section->OnConstraintsChanged();
+							}
 							break;
 						case EConstraintsManagerNotifyType::GraphUpdated:
 							break;
 					}
 				});
-		ConstraintHandlesToClear.Add(InSection->OnConstraintRemovedHandle);
-
+		
+		ConstraintHandlesToClear.Add(Section->OnConstraintRemovedHandle);
 	}
 }
 
@@ -1908,13 +1909,11 @@ void F3DTransformTrackEditor::HandleConstraintPropertyChanged(UTickableTransform
 
 void F3DTransformTrackEditor::ClearOutConstraintDelegates() 
 {
-	UWorld* World = GCurrentLevelEditingViewportClient ? GCurrentLevelEditingViewportClient->GetWorld() : nullptr;
-	FConstraintsManagerNotifyDelegate& NotifyDelegate = FConstraintsManagerController::Get(World).GetNotifyDelegate();
 	for (FDelegateHandle& Handle : ConstraintHandlesToClear)
 	{
 		if (Handle.IsValid())
 		{
-			NotifyDelegate.Remove(Handle);
+			FConstraintsManagerController::GetNotifyDelegate().Remove(Handle);
 		}
 	}
 	ConstraintHandlesToClear.Reset();
diff --git a/Engine/Source/Runtime/Experimental/Animation/Constraints/Private/ConstraintSubsystem.cpp b/Engine/Source/Runtime/Experimental/Animation/Constraints/Private/ConstraintSubsystem.cpp
--- a/Engine/Source/Runtime/Experimental/Animation/Constraints/Private/ConstraintSubsystem.cpp
+++ b/Engine/Source/Runtime/Experimental/Animation/Constraints/Private/ConstraintSubsystem.cpp
@@ -225,7 +225,7 @@ void UConstraintSubsystem::SetConstraintDependencies(
 	InFunctionToTickAfter->AddPrerequisite(this, *InFunctionToTickBefore);
 }
 
-bool UConstraintSubsystem::HasConstraint(UWorld* InWorld, UTickableConstraint* InConstraint)
+bool UConstraintSubsystem::HasConstraint(UWorld* InWorld, UTickableConstraint* InConstraint) const
 {
 	const TArray<TWeakObjectPtr<UTickableConstraint>>& Constraints = GetConstraintsArray(InWorld);
 	return  Constraints.Contains(InConstraint);
@@ -252,8 +252,7 @@ void UConstraintSubsystem::PostEditUndo()
 		ConstsInWorld.InvalidateGraph();
 	}
 	
-	const FConstraintsManagerController& Controller = FConstraintsManagerController::Get(GetWorld());
-	Controller.Notify(EConstraintsManagerNotifyType::ManagerUpdated, this);
+	FConstraintsManagerController::Notify(EConstraintsManagerNotifyType::ManagerUpdated, this);
 }
 #endif
 
@@ -380,7 +379,5 @@ void FConstraintsInWorld::InvalidateGraph()
 	{
 		EvaluationGraph.Reset();
 	}
-
-	const FConstraintsManagerController& Controller = FConstraintsManagerController::Get(World.Get());
-	Controller.Notify(EConstraintsManagerNotifyType::GraphUpdated, nullptr);
+	FConstraintsManagerController::Notify(EConstraintsManagerNotifyType::GraphUpdated, nullptr);
 }
diff --git a/Engine/Source/Runtime/Experimental/Animation/Constraints/Private/ConstraintSubsystem.h b/Engine/Source/Runtime/Experimental/Animation/Constraints/Private/ConstraintSubsystem.h
--- a/Engine/Source/Runtime/Experimental/Animation/Constraints/Private/ConstraintSubsystem.h
+++ b/Engine/Source/Runtime/Experimental/Animation/Constraints/Private/ConstraintSubsystem.h
@@ -62,7 +62,7 @@ public:
 
 	CONSTRAINTS_API void AddConstraint(UWorld* InWorld, UTickableConstraint* InConstraint);
 	CONSTRAINTS_API void RemoveConstraint(UWorld* InWorld, UTickableConstraint* InConstraint, bool bDoNoCompensate);
-	CONSTRAINTS_API bool HasConstraint(UWorld* InWorld, UTickableConstraint* InConstraint);
+	CONSTRAINTS_API bool HasConstraint(UWorld* InWorld, UTickableConstraint* InConstraint) const;
 
 	FConstraintsEvaluationGraph& GetEvaluationGraph(UWorld* InWorld);
 	
diff --git a/Engine/Source/Runtime/Experimental/Animation/Constraints/Private/ConstraintsManager.cpp b/Engine/Source/Runtime/Experimental/Animation/Constraints/Private/ConstraintsManager.cpp
--- a/Engine/Source/Runtime/Experimental/Animation/Constraints/Private/ConstraintsManager.cpp
+++ b/Engine/Source/Runtime/Experimental/Animation/Constraints/Private/ConstraintsManager.cpp
@@ -236,29 +236,37 @@ void UConstraintsManager::PostLoad()
 
 void UConstraintsManager::OnActorDestroyed(AActor* InActor)
 {
+	if (FConstraintsManagerController::bDoNotRemoveConstraint)
+	{
+		return;
+	}
+	
+	USceneComponent* SceneComponent = InActor ? InActor->GetRootComponent() : nullptr;
+	UWorld* ActorWorld = InActor ? InActor->GetWorld() : nullptr;
+	if (!ActorWorld || !SceneComponent)
+	{
+		return;
+	}
 
-	if (USceneComponent* SceneComponent = InActor->GetRootComponent())
+	TArray< int32 > IndicesToRemove;		
+	for (int32 Index = 0; Index < Constraints.Num(); ++Index)
 	{
-		TArray< int32 > IndicesToRemove;
-		for (int32 Index = 0; Index < Constraints.Num(); ++Index)
+		//if the constraint has a bound object(in sequencer) we don't remove the constraint, it could be a spawnable
+		if (Constraints[Index] && IsValid(Constraints[Index].Get()) && Constraints[Index]->HasBoundObjects() == false && Constraints[Index]->ReferencesObject(SceneComponent))
 		{
-			//if the constraint has a bound object(in sequencer) we don't remove the constraint, it could be a spawnable
-			if (Constraints[Index] && IsValid(Constraints[Index].Get()) && Constraints[Index]->HasBoundObjects() == false && Constraints[Index]->ReferencesObject(SceneComponent))
-			{
-				IndicesToRemove.Add(Index);
-			}
+			IndicesToRemove.Add(Index);
 		}
+	}
 
-		if (!IndicesToRemove.IsEmpty())
+	if (!IndicesToRemove.IsEmpty())
+	{
+		constexpr bool bDoNotCompensate = true;
+		for (int32 Index = IndicesToRemove.Num() - 1; Index >= 0; --Index)
 		{
-			if (FConstraintsManagerController::bDoNotRemoveConstraint == false)
+			const int32 ConstraintIndex = IndicesToRemove[Index]; 
+			if (Constraints.IsValidIndex(ConstraintIndex))
 			{
-				UWorld* ActorWorld = InActor->GetWorld();
-				constexpr bool bDoNotCompensate = true;
-				for (int32 Index = IndicesToRemove.Num() - 1; Index >= 0; --Index)
-				{
-					FConstraintsManagerController::Get(ActorWorld).RemoveConstraint(Index, bDoNotCompensate);
-				}
+				FConstraintsManagerController::Get(ActorWorld).RemoveConstraint(Constraints[ConstraintIndex], bDoNotCompensate);
 			}
 		}
 	}
@@ -276,9 +284,9 @@ void UConstraintsManager::RegisterDelegates(UWorld* World)
 
 void UConstraintsManager::UnregisterDelegates(UWorld* World)
 {
-	if (World)
+	if (World && OnActorDestroyedHandle.IsValid())
 	{
-		World->RemoveOnActorSpawnedHandler(OnActorDestroyedHandle);
+		World->RemoveOnActorDestroyedHandler(OnActorDestroyedHandle);
 	}
 	OnActorDestroyedHandle.Reset();
 }
@@ -288,7 +296,6 @@ void UConstraintsManager::Init(UWorld* World)
 	if (World)
 	{
 		UnregisterDelegates(World);
-
 		RegisterDelegates(World);
 	}
 }
@@ -359,6 +366,33 @@ void UConstraintsManager::Dump() const
  * FConstraintsManagerController
  **/
 
+namespace UE::Private
+{
+
+struct FScopedControllerContext
+{
+	FScopedControllerContext(UWorld* InWorld)
+		: World(InWorld)
+	{}
+
+	UWorld* GetWorld() const
+	{
+		return World.IsValid() ? World.Get() : nullptr;
+	}
+
+	~FScopedControllerContext()
+	{
+		if (UWorld* ContextWorld = GetWorld())
+		{
+			(void)FConstraintsManagerController::Get(ContextWorld);
+		}
+	}
+
+private:
+	TWeakObjectPtr<UWorld> World = nullptr;
+};
+
+}
 
 bool FConstraintsManagerController::bDoNotRemoveConstraint = false;
 
@@ -372,19 +406,23 @@ FConstraintsManagerController& FConstraintsManagerController::Get(UWorld* InWorl
 
 UConstraintsManager* FConstraintsManagerController::GetManager() 
 {
-	UConstraintsManager* Manager = FindManager();
-	if (Manager)
+	if (!World)
+	{
+		return nullptr;
+	}
+
+	const UE::Private::FScopedControllerContext Context(World);
+	if (UConstraintsManager* Manager = FindManager())
 	{
 		return Manager;
 	}
 
 	// create a new ConstraintsActor
-	AConstraintsActor* ConstraintsActor = World->SpawnActor<AConstraintsActor>();
+	AConstraintsActor* ConstraintsActor = Context.GetWorld()->SpawnActor<AConstraintsActor>();
 #if WITH_EDITOR
 	ConstraintsActor->SetActorLabel(TEXT("Constraints Manager"));
 #endif // WITH_EDITOR
 	return ConstraintsActor->ConstraintsManager;
-
 }
 
 UConstraintsManager* FConstraintsManagerController::FindManager() const
@@ -394,9 +432,6 @@ UConstraintsManager* FConstraintsManagerController::FindManager() const
 		return nullptr;
 	}
 	
-	// should we work with the persistent level?
-	const ULevel* Level = World->GetCurrentLevel();
-
 	// look for first ConstraintsActor
 	auto FindFirstConstraintsActor = [](const ULevel* Level)
 	{
@@ -408,23 +443,22 @@ UConstraintsManager* FConstraintsManagerController::FindManager() const
 		return Index != INDEX_NONE ? Cast<AConstraintsActor>(Level->Actors[Index]) : nullptr;
 	};
 
-	const AConstraintsActor* ConstraintsActor = FindFirstConstraintsActor(Level);
+	// should we work with the persistent level?
+	const ULevel* Level = World->GetCurrentLevel();
+	const AConstraintsActor* ConstraintsActor = Level ? FindFirstConstraintsActor(Level) : nullptr;
 	return ConstraintsActor ? ConstraintsActor->ConstraintsManager : nullptr;
 }
 
 void FConstraintsManagerController::DestroyManager() const
 {
-	if (!World)
+	if (!World || !World->GetCurrentLevel())
 	{
 		return;
 	}
 	
-	// should we work with the persistent level?
-	const ULevel* Level = World->GetCurrentLevel();
-
 	// note there should be only one...
 	TArray<AActor*> ConstraintsActorsToRemove;
-	Algo::CopyIf(Level->Actors, ConstraintsActorsToRemove, [](const AActor* Actor)
+	Algo::CopyIf(World->GetCurrentLevel()->Actors, ConstraintsActorsToRemove, [](const AActor* Actor)
 	{
 		return IsValid(Actor) && Actor->IsA<AConstraintsActor>();
 	} );
@@ -436,6 +470,11 @@ void FConstraintsManagerController::DestroyManager() const
 }
 
 void FConstraintsManagerController::StaticConstraintCreated(UWorld* InWorld, UTickableConstraint* InConstraint)
+{
+	StaticConstraintCreated(InConstraint);
+}
+
+void FConstraintsManagerController::StaticConstraintCreated(UTickableConstraint* InConstraint)
 {
 	if (UConstraintsManager* Manager = GetManager()) //created If needed
 	{
@@ -446,12 +485,12 @@ void FConstraintsManagerController::StaticConstraintCreated(UWorld* InWorld, UTi
 			InConstraint->Rename(nullptr, Manager, REN_DontCreateRedirectors);
 			Manager->OnConstraintAdded_BP.Broadcast(Manager, InConstraint);
 		}
-	}
+	}	
 }
 
 bool FConstraintsManagerController::AddConstraint(UTickableConstraint* InConstraint) const
 {
-	if (!InConstraint)
+	if (!InConstraint || !World)
 	{
 		return false;
 	}
@@ -460,16 +499,21 @@ bool FConstraintsManagerController::AddConstraint(UTickableConstraint* InConstra
 	{
 		return false;
 	}
-	if (Subsystem->HasConstraint(World, InConstraint) == true)
+
 	{
-		return false;
-	}
-	InConstraint->InitConstraint(World);
+		const UE::Private::FScopedControllerContext Context(World);
+		if (Subsystem->HasConstraint(Context.GetWorld(), InConstraint))
+		{
+			return false;
+		}
+		
+		InConstraint->InitConstraint(Context.GetWorld());
 
-	Subsystem->AddConstraint(World,InConstraint);
+		Subsystem->AddConstraint(Context.GetWorld(), InConstraint);
 
-	// build dependencies
-	InConstraint->AddedToWorld(World);
+		// build dependencies
+		InConstraint->AddedToWorld(Context.GetWorld());
+	}
 
 	// notify
 	Notify(EConstraintsManagerNotifyType::ConstraintAdded, InConstraint);
@@ -477,59 +521,16 @@ bool FConstraintsManagerController::AddConstraint(UTickableConstraint* InConstra
 	return true;
 }
 
-UTickableConstraint* FConstraintsManagerController::AddConstraintFromCopy(UTickableConstraint* CopyOfConstraint) const
-{
-	if (!CopyOfConstraint)
-	{
-		return nullptr;
-	}
-	UConstraintSubsystem* Subsystem = UConstraintSubsystem::Get();
-	if (!Subsystem)
-	{
-		return nullptr;
-	}
-	int32 Index = GetConstraintIndex(CopyOfConstraint->ConstraintID);
-	if (Index != INDEX_NONE)
-	{
-		const TArray<TWeakObjectPtr<UTickableConstraint>>& Constraints = Subsystem->GetConstraintsArray(World);
-		return Constraints[Index].Get();
-	}
-
-	UObject* Outer = CopyOfConstraint->GetOuter();
-	UTickableConstraint* OurCopy = CopyOfConstraint->Duplicate(Outer);
-	OurCopy->ConstraintID = CopyOfConstraint->ConstraintID;
-	OurCopy->SetActive(false); //always set false
-	if (AddConstraint(OurCopy))
-	{
-		return OurCopy;
-	}
-	return nullptr;
-}
-
-int32 FConstraintsManagerController::GetConstraintIndex(const FGuid& InGuid) const
-{
-	const UConstraintSubsystem* Subsystem = UConstraintSubsystem::Get();
-	if (!Subsystem)
-	{
-		return INDEX_NONE;
-	}
-
-	const TArray<TWeakObjectPtr<UTickableConstraint>> Constraints = Subsystem->GetConstraints(World);
-	return Constraints.IndexOfByPredicate([InGuid](const TWeakObjectPtr<UTickableConstraint>& Constraint)
-		{
-			return 	(Constraint.IsValid() && Constraint->ConstraintID == InGuid);
-		});
-}
-
 int32 FConstraintsManagerController::GetConstraintIndex(const FName& InConstraintName) const
 {
 	const UConstraintSubsystem* Subsystem = UConstraintSubsystem::Get();
-	if (!Subsystem)
+	if (!Subsystem ||!World)
 	{
 		return INDEX_NONE;
 	}
 
-	const TArray<TWeakObjectPtr<UTickableConstraint>> Constraints = Subsystem->GetConstraints(World);
+	const UE::Private::FScopedControllerContext Context(World);
+	const TArray<TWeakObjectPtr<UTickableConstraint>>& Constraints = Subsystem->GetConstraintsArray(Context.GetWorld());
 	return Constraints.IndexOfByPredicate([InConstraintName](const TWeakObjectPtr<UTickableConstraint>& Constraint)
 	{
 		return 	(Constraint.IsValid() && Constraint->GetFName() == InConstraintName);
@@ -539,31 +540,33 @@ int32 FConstraintsManagerController::GetConstraintIndex(const FName& InConstrain
 	
 bool FConstraintsManagerController::RemoveConstraint(UTickableConstraint* InConstraint, bool bDoNotCompensate) 
 {
-	if (FConstraintsManagerController::bDoNotRemoveConstraint == true)
+	if (bDoNotRemoveConstraint)
 	{
 		return false;
 	}
+	
 	UConstraintSubsystem* Subsystem = UConstraintSubsystem::Get();
-	if (!Subsystem)
+	if (!Subsystem || !World)
 	{
 		return false;
 	}
-	TArray<TWeakObjectPtr<UTickableConstraint>> Constraints = Subsystem->GetConstraints(World);
 
-	int32 Index = Constraints.IndexOfByPredicate([InConstraint](const TWeakObjectPtr<UTickableConstraint>& Constraint)
+	int32 Index = INDEX_NONE;
+	{
+		const UE::Private::FScopedControllerContext Context(World);
+		const TArray<TWeakObjectPtr<UTickableConstraint>>& Constraints = Subsystem->GetConstraintsArray(Context.GetWorld());
+		Index = Constraints.IndexOfByPredicate([InConstraint](const TWeakObjectPtr<UTickableConstraint>& Constraint)
 		{
-			return 	(Constraint.IsValid() && Constraint.Get() == InConstraint);
+			return (Constraint.IsValid() && Constraint.Get() == InConstraint);
 		});
-	if (Index == INDEX_NONE)
-	{
-		return false;
 	}
-	return RemoveConstraint(Index, bDoNotCompensate);
+
+	return Index != INDEX_NONE ? RemoveConstraint(Index, bDoNotCompensate) : false;
 }
 
 bool FConstraintsManagerController::UnregisterConstraint(UTickableConstraint* InConstraint)
 {
-	if (!InConstraint)
+	if (!InConstraint || !World)
 	{
 		return false;
 	}
@@ -574,10 +577,12 @@ bool FConstraintsManagerController::UnregisterConstraint(UTickableConstraint* In
 		return false;
 	}
 
+	const UE::Private::FScopedControllerContext Context(World);
+	
 	static constexpr bool bDoNotCompensate = true;
-	if (Subsystem->HasConstraint(World, InConstraint))
+	if (Subsystem->HasConstraint(Context.GetWorld(), InConstraint))
 	{
-		Subsystem->RemoveConstraint(World, InConstraint, bDoNotCompensate);
+		Subsystem->RemoveConstraint(Context.GetWorld(), InConstraint, bDoNotCompensate);
 		return true;
 	}
 	return false;
@@ -586,35 +591,55 @@ bool FConstraintsManagerController::UnregisterConstraint(UTickableConstraint* In
 bool FConstraintsManagerController::RemoveConstraint(const int32 InConstraintIndex,bool bDoNotCompensate) 
 {
 	UConstraintSubsystem* Subsystem = UConstraintSubsystem::Get();
-	if (!Subsystem)
+	if (!Subsystem || !World)
 	{
 		return false;
 	}
-	TArray<TWeakObjectPtr<UTickableConstraint>> Constraints = Subsystem->GetConstraints(World);
 
-	const FName ConstraintName = Constraints[InConstraintIndex]->GetFName();
-	UTickableConstraint* Constraint = Constraints[InConstraintIndex].Get();
-	
-	// notify deletion
-	auto GetRemoveNotifyType = [bDoNotCompensate]()
+	UTickableConstraint* Constraint = nullptr;
 	{
-		if(bDoNotCompensate)
+		const UE::Private::FScopedControllerContext Context(World);
+		const TArray<TWeakObjectPtr<UTickableConstraint>>& Constraints = Subsystem->GetConstraintsArray(Context.GetWorld());
+		if (Constraints.IsValidIndex(InConstraintIndex))
 		{
-			return EConstraintsManagerNotifyType::ConstraintRemoved;	
+			Constraint = Constraints[InConstraintIndex].Get();
 		}
-		return EConstraintsManagerNotifyType::ConstraintRemovedWithCompensation;
-	};
-	Notify(GetRemoveNotifyType(), Constraint);	
-	Subsystem->Modify();
-	Subsystem->RemoveConstraint(World, Constraint, bDoNotCompensate);
+	}
+
+	if (!Constraint)
+	{
+		return false;
+	}
+
+	// delete from subsystem and notify deletion
+	{
+		auto GetRemoveNotifyType = [bDoNotCompensate]()
+		{
+			if(bDoNotCompensate)
+			{
+				return EConstraintsManagerNotifyType::ConstraintRemoved;	
+			}
+			return EConstraintsManagerNotifyType::ConstraintRemovedWithCompensation;
+		};
+
+		const UE::Private::FScopedControllerContext Context(World);
+		Notify(GetRemoveNotifyType(), Constraint);	
+		Subsystem->Modify();
+		Subsystem->RemoveConstraint(Context.GetWorld(), Constraint, bDoNotCompensate);
+	}
 
 	if (UConstraintsManager* Manager = FindManager())
 	{
-		if (Manager->Constraints.Contains(Constraint) == true)
+		if (Manager->Constraints.Contains(Constraint))
 		{
-			Manager->Modify();
-			Manager->Constraints.Remove(Constraint);
-			Manager->OnConstraintRemoved_BP.Broadcast(Manager, Constraint, bDoNotCompensate);
+			// delete from manager and notify deletion
+			{
+				const UE::Private::FScopedControllerContext Context(World);
+				Manager->Modify();
+				Manager->Constraints.Remove(Constraint);
+				Manager->OnConstraintRemoved_BP.Broadcast(Manager, Constraint, bDoNotCompensate);
+			}
+			
 			// destroy constraints actor if no constraints left
 			if (Manager->Constraints.IsEmpty())
 			{
@@ -629,24 +654,32 @@ bool FConstraintsManagerController::RemoveConstraint(const int32 InConstraintInd
 
 UTickableConstraint* FConstraintsManagerController::GetConstraint(const FGuid& InGuid) const
 {
-	const int32 Index = GetConstraintIndex(InGuid);
-	if (Index == INDEX_NONE)
+	const UConstraintSubsystem* Subsystem = UConstraintSubsystem::Get();
+	if (!Subsystem || !World)
 	{
 		return nullptr;
 	}
 
-	return GetConstraint(Index);
+	const UE::Private::FScopedControllerContext Context(World);
+	const TArray<TWeakObjectPtr<UTickableConstraint>>& Constraints = Subsystem->GetConstraintsArray(Context.GetWorld());
+	const int32 Index = Constraints.IndexOfByPredicate([InGuid](const TWeakObjectPtr<UTickableConstraint>& Constraint)
+	{
+		return 	(Constraint.IsValid() && Constraint->ConstraintID == InGuid);
+	});
+
+	return Constraints.IsValidIndex(Index) ? Constraints[Index].Get() : nullptr;
 }
 
 UTickableConstraint* FConstraintsManagerController::GetConstraint(const int32 InConstraintIndex) const
 {
 	UConstraintSubsystem* Subsystem = UConstraintSubsystem::Get();
-	if (!Subsystem)
+	if (!Subsystem || !World)
 	{
 		return nullptr;
 	}
-	TArray<TWeakObjectPtr<UTickableConstraint>> Constraints = Subsystem->GetConstraints(World);
 
+	const UE::Private::FScopedControllerContext Context(World);
+	const TArray<TWeakObjectPtr<UTickableConstraint>>& Constraints = Subsystem->GetConstraintsArray(Context.GetWorld());
 	if (!Constraints.IsValidIndex(InConstraintIndex))
 	{
 		return nullptr;	
@@ -667,11 +700,13 @@ TArray< TWeakObjectPtr<UTickableConstraint> > FConstraintsManagerController::Get
 	}
 
 	UConstraintSubsystem* Subsystem = UConstraintSubsystem::Get();
-	if (!Subsystem)
+	if (!Subsystem || !World)
 	{
 		return DummyArray;
 	}
-	TArray<TWeakObjectPtr<UTickableConstraint>> Constraints = Subsystem->GetConstraints(World);
+
+	const UE::Private::FScopedControllerContext Context(World);
+	const TArray<TWeakObjectPtr<UTickableConstraint>>& Constraints = Subsystem->GetConstraintsArray(Context.GetWorld());
 
 	using ConstraintPtr = TWeakObjectPtr<UTickableConstraint>;
 
@@ -684,15 +719,19 @@ TArray< TWeakObjectPtr<UTickableConstraint> > FConstraintsManagerController::Get
 	if (bSorted)
 	{
 		// LHS ticks before RHS = LHS is a prerex of RHS 
-		auto TicksBefore = [this](const UTickableConstraint& LHS, const UTickableConstraint& RHS)
+		auto TicksBefore = [&Context](const UTickableConstraint& LHS, const UTickableConstraint& RHS)
 		{
-			const TArray<FTickPrerequisite>& RHSPrerex = RHS.GetTickFunction(World).GetPrerequisites();
-			const FConstraintTickFunction& LHSTickFunction = LHS.GetTickFunction(World);
-			const bool bIsLHSAPrerexOfRHS = RHSPrerex.ContainsByPredicate([&LHSTickFunction](const FTickPrerequisite& Prerex)
+			if (UWorld* CtxWorld = Context.GetWorld())
 			{
-				return Prerex.PrerequisiteTickFunction == &LHSTickFunction;
-			});
-			return bIsLHSAPrerexOfRHS;
+				const TArray<FTickPrerequisite>& RHSPrerex = RHS.GetTickFunction(CtxWorld).GetPrerequisites();
+				const FConstraintTickFunction& LHSTickFunction = LHS.GetTickFunction(CtxWorld);
+				const bool bIsLHSAPrerexOfRHS = RHSPrerex.ContainsByPredicate([&LHSTickFunction](const FTickPrerequisite& Prerex)
+				{
+					return Prerex.PrerequisiteTickFunction == &LHSTickFunction;
+				});
+				return bIsLHSAPrerexOfRHS;
+			}
+			return false;
 		};
 		
 		Algo::StableSort(FilteredConstraints, [TicksBefore](const ConstraintPtr& LHS, const ConstraintPtr& RHS)
@@ -709,28 +748,42 @@ void FConstraintsManagerController::SetConstraintsDependencies(
 	const FName& InNameToTickAfter) const
 {
 	UConstraintSubsystem* Subsystem = UConstraintSubsystem::Get();
-	if (!Subsystem)
+	if (!Subsystem || !World)
 	{
 		return;
 	}
-	const TArray<TWeakObjectPtr<UTickableConstraint>> Constraints = Subsystem->GetConstraints(World);
 
-	const int32 IndexBefore = GetConstraintIndex(InNameToTickBefore);
-	const int32 IndexAfter = GetConstraintIndex(InNameToTickAfter);
-	if (IndexBefore == INDEX_NONE || IndexAfter == INDEX_NONE || IndexAfter == IndexBefore)
+	const UE::Private::FScopedControllerContext Context(World);
+
+	const TArray<TWeakObjectPtr<UTickableConstraint>>& Constraints = Subsystem->GetConstraintsArray(Context.GetWorld());
+
+	auto GetIndex = [&Constraints](const FName& InName)
+	{
+		return Constraints.IndexOfByPredicate([InName](const TWeakObjectPtr<UTickableConstraint>& Constraint)
+		{
+			return (Constraint.IsValid() && Constraint->GetFName() == InName);
+		} );		
+	};
+	
+	const int32 IndexBefore = GetIndex(InNameToTickBefore);
+	const int32 IndexAfter = GetIndex(InNameToTickAfter);
+	if (!Constraints.IsValidIndex(IndexBefore) || !Constraints.IsValidIndex(IndexAfter) || IndexAfter == IndexBefore)
 	{
 		return;
 	}
 
-	FConstraintTickFunction& FunctionToTickBefore = Constraints[IndexBefore]->GetTickFunction(World);
-	FConstraintTickFunction& FunctionToTickAfter = Constraints[IndexAfter]->GetTickFunction(World);
+	if (UWorld* CtxWorld = Context.GetWorld())
+	{
+		FConstraintTickFunction& FunctionToTickBefore = Constraints[IndexBefore]->GetTickFunction(CtxWorld);
+		FConstraintTickFunction& FunctionToTickAfter = Constraints[IndexAfter]->GetTickFunction(CtxWorld);
 
-	Subsystem->SetConstraintDependencies(&FunctionToTickBefore, &FunctionToTickAfter);
+		Subsystem->SetConstraintDependencies(&FunctionToTickBefore, &FunctionToTickAfter);
 
-	InvalidateEvaluationGraph();
+		InvalidateEvaluationGraph();
+	}
 }
 
-void FConstraintsManagerController::SetConstraintsDependencies(const struct FGuid& InGuidToTickBefore, const struct FGuid& InGuidToTickAfter) const
+void FConstraintsManagerController::SetConstraintsDependencies(const FGuid& InGuidToTickBefore, const FGuid& InGuidToTickAfter) const
 {
 	if (!InGuidToTickBefore.IsValid() || !InGuidToTickAfter.IsValid() || InGuidToTickBefore == InGuidToTickAfter)
 	{
@@ -738,12 +791,13 @@ void FConstraintsManagerController::SetConstraintsDependencies(const struct FGui
 	}
 
 	UConstraintSubsystem* Subsystem = UConstraintSubsystem::Get();
-	if (!Subsystem)
+	if (!Subsystem || !World)
 	{
 		return;
 	}
 
-	const TArray<TWeakObjectPtr<UTickableConstraint>>& Constraints = Subsystem->GetConstraintsArray(World);
+	const UE::Private::FScopedControllerContext Context(World);
+	const TArray<TWeakObjectPtr<UTickableConstraint>>& Constraints = Subsystem->GetConstraintsArray(Context.GetWorld());
 	const int32 IndexBefore = Constraints.IndexOfByPredicate([InGuidToTickBefore](const TWeakObjectPtr<UTickableConstraint>& Constraint)
 	{
 		return Constraint.IsValid() && Constraint->ConstraintID == InGuidToTickBefore;
@@ -758,38 +812,52 @@ void FConstraintsManagerController::SetConstraintsDependencies(const struct FGui
 	{
 		return;
 	}
-	
-	FConstraintTickFunction& FunctionToTickBefore = Constraints[IndexBefore]->GetTickFunction(World);
-	FConstraintTickFunction& FunctionToTickAfter = Constraints[IndexAfter]->GetTickFunction(World);
 
-	Subsystem->SetConstraintDependencies(&FunctionToTickBefore, &FunctionToTickAfter);
+	if (UWorld* CtxWorld = Context.GetWorld())
+	{
+		FConstraintTickFunction& FunctionToTickBefore = Constraints[IndexBefore]->GetTickFunction(CtxWorld);
+		FConstraintTickFunction& FunctionToTickAfter = Constraints[IndexAfter]->GetTickFunction(CtxWorld);
+
+		Subsystem->SetConstraintDependencies(&FunctionToTickBefore, &FunctionToTickAfter);
 
-	InvalidateEvaluationGraph();
+		InvalidateEvaluationGraph();
+	}
 }
 
 const TArray< TWeakObjectPtr<UTickableConstraint> >& FConstraintsManagerController::GetConstraintsArray() const
 {
 	UConstraintSubsystem* Subsystem = UConstraintSubsystem::Get();
-	if (!Subsystem)
+	if (!Subsystem || !World)
 	{
 		static const TArray< TWeakObjectPtr<UTickableConstraint> > Empty;
 		return Empty;
 	}
-	const TArray<TWeakObjectPtr<UTickableConstraint>>& Constraints = Subsystem->GetConstraintsArray(World);
-	return Constraints;
+
+	const UE::Private::FScopedControllerContext Context(World);
+	return Subsystem->GetConstraintsArray(Context.GetWorld());
 }
 
 bool FConstraintsManagerController::RemoveAllConstraints(bool bDoNotCompensate) 
 {
-	if (UConstraintSubsystem* Subsystem = UConstraintSubsystem::Get())
+	UConstraintSubsystem* Subsystem = UConstraintSubsystem::Get();
+	if (Subsystem && World)
 	{
 		Subsystem->Modify();
-		TArray<TWeakObjectPtr<UTickableConstraint>> Constraints = Subsystem->GetConstraints(World);
+
+		TArray<TWeakObjectPtr<UTickableConstraint>> Constraints;
+		{
+			const UE::Private::FScopedControllerContext Context(World);
+			Constraints = Subsystem->GetConstraints(Context.GetWorld());
+		}
+
+		bool bRemoved = false;
 		for (TWeakObjectPtr<UTickableConstraint>& Constraint : Constraints)
 		{
-			RemoveConstraint(Constraint.Get(), bDoNotCompensate);
+			bRemoved |= RemoveConstraint(Constraint.Get(), bDoNotCompensate);
 		}
+		return bRemoved;
 	}
+	
 	return false;
 }
 
@@ -822,23 +890,25 @@ TArray< TWeakObjectPtr<UTickableConstraint> > FConstraintsManagerController::Get
 TArray< TWeakObjectPtr<UTickableConstraint> > FConstraintsManagerController::GetAllConstraints(const bool bSorted) const
 {
 	UConstraintSubsystem* Subsystem = UConstraintSubsystem::Get();
-	if (!Subsystem)
+	if (!Subsystem || !World)
 	{
 		static const TArray< TWeakObjectPtr<UTickableConstraint> > Empty;
 		return Empty;
 	}
 
+	const UE::Private::FScopedControllerContext Context(World);
+	
 	// use evaluation graph cached data
 	if (FConstraintsEvaluationGraph::UseEvaluationGraph() && bSorted)
 	{
 		TArray<TWeakObjectPtr<UTickableConstraint>> SortedConstraints;
-		if (Subsystem->GetEvaluationGraph(World).GetSortedConstraints(SortedConstraints))
+		if (Subsystem->GetEvaluationGraph(Context.GetWorld()).GetSortedConstraints(SortedConstraints))
 		{
 		   return SortedConstraints;
 		}
 	}
 
-	TArray<TWeakObjectPtr<UTickableConstraint>> Constraints = Subsystem->GetConstraints(World);
+	TArray<TWeakObjectPtr<UTickableConstraint>> Constraints = Subsystem->GetConstraints(Context.GetWorld());
 
 	// Remove stale constraints (note that GCd constraints should already have been caught by UConstraintSubsystem)
 	Constraints.RemoveAll([](const TWeakObjectPtr<UTickableConstraint>& ExistingConstraint) -> bool
@@ -851,7 +921,7 @@ TArray< TWeakObjectPtr<UTickableConstraint> > FConstraintsManagerController::Get
 		return Constraints;
 	}
 
-	UE::Constraints::Graph::SortConstraints(World, Constraints);
+	UE::Constraints::Graph::SortConstraints(Context.GetWorld(), Constraints);
 	
 	return Constraints;
 }
@@ -871,7 +941,19 @@ void FConstraintsManagerController::EvaluateAllConstraints() const
 	}
 }
 
-void FConstraintsManagerController::Notify(EConstraintsManagerNotifyType InNotifyType, UObject* InObject) const
+FConstraintsManagerController::FOnSceneComponentConstrained& FConstraintsManagerController::OnSceneComponentConstrained()
+{
+	static FOnSceneComponentConstrained SceneComponentConstrained;
+	return SceneComponentConstrained;
+}
+
+FConstraintsManagerNotifyDelegate& FConstraintsManagerController::GetNotifyDelegate()
+{
+	static FConstraintsManagerNotifyDelegate NotifyDelegate;
+	return NotifyDelegate;
+}
+
+void FConstraintsManagerController::Notify(EConstraintsManagerNotifyType InNotifyType, UObject* InObject)
 {
 	switch (InNotifyType)
 	{
@@ -893,28 +975,34 @@ void FConstraintsManagerController::Notify(EConstraintsManagerNotifyType InNotif
 			break;
 	}
 
-	NotifyDelegate.Broadcast(InNotifyType, InObject);
+	GetNotifyDelegate().Broadcast(InNotifyType, InObject);
 }
 
 void FConstraintsManagerController::MarkConstraintForEvaluation(UTickableConstraint* InConstraint) const
 {
 	UConstraintSubsystem* SubSystem = UConstraintSubsystem::Get();
-	if (!SubSystem)
+	if (!SubSystem || !World)
 	{
 		return;
 	}
-	SubSystem->GetEvaluationGraph(World).MarkForEvaluation(InConstraint);
+
+	const UE::Private::FScopedControllerContext Context(World);
+	SubSystem->GetEvaluationGraph(Context.GetWorld()).MarkForEvaluation(InConstraint);
 }
 
 void FConstraintsManagerController::InvalidateEvaluationGraph() const
 {
 	UConstraintSubsystem* SubSystem = UConstraintSubsystem::Get();
-	if (!SubSystem)
+	if (!SubSystem || !World)
 	{
 		return;
 	}
+
+	{
+		const UE::Private::FScopedControllerContext Context(World);
+		SubSystem->GetEvaluationGraph(Context.GetWorld()).InvalidateData();
+	}
 	
-	SubSystem->GetEvaluationGraph(World).InvalidateData();
 	Notify(EConstraintsManagerNotifyType::GraphUpdated, nullptr);
 }
 
@@ -925,12 +1013,13 @@ void FConstraintsManagerController::FlushEvaluationGraph() const
 		return;
 	}
 	UConstraintSubsystem* SubSystem = UConstraintSubsystem::Get();
-	if (!SubSystem)
+	if (!SubSystem || !World)
 	{
 		return;
 	}
 
-	FConstraintsEvaluationGraph& EvaluationGraph = SubSystem->GetEvaluationGraph(World);
+	const UE::Private::FScopedControllerContext Context(World);
+	FConstraintsEvaluationGraph& EvaluationGraph = SubSystem->GetEvaluationGraph(Context.GetWorld());
 	if (EvaluationGraph.IsPendingEvaluation())
 	{
 		EvaluationGraph.FlushPendingEvaluations();
@@ -945,17 +1034,19 @@ bool FConstraintsManagerController::DoesExistInAnyWorld(UTickableConstraint* InC
 	{
 		if (UConstraintSubsystem * Subsystem = UConstraintSubsystem::Get())
 		{
-			UWorld* CurrentWorld = World; //save current World
+			const UE::Private::FScopedControllerContext CtrlContext(World);
+			
 			for (const FWorldContext& Context : GEngine->GetWorldContexts())
 			{
-				World = Context.World();
-				if (Subsystem->GetConstraintsArray(World).Find(InConstraint) != INDEX_NONE)
+				if (Context.World())
 				{
-					bFound = true;
-					break;
+					if (Subsystem->GetConstraintsArray(Context.World()).Find(InConstraint) != INDEX_NONE)
+					{
+						bFound = true;
+						break;
+					}
 				}
 			}
-			World = CurrentWorld; //restore current World
 		}
 	}
 	return bFound;
diff --git a/Engine/Source/Runtime/Experimental/Animation/Constraints/Private/ConstraintsScripting.cpp b/Engine/Source/Runtime/Experimental/Animation/Constraints/Private/ConstraintsScripting.cpp
--- a/Engine/Source/Runtime/Experimental/Animation/Constraints/Private/ConstraintsScripting.cpp
+++ b/Engine/Source/Runtime/Experimental/Animation/Constraints/Private/ConstraintsScripting.cpp
@@ -80,7 +80,7 @@ bool UConstraintsScriptingLibrary::AddConstraint(UWorld* InWorld, UTransformable
 	}
 	
 	FConstraintsManagerController& Controller = FConstraintsManagerController::Get(InWorld);
-	Controller.StaticConstraintCreated(InWorld, InConstraint);
+	Controller.StaticConstraintCreated(InConstraint);
 	
 	return true;
 }
diff --git a/Engine/Source/Runtime/Experimental/Animation/Constraints/Private/Transform/AnimationEvaluation.cpp b/Engine/Source/Runtime/Experimental/Animation/Constraints/Private/Transform/AnimationEvaluation.cpp
--- a/Engine/Source/Runtime/Experimental/Animation/Constraints/Private/Transform/AnimationEvaluation.cpp
+++ b/Engine/Source/Runtime/Experimental/Animation/Constraints/Private/Transform/AnimationEvaluation.cpp
@@ -470,23 +470,24 @@ const FAnimationEvaluator& FAnimationEvaluationCache::GetEvaluator(USkeletalMesh
 
 void FAnimationEvaluationCache::RegisterNotifications()
 {
-	FConstraintsManagerController& Controller = FConstraintsManagerController::Get(nullptr);
-	ConstraintsNotificationHandle =
-		Controller.GetNotifyDelegate().AddLambda([this](EConstraintsManagerNotifyType InNotifyType, UObject* InObject)
+	if (!ConstraintsNotificationHandle.IsValid())
 	{
-		if (InNotifyType == EConstraintsManagerNotifyType::GraphUpdated)
-		{
-			PerSkeletalMeshEvaluator.Reset();
-		}
-	});
+		ConstraintsNotificationHandle =
+		   FConstraintsManagerController::GetNotifyDelegate().AddLambda([this](EConstraintsManagerNotifyType InNotifyType, UObject* InObject)
+	   {
+		   if (InNotifyType == EConstraintsManagerNotifyType::GraphUpdated)
+		   {
+			   PerSkeletalMeshEvaluator.Reset();
+		   }
+	   });
+	}
 }
 
 void FAnimationEvaluationCache::UnregisterNotifications()
 {
 	if (ConstraintsNotificationHandle.IsValid())
 	{
-		FConstraintsManagerController& Controller = FConstraintsManagerController::Get(nullptr);
-		Controller.GetNotifyDelegate().Remove(ConstraintsNotificationHandle);
+		FConstraintsManagerController::GetNotifyDelegate().Remove(ConstraintsNotificationHandle);
 		ConstraintsNotificationHandle.Reset();
 	}
 }
diff --git a/Engine/Source/Runtime/Experimental/Animation/Constraints/Private/Transform/TransformConstraintUtil.cpp b/Engine/Source/Runtime/Experimental/Animation/Constraints/Private/Transform/TransformConstraintUtil.cpp
--- a/Engine/Source/Runtime/Experimental/Animation/Constraints/Private/Transform/TransformConstraintUtil.cpp
+++ b/Engine/Source/Runtime/Experimental/Animation/Constraints/Private/Transform/TransformConstraintUtil.cpp
@@ -126,7 +126,7 @@ TOptional<FTransform> FConstraintsInteractionCache::GetParentTransform(const uin
 	if (ConstraintsIndex != INDEX_NONE)
 	{
 		UTickableTransformConstraint* Constraint = Cast<UTickableTransformConstraint>(TransformConstraints[ConstraintsIndex].Get());
-		if (Constraint && Constraint->ParentTRSHandle)
+		if (Constraint && IsValid(Constraint->ParentTRSHandle))
 		{
 			return Constraint->ParentTRSHandle->GetGlobalTransform();	
 		}
@@ -251,25 +251,26 @@ void FConstraintsInteractionCache::Reset()
 	PerTargetConstraints.Reset();
 }
 
-void FConstraintsInteractionCache::RegisterNotifications(UWorld* InWorld)
+void FConstraintsInteractionCache::RegisterNotifications()
 {
-	FConstraintsManagerController& Controller = FConstraintsManagerController::Get(InWorld);
-	ConstraintsNotificationHandle =
-		Controller.GetNotifyDelegate().AddLambda([this](EConstraintsManagerNotifyType InNotifyType, UObject* InObject)
+	if (!ConstraintsNotificationHandle.IsValid())
 	{
-		if (InNotifyType == EConstraintsManagerNotifyType::GraphUpdated)
-		{
-			Reset();
-		}
-	});
+		ConstraintsNotificationHandle =
+		   FConstraintsManagerController::GetNotifyDelegate().AddLambda([this](EConstraintsManagerNotifyType InNotifyType, UObject* InObject)
+	   {
+		   if (InNotifyType == EConstraintsManagerNotifyType::GraphUpdated)
+		   {
+			   Reset();
+		   }
+	   });
+	}
 }
 
-void FConstraintsInteractionCache::UnregisterNotifications(UWorld* InWorld)
+void FConstraintsInteractionCache::UnregisterNotifications()
 {
 	if (ConstraintsNotificationHandle.IsValid())
 	{
-		FConstraintsManagerController& Controller = FConstraintsManagerController::Get(InWorld);
-		Controller.GetNotifyDelegate().Remove(ConstraintsNotificationHandle);
+		FConstraintsManagerController::GetNotifyDelegate().Remove(ConstraintsNotificationHandle);
 		ConstraintsNotificationHandle.Reset();
 	}
 }
diff --git a/Engine/Source/Runtime/Experimental/Animation/Constraints/Private/Transform/TransformableHandle.cpp b/Engine/Source/Runtime/Experimental/Animation/Constraints/Private/Transform/TransformableHandle.cpp
--- a/Engine/Source/Runtime/Experimental/Animation/Constraints/Private/Transform/TransformableHandle.cpp
+++ b/Engine/Source/Runtime/Experimental/Animation/Constraints/Private/Transform/TransformableHandle.cpp
@@ -142,11 +142,9 @@ void UTransformableComponentHandle::SetGlobalTransform(const FTransform& InGloba
 		
 		Component->SetWorldTransform(InGlobal);
 		
-		UWorld* World = Component->GetWorld();
-		if (World)
+		if (Component->GetWorld())
 		{
-			FConstraintsManagerController& Controller = FConstraintsManagerController::Get(World);
-			Controller.OnSceneComponentConstrained().Broadcast(Component.Get());
+			FConstraintsManagerController::OnSceneComponentConstrained().Broadcast(Component.Get());
 		}
 
 		if (bMarkRenderDirty)
@@ -161,11 +159,9 @@ void UTransformableComponentHandle::SetLocalTransform(const FTransform& InLocal)
 	if(Component.IsValid())
 	{
 		Component->SetRelativeTransform(InLocal);
-		UWorld* World = Component->GetWorld();
-		if (World)
+		if (Component->GetWorld())
 		{
-			FConstraintsManagerController& Controller = FConstraintsManagerController::Get(World);
-			Controller.OnSceneComponentConstrained().Broadcast(Component.Get());
+			FConstraintsManagerController::OnSceneComponentConstrained().Broadcast(Component.Get());
 		}
 	}
 }
diff --git a/Engine/Source/Runtime/Experimental/Animation/Constraints/Public/ConstraintsManager.h b/Engine/Source/Runtime/Experimental/Animation/Constraints/Public/ConstraintsManager.h
--- a/Engine/Source/Runtime/Experimental/Animation/Constraints/Public/ConstraintsManager.h
+++ b/Engine/Source/Runtime/Experimental/Animation/Constraints/Public/ConstraintsManager.h
@@ -302,28 +302,25 @@ public:
 		const FName& InNameToTickAfter) const;
 
 	CONSTRAINTS_API void SetConstraintsDependencies(
-		const struct FGuid& InGuidToTickBefore,
-		const struct FGuid& InGuidToTickAfter) const;
+		const FGuid& InGuidToTickBefore,
+		const FGuid& InGuidToTickAfter) const;
 
 	/** Go through each constraint in order and evaluate and tick them*/
 	CONSTRAINTS_API void EvaluateAllConstraints() const;
 
 	/** when PIEing/Simulating it's possible that the constraint isn't in the active manager but still lives*/
+	UE_DEPRECATED(5.7, "Function is no longer needed.")
 	CONSTRAINTS_API bool DoesExistInAnyWorld(UTickableConstraint* InConstraint);
 
-	/** If static contraint created, we call this, which stores it to the manager/actor*/
+	UE_DEPRECATED(5.7, "Function is no longer needed.")
 	CONSTRAINTS_API void StaticConstraintCreated(UWorld* InWorld, UTickableConstraint* InConstraint);
 
+	/** If a static constraint is created, we call this, which stores it to the manager/actor*/
+	CONSTRAINTS_API void StaticConstraintCreated(UTickableConstraint* InConstraint);
+
 private:
 	/** Get the index of the given constraint's name. */
 	int32 GetConstraintIndex(const FName& InConstraintName) const;
-
-	/** Delegeate that's fired when a scene component is constrained, this is needed to make sure things like gizmo's get updated after the constraint tick happens*/
-	DECLARE_MULTICAST_DELEGATE_OneParam(FOnSceneComponentConstrained, USceneComponent* /*InSceneComponent*/);
-	FOnSceneComponentConstrained SceneComponentConstrained;
-
-	/** Delegate to trigger changes in the constraints manager. */
-	FConstraintsManagerNotifyDelegate NotifyDelegate;
 	
 	/** Find the existing Constraints Manager in World or create a new one. */
 	CONSTRAINTS_API UConstraintsManager* GetManager();
@@ -339,13 +336,14 @@ private:
 
 public:
 	/** Delegate that's fired when a scene component is constrained, this is needed to make sure things like gizmo's get updated after the constraint tick happens*/
-	FOnSceneComponentConstrained& OnSceneComponentConstrained() { return SceneComponentConstrained; }
+	DECLARE_MULTICAST_DELEGATE_OneParam(FOnSceneComponentConstrained, USceneComponent* /*InSceneComponent*/);
+	CONSTRAINTS_API static FOnSceneComponentConstrained& OnSceneComponentConstrained();
 
 	/** Delegate to trigger changes in the constraints manager. */
-	FConstraintsManagerNotifyDelegate& GetNotifyDelegate() { return NotifyDelegate; }
+	CONSTRAINTS_API static FConstraintsManagerNotifyDelegate& GetNotifyDelegate();
 
 	/** Notify from changes in the constraints manager. */
-	CONSTRAINTS_API void Notify(EConstraintsManagerNotifyType InNotifyType, UObject* InObject) const;
+	CONSTRAINTS_API static void Notify(EConstraintsManagerNotifyType InNotifyType, UObject* InObject);
 
 	/** Evaluation graph API. */
 	CONSTRAINTS_API void MarkConstraintForEvaluation(UTickableConstraint* InConstraint) const;
diff --git a/Engine/Source/Runtime/Experimental/Animation/Constraints/Public/Transform/TransformConstraintUtil.h b/Engine/Source/Runtime/Experimental/Animation/Constraints/Public/Transform/TransformConstraintUtil.h
--- a/Engine/Source/Runtime/Experimental/Animation/Constraints/Public/Transform/TransformConstraintUtil.h
+++ b/Engine/Source/Runtime/Experimental/Animation/Constraints/Public/Transform/TransformConstraintUtil.h
@@ -45,8 +45,8 @@ public:
 
 	UE_API bool HasAnyDependency(UObject* InChild, UObject* InParent, UWorld* InWorld);
 	
-	UE_API void RegisterNotifications(UWorld* InWorld);
-	UE_API void UnregisterNotifications(UWorld* InWorld);
+	UE_API void RegisterNotifications();
+	UE_API void UnregisterNotifications();
 
 private:
 
