Constraints: recursive anim evaluation (follow-up)
- mark the parent for evaluation at least once as the constraint being created is not evaluated
- if any of the component's parent is a skeletal mesh and has a valid evaluator, then it will be used to compose the global transform
- get the transform from the last cached component space transforms (when possible) if the evluator update couldn't finish (existing parallel evaluation or post-evaluating)

#rnx

diff --git a/Engine/Plugins/Animation/ControlRig/Source/ControlRig/Private/Sequencer/MovieSceneControlRigParameterTemplate.cpp b/Engine/Plugins/Animation/ControlRig/Source/ControlRig/Private/Sequencer/MovieSceneControlRigParameterTemplate.cpp
--- a/Engine/Plugins/Animation/ControlRig/Source/ControlRig/Private/Sequencer/MovieSceneControlRigParameterTemplate.cpp
+++ b/Engine/Plugins/Animation/ControlRig/Source/ControlRig/Private/Sequencer/MovieSceneControlRigParameterTemplate.cpp
@@ -1482,7 +1482,8 @@ struct FControlRigParameterExecutionToken : IMovieSceneExecutionToken
 								if (UTransformableComponentHandle* ParentHandle = Cast<UTransformableComponentHandle>(TransformConstraint->ParentTRSHandle))
 								{
 									constexpr bool bRecursive = true;
-									TransformableHandleUtils::MarkComponentForEvaluation(Cast<USkeletalMeshComponent>(ParentHandle->GetTarget()), bRecursive);
+									USceneComponent* ParentComponent = TransformableHandleUtils::GetTarget<USceneComponent>(TransformConstraint->ParentTRSHandle);
+									TransformableHandleUtils::MarkComponentForEvaluation(ParentComponent, bRecursive);
 								}
 							}
 						}
diff --git a/Engine/Plugins/Animation/ControlRig/Source/ControlRigEditor/Private/EditMode/ControlRigEditMode.cpp b/Engine/Plugins/Animation/ControlRig/Source/ControlRigEditor/Private/EditMode/ControlRigEditMode.cpp
--- a/Engine/Plugins/Animation/ControlRig/Source/ControlRigEditor/Private/EditMode/ControlRigEditMode.cpp
+++ b/Engine/Plugins/Animation/ControlRig/Source/ControlRigEditor/Private/EditMode/ControlRigEditMode.cpp
@@ -89,6 +89,7 @@
 #include "Picker/FloatingSpacePickerManager.h"
 #include "Sequencer/AnimLayers/AnimLayers.h"
 #include "Tools/MotionTrailOptions.h"
+#include "Transform/TransformableHandleUtils.h"
 
 #include UE_INLINE_GENERATED_CPP_BY_NAME(ControlRigEditMode)
 
@@ -783,10 +784,21 @@ void FControlRigEditMode::Tick(FEditorViewportClient* ViewportClient, float Delt
 			UControlRig* Rig = NonInteractingRig.Get();
 			if (USceneComponent* BoundComponent = Rig ? GetHostingSceneComponent(Rig) : nullptr)
 			{
+				TransformableHandleUtils::MarkComponentForEvaluation(BoundComponent, false);
+				
 				for (const TWeakObjectPtr<UTickableConstraint>& WeakConstraint: ConstraintsCache.Get(BoundComponent, WorldPtr))
 				{
 					if (UTickableConstraint* Constraint = WeakConstraint.Get(); Constraint->Active)
 					{
+						if (UTickableParentConstraint* ParentConst = Cast<UTickableParentConstraint>(Constraint))
+						{
+							if (USceneComponent* ParentComponent = TransformableHandleUtils::GetTarget<USceneComponent>(ParentConst->ParentTRSHandle))
+							{
+								constexpr bool bRecursive = true;
+								TransformableHandleUtils::MarkComponentForEvaluation(ParentComponent, bRecursive);								
+							}
+						}
+						
 						FConstraintsManagerController::Get(WorldPtr).MarkConstraintForEvaluation(Constraint);
 					}
 				}
@@ -7112,7 +7124,9 @@ void FControlRigEditMode::HandleActorMoving(AActor* InMovedActor)
 {
 	if (!bUpdateNonInteractingRigs && InMovedActor)
 	{
-		bUpdateNonInteractingRigs = SelectedObjects.Contains(InMovedActor->GetRootComponent());
+		bUpdateNonInteractingRigs =
+			SelectedObjects.Contains(InMovedActor) ||
+			SelectedObjects.Contains(InMovedActor->GetRootComponent());
 
 		if (!bUpdateNonInteractingRigs)
 		{
@@ -7154,10 +7168,7 @@ void FControlRigEditMode::CacheAnyOtherSelectedObjects()
 		{
 			if (AActor* Actor = Cast<AActor>(WeakActor))
 			{
-				if (USceneComponent* SceneComponent = Actor->GetRootComponent())
-				{
-					SelectedObjects.Add(SceneComponent);
-				}
+				SelectedObjects.Add(Actor);
 			}
 		}
 		return;
@@ -7241,22 +7252,27 @@ void FControlRigEditMode::CacheNonInteractingRigsToUpdate()
 			{
 				if (USceneComponent* BoundComponent = GetHostingSceneComponent(NonInteractingRig))
 				{
+					FComponentDependency ComponentDependency(BoundComponent, WorldPtr, ConstraintsCache);
+					
 					for (const TWeakObjectPtr<UObject>& SelectedObject: SelectedObjects)
 					{
 						if (UObject* Object = SelectedObject.Get())
 						{
 							if (AActor* Actor = Cast<AActor>(Object))
 							{
-								Object = Actor->GetRootComponent();
-							}
-							
-							if (BoundComponent == Object)
-							{
-								NonInteractingRigs.Add(NonInteractingRig);
-								break;
-							}
+								const TInlineComponentArray<USceneComponent*> Components(Actor);
+								const bool bDepends = Components.ContainsByPredicate([&ComponentDependency](USceneComponent* Component)
+								{
+									return ComponentDependency.DependsOn(Component);
+								});
 
-							if (ConstraintsCache.HasAnyDependency(BoundComponent, Object, WorldPtr))
+								if (bDepends)
+								{
+									NonInteractingRigs.Add(NonInteractingRig);
+									break;
+								}
+							}
+							else if (ComponentDependency.DependsOn(Object))
 							{
 								NonInteractingRigs.Add(NonInteractingRig);
 								break;
diff --git a/Engine/Plugins/Animation/ControlRig/Source/ControlRigEditor/Private/Sequencer/EditModeAnimationUtil.cpp b/Engine/Plugins/Animation/ControlRig/Source/ControlRigEditor/Private/Sequencer/EditModeAnimationUtil.cpp
--- a/Engine/Plugins/Animation/ControlRig/Source/ControlRigEditor/Private/Sequencer/EditModeAnimationUtil.cpp
+++ b/Engine/Plugins/Animation/ControlRig/Source/ControlRigEditor/Private/Sequencer/EditModeAnimationUtil.cpp
@@ -13,6 +13,7 @@
 #include "LevelSequence.h"
 #include "Rigs/RigHierarchyDefines.h"
 #include "Sequencer/AnimationAuthoringSettings.h"
+#include "Transform/TransformConstraintUtil.h"
 
 namespace UE::AnimationEditMode
 {
@@ -231,5 +232,46 @@ void FControlRigKeyframer::OnSettingsChanged(const UAnimationAuthoringSettings*
 		}
 	}
 }
+
+FComponentDependency::FComponentDependency( USceneComponent* InComponent,UWorld* InWorld,
+	TransformConstraintUtil::FConstraintsInteractionCache& InCacheRef)
+	: Component(InComponent)
+	, World(InWorld)
+	, ConstraintsCache(InCacheRef)
+{}
+
+bool FComponentDependency::DependsOn(UObject* InObject)
+{
+	if (this->IsValid(InObject))
+	{
+		if (Component == InObject)
+		{
+			return true;
+		}
+
+		if (ConstraintsCache.HasAnyDependency(Component, InObject, World))
+		{
+			return true;
+		}
+
+		if (USceneComponent* InComponent = Cast<USceneComponent>(InObject))
+		{
+			for (const TObjectPtr<USceneComponent>& Child: InComponent->GetAttachChildren())
+			{
+				if (DependsOn(Child.Get()))
+				{
+					return true;
+				}
+			}
+		}
+	}
+    				
+	return false;
+}
+
+bool FComponentDependency::IsValid(const UObject* InObject) const
+{
+	return Component && World && InObject;
+}
 	
 }
diff --git a/Engine/Plugins/Animation/ControlRig/Source/ControlRigEditor/Public/Sequencer/EditModeAnimationUtil.h b/Engine/Plugins/Animation/ControlRig/Source/ControlRigEditor/Public/Sequencer/EditModeAnimationUtil.h
--- a/Engine/Plugins/Animation/ControlRig/Source/ControlRigEditor/Public/Sequencer/EditModeAnimationUtil.h
+++ b/Engine/Plugins/Animation/ControlRig/Source/ControlRigEditor/Public/Sequencer/EditModeAnimationUtil.h
@@ -16,6 +16,11 @@ class UControlRig;
 struct FRigControlElement;
 struct FControlRigInteractionTransformContext;
 
+namespace UE::TransformConstraintUtil
+{
+	struct FConstraintsInteractionCache;
+}
+
 namespace UE::AnimationEditMode
 {
 
@@ -132,5 +137,27 @@ private:
 	// Used to track changes to animation authoring settings.
 	void OnSettingsChanged(const UAnimationAuthoringSettings* InSettings);
 };
+
+/**
+ * FComponentDependency helps query dependencies between scene components from a constraint standpoint   
+ */
+
+struct FComponentDependency
+{
+	FComponentDependency(USceneComponent* InComponent, UWorld* InWorld, TransformConstraintUtil::FConstraintsInteractionCache& InCacheRef);
+
+	/**
+	 * Returns true if the stored component depends on InObject or one of its children if it's a scene component.
+	 */
+	bool DependsOn(UObject* InObject);
 	
+private:
+
+	bool IsValid(const UObject* InObject) const;
+
+	USceneComponent* Component = nullptr;
+	UWorld* World = nullptr;
+	TransformConstraintUtil::FConstraintsInteractionCache& ConstraintsCache;
+};
+
 }
diff --git a/Engine/Source/Editor/MovieSceneTools/Private/Constraints/MovieSceneConstraintChannelHelper.cpp b/Engine/Source/Editor/MovieSceneTools/Private/Constraints/MovieSceneConstraintChannelHelper.cpp
--- a/Engine/Source/Editor/MovieSceneTools/Private/Constraints/MovieSceneConstraintChannelHelper.cpp
+++ b/Engine/Source/Editor/MovieSceneTools/Private/Constraints/MovieSceneConstraintChannelHelper.cpp
@@ -22,6 +22,7 @@
 #include "Channels/MovieSceneChannelTraits.h"
 #include "Constraints/TransformConstraintChannelInterface.h"
 #include "Framework/Notifications/NotificationManager.h"
+#include "Transform/TransformableHandleUtils.h"
 #include "Widgets/Notifications/SNotificationList.h"
 
 /*
@@ -39,7 +40,7 @@ void FCompensationEvaluator::ComputeLocalTransforms(UWorld* InWorld, const FEval
 {
 	using namespace UE::TransformConstraintUtil;
 	
-	if (!InEvalParams.IsValid())
+	if (!Constraint || !Handle || !InEvalParams.IsValid())
 	{
 		return;
 	}
@@ -51,11 +52,10 @@ void FCompensationEvaluator::ComputeLocalTransforms(UWorld* InWorld, const FEval
 		return;
 	}
 	
-	const TArray<ConstraintPtr> ConstraintsMinusThis =
-		Constraints.FilterByPredicate([this](const ConstraintPtr& InConstraint)
-		{
-			return InConstraint != Constraint;
-		});
+	const TArray<ConstraintPtr> ConstraintsMinusThis = Constraints.FilterByPredicate([this](const ConstraintPtr& InConstraint)
+	{
+		return InConstraint != Constraint;
+	});
 
 	// find last active constraint in the list that is different than the on we want to compensate for
 	auto GetLastActiveConstraint = [ConstraintsMinusThis]()
@@ -66,7 +66,29 @@ void FCompensationEvaluator::ComputeLocalTransforms(UWorld* InWorld, const FEval
 		// if found, return its parent global transform
 		return LastActiveIndex > INDEX_NONE ? Cast<UTickableTransformConstraint>(ConstraintsMinusThis[LastActiveIndex]) : nullptr;
 	};
-
+	
+	bool bShouldPrepareParent = false;
+	if (TransformableHandleUtils::SkipTicking())
+	{
+		USceneComponent* ChildComponent = TransformableHandleUtils::GetTarget<USceneComponent>(Handle);
+		USceneComponent* ParentComponent = TransformableHandleUtils::GetTarget<USceneComponent>(Constraint->ParentTRSHandle);
+		if (ParentComponent && ParentComponent != ChildComponent)
+		{
+			bShouldPrepareParent = [](const USceneComponent* Component)
+			{
+				while (Component)
+				{
+					if (Component->IsA<USkeletalMeshComponent>())
+					{
+						return true;
+					}
+					Component = Component->GetAttachParent();
+				}
+				return false;
+			}(ParentComponent);
+		}
+	}
+	
 	ISequencer* Sequencer = InEvalParams.Sequencer;
 	UMovieScene* MovieScene = Sequencer->GetFocusedMovieSceneSequence()->GetMovieScene();
 	const FFrameRate TickResolution = MovieScene->GetTickResolution();
@@ -122,6 +144,13 @@ void FCompensationEvaluator::ComputeLocalTransforms(UWorld* InWorld, const FEval
 		
 			Sequencer->GetEvaluationTemplate().EvaluateSynchronousBlocking(Context);
 
+			if (bShouldPrepareParent)
+			{
+				constexpr bool bRecursive = true;
+				USceneComponent* ParentComponent = TransformableHandleUtils::GetTarget<USceneComponent>(Constraint->ParentTRSHandle);
+				TransformableHandleUtils::MarkComponentForEvaluation(ParentComponent, bRecursive);
+			}
+			
 			// evaluate constraints
 			for (const TWeakObjectPtr<UTickableConstraint>& InConstraint : AllConstraints)
 			{
diff --git a/Engine/Source/Runtime/Experimental/Animation/Constraints/Private/Transform/AnimationEvaluation.cpp b/Engine/Source/Runtime/Experimental/Animation/Constraints/Private/Transform/AnimationEvaluation.cpp
--- a/Engine/Source/Runtime/Experimental/Animation/Constraints/Private/Transform/AnimationEvaluation.cpp
+++ b/Engine/Source/Runtime/Experimental/Animation/Constraints/Private/Transform/AnimationEvaluation.cpp
@@ -77,9 +77,20 @@ void FAnimationEvaluator::Update(const bool bRefreshBoneTransforms)
 	}
 }
 
-bool FAnimationEvaluator::IsValid() const
+bool FAnimationEvaluator::IsValid(const bool bForTransform) const
 {
-	return SkeletalMeshComponent.IsValid() && Context.SkeletalMesh && !Context.ComponentSpaceTransforms.IsEmpty();
+	if  (!SkeletalMeshComponent.IsValid() || Context.ComponentSpaceTransforms.IsEmpty())
+	{
+		return false;
+	}
+
+	if (bForTransform)
+	{
+		const TArray<FTransform>& ComponentSpaceTransforms = SkeletalMeshComponent->GetComponentSpaceTransforms();
+		return Context.ComponentSpaceTransforms.Num() == ComponentSpaceTransforms.Num();
+	}
+	
+	return Context.SkeletalMesh != nullptr;
 }
 
 FTransform FAnimationEvaluator::GetGlobalTransform(const FName InSocketName) const
@@ -153,8 +164,11 @@ FTransform FAnimationEvaluator::GetRelativeTransform(const FName InSocketName) c
 	};
 	
 	const FTransform& RelativeTransform = RawSkeletalMeshComponent->GetRelativeTransform();
-	const TArray<FTransform>& ComponentSpaceTransforms = IsValid() ?
+
+	constexpr bool bValidForTransforms = true;
+	const TArray<FTransform>& ComponentSpaceTransforms = IsValid(bValidForTransforms) ?
 		Context.ComponentSpaceTransforms : RawSkeletalMeshComponent->GetComponentSpaceTransforms();
+	
 	return InSocketName == NAME_None ? RelativeTransform : GetSocketTransform(ComponentSpaceTransforms) * RelativeTransform;
 }
 
diff --git a/Engine/Source/Runtime/Experimental/Animation/Constraints/Private/Transform/TransformableHandle.cpp b/Engine/Source/Runtime/Experimental/Animation/Constraints/Private/Transform/TransformableHandle.cpp
--- a/Engine/Source/Runtime/Experimental/Animation/Constraints/Private/Transform/TransformableHandle.cpp
+++ b/Engine/Source/Runtime/Experimental/Animation/Constraints/Private/Transform/TransformableHandle.cpp
@@ -173,19 +173,22 @@ FTransform UTransformableComponentHandle::GetLocalTransform() const
 
 FTransform UTransformableComponentHandle::GetGlobalTransform() const
 {
-	if (!Component.IsValid())
+	USceneComponent* SceneComponent = Component.Get();
+	if (!SceneComponent)
 	{
 		return FTransform::Identity;
 	}
 	
-	using namespace TransformableHandleUtils;
-	if (SkipTicking())
+	if (TransformableHandleUtils::SkipTicking())
 	{
-		const UE::Anim::FAnimationEvaluator& AnimationEvaluator = EvaluateComponent(Component.Get());
-		if (AnimationEvaluator.IsValid())
+		constexpr bool bValidForTransforms = true;
+		const UE::Anim::FAnimationEvaluator& AnimationEvaluator = TransformableHandleUtils::EvaluateComponent(SceneComponent);
+		if (AnimationEvaluator.IsValid(bValidForTransforms))
 		{
 			return AnimationEvaluator.GetGlobalTransform(SocketName);
 		}
+		
+		return TransformableHandleUtils::GetGlobalTransform(SceneComponent, SocketName);
 	}
 	
 	return Component.IsValid() ? Component->GetSocketTransform(SocketName) : FTransform::Identity;
@@ -524,10 +527,13 @@ bool UTransformableComponentHandle::AddTransformKeys(const TArray<FFrameNumber>&
 
 	for (int32 ChannelIndex = 0; ChannelIndex < 9; ChannelIndex++)
 	{
-		if (!Channels[ChannelIndex]->GetDefault().IsSet())
+		if (Channels.IsValidIndex(ChannelIndex))
 		{
-			const double Value = GetValue(ChannelIndex, Location0, Rotation0, Scale3D0);
-			Channels[ChannelIndex]->SetDefault(Value);
+			if (!Channels[ChannelIndex]->GetDefault().IsSet())
+			{
+				const double Value = GetValue(ChannelIndex, Location0, Rotation0, Scale3D0);
+				Channels[ChannelIndex]->SetDefault(Value);
+			}
 		}
 	}
 
@@ -543,27 +549,30 @@ bool UTransformableComponentHandle::AddTransformKeys(const TArray<FFrameNumber>&
 
 		for (const int32 ChannelIndex : ChannelsIndexToKey)
 		{
-			const double Value = GetValue(ChannelIndex, Location, Rotation, Scale3D);
-			TMovieSceneChannelData<FMovieSceneDoubleValue> ChannelData = Channels[ChannelIndex]->GetData();
+			if (Channels.IsValidIndex(ChannelIndex))
 			{
-				int32 ExistingIndex = ChannelData.FindKey(Frame);
-				if (ExistingIndex != INDEX_NONE)
-				{
-					FMovieSceneDoubleValue& DoubleValue = ChannelData.GetValues()[ExistingIndex]; //-V758
-					DoubleValue.Value = Value;
-				}
-				else
+				const double Value = GetValue(ChannelIndex, Location, Rotation, Scale3D);
+				TMovieSceneChannelData<FMovieSceneDoubleValue> ChannelData = Channels[ChannelIndex]->GetData();
 				{
-					FMovieSceneDoubleValue NewKey(Value);
-					ERichCurveTangentWeightMode WeightedMode = RCTWM_WeightedNone;
-					NewKey.InterpMode = ERichCurveInterpMode::RCIM_Cubic;
-					NewKey.TangentMode = ERichCurveTangentMode::RCTM_Auto;
-					NewKey.Tangent.ArriveTangent = 0.0f;
-					NewKey.Tangent.LeaveTangent = 0.0f;
-					NewKey.Tangent.TangentWeightMode = WeightedMode;
-					NewKey.Tangent.ArriveTangentWeight = 0.0f;
-					NewKey.Tangent.LeaveTangentWeight = 0.0f;
-					ChannelData.AddKey(Frame, NewKey);
+					int32 ExistingIndex = ChannelData.FindKey(Frame);
+					if (ExistingIndex != INDEX_NONE)
+					{
+						FMovieSceneDoubleValue& DoubleValue = ChannelData.GetValues()[ExistingIndex]; //-V758
+						DoubleValue.Value = Value;
+					}
+					else
+					{
+						FMovieSceneDoubleValue NewKey(Value);
+						ERichCurveTangentWeightMode WeightedMode = RCTWM_WeightedNone;
+						NewKey.InterpMode = ERichCurveInterpMode::RCIM_Cubic;
+						NewKey.TangentMode = ERichCurveTangentMode::RCTM_Auto;
+						NewKey.Tangent.ArriveTangent = 0.0f;
+						NewKey.Tangent.LeaveTangent = 0.0f;
+						NewKey.Tangent.TangentWeightMode = WeightedMode;
+						NewKey.Tangent.ArriveTangentWeight = 0.0f;
+						NewKey.Tangent.LeaveTangentWeight = 0.0f;
+						ChannelData.AddKey(Frame, NewKey);
+					}
 				}
 			}
 		}
@@ -572,7 +581,10 @@ bool UTransformableComponentHandle::AddTransformKeys(const TArray<FFrameNumber>&
 	//now we need to set auto tangents
 	for (const int32 ChannelIndex : ChannelsIndexToKey)
 	{
-		Channels[ChannelIndex]->AutoSetTangents();
+		if (Channels.IsValidIndex(ChannelIndex))
+		{
+			Channels[ChannelIndex]->AutoSetTangents();
+		}
 	}
 	return true;
 }
diff --git a/Engine/Source/Runtime/Experimental/Animation/Constraints/Private/Transform/TransformableHandleUtils.cpp b/Engine/Source/Runtime/Experimental/Animation/Constraints/Private/Transform/TransformableHandleUtils.cpp
--- a/Engine/Source/Runtime/Experimental/Animation/Constraints/Private/Transform/TransformableHandleUtils.cpp
+++ b/Engine/Source/Runtime/Experimental/Animation/Constraints/Private/Transform/TransformableHandleUtils.cpp
@@ -95,20 +95,92 @@ void TickSkeletalMeshComponent(USkeletalMeshComponent* InSkeletalMeshComponent)
 	
 void MarkComponentForEvaluation(const USceneComponent* InSceneComponent, const bool bRecursive)
 {
-	if (const USkeletalMeshComponent* SkeletalMeshComponent = Cast<USkeletalMeshComponent>(InSceneComponent))
+	if (InSceneComponent)
 	{
-		UE::Anim::FAnimationEvaluationCache::Get().MarkForEvaluation(SkeletalMeshComponent);
+		if (const USkeletalMeshComponent* SkeletalMeshComponent = Cast<USkeletalMeshComponent>(InSceneComponent))
+		{
+			UE::Anim::FAnimationEvaluationCache::Get().MarkForEvaluation(SkeletalMeshComponent);
+		}
 		
 		if (bRecursive)
 		{
-			USceneComponent* Parent = SkeletalMeshComponent->GetAttachParent();
-            while (Parent)
-            {
-            	MarkComponentForEvaluation(Parent);
-            	Parent = Parent->GetAttachParent();
-            }
-        }
+			USceneComponent* Parent = InSceneComponent->GetAttachParent();
+			while (Parent)
+			{
+				MarkComponentForEvaluation(Parent);
+				Parent = Parent->GetAttachParent();
+			}
+		}
+	}
+}
+
+FTransform GetGlobalTransform(USceneComponent* InSceneComponent, const FName InSocketName)
+{
+	if (!InSceneComponent)
+	{
+		return FTransform::Identity;
+	}
+
+	constexpr bool bValidForTransforms = true;
+	const UE::Anim::FAnimationEvaluator& AnimationEvaluator = EvaluateComponent(InSceneComponent);
+	if (AnimationEvaluator.IsValid(bValidForTransforms))
+	{
+		return AnimationEvaluator.GetGlobalTransform(InSocketName);
 	}
+
+	if (!InSceneComponent->GetAttachParent())
+	{
+		return InSceneComponent->GetSocketTransform(InSocketName);
+	}
+	
+	const bool bHasSkeletalMeshInHierarchy = [](const USceneComponent* Component)
+	{
+		while (Component)
+		{
+			if (Component->IsA<USkeletalMeshComponent>())
+			{
+				return true;
+			}
+			Component = Component->GetAttachParent();
+		}
+		return false;
+	}(InSceneComponent);
+
+	if (!bHasSkeletalMeshInHierarchy)
+	{
+		return InSceneComponent->GetSocketTransform(InSocketName);
+	}
+
+	FTransform Global = InSceneComponent->GetSocketTransform(InSocketName, RTS_Component) * InSceneComponent->GetRelativeTransform();
+
+	USceneComponent* Parent = InSceneComponent->GetAttachParent();
+	FName AttachSocket = InSceneComponent->GetAttachSocketName();
+	while (Parent)
+	{
+		USkeletalMeshComponent* ParentSkeletalMeshComponent = Cast<USkeletalMeshComponent>(Parent);
+
+		bool bFromEvaluator = false;
+		if (ParentSkeletalMeshComponent && AttachSocket != NAME_None)
+		{
+			const UE::Anim::FAnimationEvaluator& ParentEvaluator = EvaluateComponent(ParentSkeletalMeshComponent);
+			if (ParentEvaluator.IsValid(bValidForTransforms))
+			{
+				Global = Global * ParentEvaluator.GetRelativeTransform(AttachSocket);
+				bFromEvaluator = true;
+			}
+		}
+
+		if (!bFromEvaluator)
+		{
+			FTransform Socket = Parent->GetSocketTransform(AttachSocket, RTS_Component) * Parent->GetRelativeTransform();
+			Global = Global * Socket;
+		}
+
+		AttachSocket = Parent->GetAttachSocketName();
+		Parent = Parent->GetAttachParent();
+	}
+
+	return Global;
 }
 
 const UE::Anim::FAnimationEvaluator& EvaluateComponent(USceneComponent* InSceneComponent)
diff --git a/Engine/Source/Runtime/Experimental/Animation/Constraints/Public/Transform/AnimationEvaluation.h b/Engine/Source/Runtime/Experimental/Animation/Constraints/Public/Transform/AnimationEvaluation.h
--- a/Engine/Source/Runtime/Experimental/Animation/Constraints/Public/Transform/AnimationEvaluation.h
+++ b/Engine/Source/Runtime/Experimental/Animation/Constraints/Public/Transform/AnimationEvaluation.h
@@ -45,12 +45,17 @@ struct FAnimationEvaluator
 	/** Clears the context and refresh the bones transforms if bRefreshBoneTransforms is true. */ 
 	void Update(const bool bRefreshBoneTransforms);
 	
-	/** Returns true if the skeletal mesh component, its skeletal mesh and the transforms are valid for use. */
-	bool IsValid() const;
+	/** Returns true if the skeletal mesh component, its skeletal mesh and the transforms are valid for use.
+	 * If bForTransform is true, the sizes of the component space transforms from the skeletal mesh component and those of the context are tested. 
+	 */
+	bool IsValid(const bool bForTransform = false) const;
 	
 	/** Returns InSocketName component space transform composed with the skeletal mesh component's global transform. */
 	CONSTRAINTS_API FTransform GetGlobalTransform(const FName InSocketName) const;
 
+	/** Returns InSocketName component space transform composed with the skeletal mesh component's relative transform. */
+	FTransform GetRelativeTransform(const FName InSocketName) const;
+
 	/** Adds a post-evaluation task to this evaluator if it has not already been added (returns true if added). */
 	bool AddPostEvaluationTask(const FAnimationEvaluationTask& InTask);
 	
@@ -69,9 +74,6 @@ private:
 
 	/** Updates the bones transforms once the skeletal mesh is finally up to date. */
 	void BoneTransformsFinalized();
-
-	/** Returns InSocketName component space transform composed with the skeletal mesh component's relative transform. */
-	FTransform GetRelativeTransform(const FName InSocketName) const;
 	
 	TWeakObjectPtr<USkeletalMeshComponent> SkeletalMeshComponent = nullptr;
 	FAnimationEvaluationContext Context;
diff --git a/Engine/Source/Runtime/Experimental/Animation/Constraints/Public/Transform/TransformableHandleUtils.h b/Engine/Source/Runtime/Experimental/Animation/Constraints/Public/Transform/TransformableHandleUtils.h
--- a/Engine/Source/Runtime/Experimental/Animation/Constraints/Public/Transform/TransformableHandleUtils.h
+++ b/Engine/Source/Runtime/Experimental/Animation/Constraints/Public/Transform/TransformableHandleUtils.h
@@ -3,14 +3,16 @@
 #pragma once
 
 #include "AnimationEvaluation.h"
+#include "TransformableHandle.h"
 #include "HAL/Platform.h"
 
 class USceneComponent;
 class USkeletalMeshComponent;
+class UTransformableHandle;
 
 namespace TransformableHandleUtils
 {
-	/** Returns true if the constraints' evaluation scheme shoud skip ticking skeletal meshes. */
+	/** Returns true if the constraints' evaluation scheme should skip ticking skeletal meshes. */
 	CONSTRAINTS_API bool SkipTicking();
 
 	/** Force ticking all the skeletal meshes related to this component. */
@@ -24,12 +26,27 @@ namespace TransformableHandleUtils
 	 * If bRecursive is true, then all skeletal mesh components up the hierarchy will be marked for evaluation.
 	 */
 	CONSTRAINTS_API void MarkComponentForEvaluation(const USceneComponent* InSceneComponent, const bool bRecursive = false);
+
+	/**
+	 * Returns the global transform of the scene component (composed with a socket component transform if InSocketName is not none).
+	 * If any of the component's parent is a skeletal mesh and has a valid evaluator, then it will be used to compose the final transform.  
+	 */
+	CONSTRAINTS_API FTransform GetGlobalTransform(USceneComponent* InSceneComponent, const FName InSocketName);
 	
 	/** Returns an updated version of InSceneComponent's animation evaluator. */
 	CONSTRAINTS_API const UE::Anim::FAnimationEvaluator& EvaluateComponent(USceneComponent* InSceneComponent);
 	
-	/** Returns an updated version of InSceneComponent's animation evaluator and adds the input post-evaluation task if not alread added. */
+	/** Returns an updated version of InSceneComponent's animation evaluator and adds the input post-evaluation task if not already added. */
 	CONSTRAINTS_API const UE::Anim::FAnimationEvaluator& EvaluateComponent(
 		USceneComponent* InSceneComponent,
 		const UE::Anim::FAnimationEvaluationTask& InTask);
+
+	/**
+	 * Returns a cast version of InHandle's target (if it's valid).  
+	 */
+	template<typename T = UObject>
+	T* GetTarget(const TObjectPtr<UTransformableHandle>& InHandle)
+	{
+		return IsValid(InHandle) && InHandle->IsValid() ? Cast<T>(InHandle->GetTarget().Get()) : nullptr;
+	}
 }
diff --git a/Engine/Source/Runtime/MovieSceneTracks/Private/Systems/MovieSceneConstraintSystem.cpp b/Engine/Source/Runtime/MovieSceneTracks/Private/Systems/MovieSceneConstraintSystem.cpp
--- a/Engine/Source/Runtime/MovieSceneTracks/Private/Systems/MovieSceneConstraintSystem.cpp
+++ b/Engine/Source/Runtime/MovieSceneTracks/Private/Systems/MovieSceneConstraintSystem.cpp
@@ -183,7 +183,8 @@ struct FEvaluateConstraintChannels
 				if (UTransformableComponentHandle* ParentHandle = Cast<UTransformableComponentHandle>(TransformConstraint->ParentTRSHandle))
 				{
 					constexpr bool bRecursive = true;
-					TransformableHandleUtils::MarkComponentForEvaluation(Cast<USkeletalMeshComponent>(ParentHandle->GetTarget()), bRecursive);
+					USceneComponent* ParentComponent = TransformableHandleUtils::GetTarget<USceneComponent>(TransformConstraint->ParentTRSHandle);
+					TransformableHandleUtils::MarkComponentForEvaluation(ParentComponent, bRecursive);
 				}
 			}
 			
