When Post integrate callback was used, aabb tree pruning was skipped and log was spammed.
If post integrate callback is used. Wait is called on the integrate tasks then we just need to launch tasks for pruning.
Otherwise tasks get linked to their corresponding batch.

diff --git a/Engine/Source/Runtime/Experimental/Chaos/Public/Chaos/TaskDispatcherEvolution.h b/Engine/Source/Runtime/Experimental/Chaos/Public/Chaos/TaskDispatcherEvolution.h
--- a/Engine/Source/Runtime/Experimental/Chaos/Public/Chaos/TaskDispatcherEvolution.h
+++ b/Engine/Source/Runtime/Experimental/Chaos/Public/Chaos/TaskDispatcherEvolution.h
@@ -192,11 +192,17 @@ public:
 
 			CurrentPendingTasks.Add(CleanUpTask);
 		}
+
 		// Prune Dynamics particles
-		if (NumBatchDyn > 0 && IntegrationPendingTasks.Num() == NumBatchDyn)
+		// If IntegrationPendingTasks.Num() == 0, it means a post integrate callback was called, then we waited for the tasks to finish and then cleared them.
+		// In this case we still do the pruning but without waiting for the integrate tasks, because it is already done.
+		if (NumBatchDyn > 0 && (IntegrationPendingTasks.Num() == NumBatchDyn || IntegrationPendingTasks.IsEmpty()))
 		{
-			check(IntegrationPendingTasks.Num() == NumBatchDyn);
+			check(IntegrationPendingTasks.Num() == NumBatchDyn || IntegrationPendingTasks.IsEmpty());
 			check(NumBatchDyn <= InternalAccelerationQueue.PendingDataArrays.Num());
+
+			const bool bLinkIntegrateTask = !IntegrationPendingTasks.IsEmpty();
+
 			NewDynPendingDataArray.SetNum(NumBatchDyn);
 			TArray<UE::Tasks::FTask> DynamicTasks;
 			for (int32 BatchIndex = 0; BatchIndex < NumBatchDyn; BatchIndex++)
@@ -214,7 +220,7 @@ public:
 							NewPendingData.Add(PendingData);
 						}
 					}
-				}, IntegrationPendingTasks[BatchIndex], LowLevelTasks::ETaskPriority::High);
+					}, bLinkIntegrateTask ? IntegrationPendingTasks[BatchIndex] : UE::Tasks::FTask{}, LowLevelTasks::ETaskPriority::High);
 				DynamicTasks.Add(PendingTask);
 			}
 			// In order to safely copy all all pruned data, we need to be sure 
@@ -232,16 +238,22 @@ public:
 			}, DynamicTasks, LowLevelTasks::ETaskPriority::High);
 			CurrentPendingTasks.Add(CopyDynamicTask);
 		}
-		else if (NumBatchDyn > 0 && IntegrationPendingTasks.Num() != NumBatchDyn)
+		else if (NumBatchDyn > 0)
 		{
 			UE_LOG(LogChaos, Warning, TEXT("No pruning happened in the Spatial Acceleration structure, a slight perfromance hit could occured. IntegrationPendingTasks.Num(): %d !=  NumBatchDyn:  %d"), IntegrationPendingTasks.Num(), NumBatchDyn);
 			CurrentPendingTasks.Append(IntegrationPendingTasks);
 		}
+
 		int32 NumBatchKin = InternalAccelerationQueue.PendingDataArrays.Num() - InternalAccelerationQueue.KinematicBatchStartIndex;
 		// Prune Kinematics particles
-		if (NumBatchKin > 0 && InternalAccelerationQueue.KinematicBatchStartIndex > 0 && KinematicPendingTasks.Num() == NumBatchKin)
+		// If KinematicPendingTasks.Num() == 0, it means a post integrate callback was called, then we waited for the tasks to finish and then cleared them.
+		// In this case we still do the pruning but without waiting for the kinematic update tasks
+		if (NumBatchKin > 0 && InternalAccelerationQueue.KinematicBatchStartIndex > 0 && (KinematicPendingTasks.Num() == NumBatchKin || KinematicPendingTasks.IsEmpty()))
 		{
-			check(KinematicPendingTasks.Num() == NumBatchKin);
+			check(KinematicPendingTasks.Num() == NumBatchKin || KinematicPendingTasks.Num() == 0);
+			
+			const bool bHasToLinkTask = !KinematicPendingTasks.IsEmpty();
+			
 			NewKinPendingDataArray.SetNum(NumBatchKin);
 			TArray<UE::Tasks::FTask> KinematicTasks;
 			for (int32 BatchIndex = 0; BatchIndex < NumBatchKin; BatchIndex++)
@@ -258,9 +270,10 @@ public:
 							NewPendingData.Add(PendingData);
 						}
 					}
-				}, KinematicPendingTasks[BatchIndex], LowLevelTasks::ETaskPriority::High);
+				}, bHasToLinkTask? KinematicPendingTasks[BatchIndex]: UE::Tasks::FTask{}, LowLevelTasks::ETaskPriority::High);
 				KinematicTasks.Add(PendingTask);
 			}
+			
 			KinematicTasks.Add(CleanUpTask);
 			UE::Tasks::FTask CopyKinematicTask = UE::Tasks::Launch(UE_SOURCE_LOCATION, [this, &NewKinPendingDataArray, NumBatchKin]()
 			{
@@ -274,8 +287,9 @@ public:
 			}, KinematicTasks, LowLevelTasks::ETaskPriority::High);
 			CurrentPendingTasks.Add(CopyKinematicTask);
 		}
-		else if (NumBatchKin > 0 && InternalAccelerationQueue.KinematicBatchStartIndex > 0 && KinematicPendingTasks.Num() != NumBatchKin)
+		else if (NumBatchKin > 0 && InternalAccelerationQueue.KinematicBatchStartIndex > 0)
 		{
+			// In this case, no post-integrate was called, but somehow NumBatchKin doesn't equal KinematicPendingTasks.Num()
 			UE_LOG(LogChaos, Warning, TEXT("No pruning happened in the Spatial Acceleration structure, a slight perfromance hit could occured. KinematicPendingTasks.Num(): %d !=  NumBatchKin:  %d"), KinematicPendingTasks.Num(), NumBatchKin);
 			CurrentPendingTasks.Append(KinematicPendingTasks);
 		}
