Fix potential use after free in on-demand content installer.

When a deduped HTTP request is comleted the content installer iterates over the linked list
of handles instead of poping from the list. Since the handles are owned by the install request
the next pointer can become invalid when any of the install request completes.

#rnx

diff --git a/Engine/Source/Runtime/Experimental/IoStore/OnDemand/Private/OnDemandContentInstaller.cpp b/Engine/Source/Runtime/Experimental/IoStore/OnDemand/Private/OnDemandContentInstaller.cpp
--- a/Engine/Source/Runtime/Experimental/IoStore/OnDemand/Private/OnDemandContentInstaller.cpp
+++ b/Engine/Source/Runtime/Experimental/IoStore/OnDemand/Private/OnDemandContentInstaller.cpp
@@ -1151,7 +1151,7 @@ void FOnDemandContentInstaller::ProcessDownloadedChunk(
 	bool bFoundRequest = false;
 	for (auto It = PendingChunkDownloads.CreateKeyIterator(ChunkHash); It; ++It)
 	{
-		const FChunkDownloadRequest& ChunkRequest = It.Value();
+		FChunkDownloadRequest& ChunkRequest = It.Value();
 		if (ChunkRequest.RequestId != RequestId)
 		{
 			continue;
@@ -1172,9 +1172,9 @@ void FOnDemandContentInstaller::ProcessDownloadedChunk(
 		}
 
 		bFoundRequest = true;
-		for (FChunkHttpRequestHandle& ChunkHttpRequestHandle : ChunkRequest.ChunkRequestHttpHandles)
+		while (FChunkHttpRequestHandle* ChunkHttpRequestHandle = ChunkRequest.ChunkRequestHttpHandles.PopHead())
 		{
-			OnProcessDownloadedChunkNotifyRequest(ChunkHttpRequestHandle, Result, bCancelled);
+			OnProcessDownloadedChunkNotifyRequest(*ChunkHttpRequestHandle, Result, bCancelled);
 		}
 
 		It.RemoveCurrent();
