Dataflow : prevent crash when undo / redo subgraphs rapidly
- make the node serialization optional since there's case where bNodeSerializable is not symetrical between save and load

diff --git a/Engine/Source/Runtime/Experimental/Dataflow/Engine/Private/Dataflow/DataflowEdNode.cpp b/Engine/Source/Runtime/Experimental/Dataflow/Engine/Private/Dataflow/DataflowEdNode.cpp
--- a/Engine/Source/Runtime/Experimental/Dataflow/Engine/Private/Dataflow/DataflowEdNode.cpp
+++ b/Engine/Source/Runtime/Experimental/Dataflow/Engine/Private/Dataflow/DataflowEdNode.cpp
@@ -2,6 +2,7 @@
 
 #include "Dataflow/DataflowEdNode.h"
 
+#include "Dataflow/DataflowArchive.h"
 #include "Dataflow/DataflowNode.h"
 #include "Dataflow/DataflowCore.h"
 #include "Dataflow/DataflowCoreNodes.h"
@@ -726,12 +727,26 @@ void UDataflowEdNode::Serialize(FArchive& Ar)
 		TSharedPtr<FDataflowNode> DataflowNode = GetDataflowNode();
 
 		// the DataflowNode may not be always valid so we need to serialize that part first so that Svaing and Loading behave exactly the same way 
-		bool bNodeSerializable = bCanSerializeNode && DataflowNode.IsValid();
-		Ar << bNodeSerializable;
-
-		if (bNodeSerializable)
+		const bool bNodeSerializable = bCanSerializeNode && DataflowNode.IsValid();
+		if (Ar.IsLoading())
+		{
+			DATAFLOW_OPTIONAL_BLOCK_READ_BEGIN(bNodeSerializable)
+			{
+				DataflowNode->SerializeInternal(Ar);
+			}
+			DATAFLOW_OPTIONAL_BLOCK_READ_ELSE()
+			DATAFLOW_OPTIONAL_BLOCK_READ_END()
+		}
+		else if (Ar.IsSaving())
 		{
-			DataflowNode->SerializeInternal(Ar);
+			DATAFLOW_OPTIONAL_BLOCK_WRITE_BEGIN()
+			{
+				if (bNodeSerializable)
+				{
+					DataflowNode->SerializeInternal(Ar);
+				}
+			}
+			DATAFLOW_OPTIONAL_BLOCK_WRITE_END()
 		}
 	}
 	// some double level template type have unwanted spaces 
