Dataflow : attempt to make Dataflow instance more robust when receiving a object property changed notification

#Synced-CL 49114577

diff --git a/Engine/Source/Runtime/Experimental/Dataflow/Engine/Private/Dataflow/DataflowInstance.cpp b/Engine/Source/Runtime/Experimental/Dataflow/Engine/Private/Dataflow/DataflowInstance.cpp
--- a/Engine/Source/Runtime/Experimental/Dataflow/Engine/Private/Dataflow/DataflowInstance.cpp
+++ b/Engine/Source/Runtime/Experimental/Dataflow/Engine/Private/Dataflow/DataflowInstance.cpp
@@ -508,7 +508,7 @@ FDataflowInstance::~FDataflowInstance()
 #if WITH_EDITOR
 void FDataflowInstance::OnOwnerPostEditChangeProperty(UObject* InObject, struct FPropertyChangedEvent& InPropertyChangedEvent)
 {
-	if (Owner != nullptr && Owner == InObject)
+	if (IsValid(Owner) && IsValid(InObject) && Owner == InObject)
 	{
 		VariableOverrides.OnOwnerPostEditChangeProperty(InPropertyChangedEvent);
 
