Fix to avoid large solver iteration count for geometry collection
- This is a 5.7 only fix that follows the hotfix rules

diff --git a/Engine/Source/Runtime/Experimental/GeometryCollectionEngine/Private/GeometryCollection/GeometryCollectionComponent.cpp b/Engine/Source/Runtime/Experimental/GeometryCollectionEngine/Private/GeometryCollection/GeometryCollectionComponent.cpp
--- a/Engine/Source/Runtime/Experimental/GeometryCollectionEngine/Private/GeometryCollection/GeometryCollectionComponent.cpp
+++ b/Engine/Source/Runtime/Experimental/GeometryCollectionEngine/Private/GeometryCollection/GeometryCollectionComponent.cpp
@@ -3936,9 +3936,12 @@ PRAGMA_ENABLE_DEPRECATION_WARNINGS
 		SimulationParameters.UseCCD = BodyInstance.bUseCCD;
 		SimulationParameters.UseMACD = BodyInstance.GetUseMACD();
 		SimulationParameters.AllowPartialIslandSleep = BodyInstance.GetAllowPartialIslandSleep();
-		SimulationParameters.PositionSolverIterations = BodyInstance.GetPositionSolverIterationCount();
-		SimulationParameters.VelocitySolverIterations = BodyInstance.GetVelocitySolverIterationCount();
-		SimulationParameters.ProjectionSolverIterations = BodyInstance.GetProjectionSolverIterationCount();
+		const int32 PositionSolverIterationCount = BodyInstance.GetPositionSolverIterationCount();
+		SimulationParameters.PositionSolverIterations = (PositionSolverIterationCount >= 0) ? FMath::Min(PositionSolverIterationCount, 255) : 8;
+		const int32 VelocitySolverIterationCount = BodyInstance.GetVelocitySolverIterationCount();
+		SimulationParameters.VelocitySolverIterations = (VelocitySolverIterationCount >= 0) ? FMath::Min(VelocitySolverIterationCount, 255) : 1;
+		const int32 ProjectionSolverIterationCount = BodyInstance.GetProjectionSolverIterationCount();
+		SimulationParameters.ProjectionSolverIterations = (ProjectionSolverIterationCount >= 0) ? FMath::Min(ProjectionSolverIterationCount, 255) : 1;
 		SimulationParameters.LinearDamping = BodyInstance.LinearDamping;
 		SimulationParameters.AngularDamping = BodyInstance.AngularDamping;
 		SimulationParameters.InitialOverlapDepenetrationVelocity = BodyInstance.GetMaxDepenetrationVelocity();
