Chaos : Fix traces not working properly on partialy destroyed geometry collection when removal feature is enabled

diff --git a/Engine/Source/Runtime/Experimental/Chaos/Private/PBDRigidsSolver.cpp b/Engine/Source/Runtime/Experimental/Chaos/Private/PBDRigidsSolver.cpp
--- a/Engine/Source/Runtime/Experimental/Chaos/Private/PBDRigidsSolver.cpp
+++ b/Engine/Source/Runtime/Experimental/Chaos/Private/PBDRigidsSolver.cpp
@@ -422,7 +422,7 @@ namespace Chaos
 		bool bRemoveParticleFromMovingKinematicsOnDisable = true;
 		FAutoConsoleVariableRef CVarChaosRemoveParticleFromMovingKinematicsOnDisable(TEXT("p.Chaos.RemoveParticleFromMovingKinematicsOnDisable"), bRemoveParticleFromMovingKinematicsOnDisable, TEXT(""));
 
-		bool bSingleThreadPushData = false;
+		bool bSingleThreadPushData = true;
 		FAutoConsoleVariableRef CVarChaosSingleThreadPushData(TEXT("p.Chaos.SingleThreadPushData"), bSingleThreadPushData, TEXT("Run push Data in Single Thread."));
 	}
 }
diff --git a/Engine/Source/Runtime/Experimental/Chaos/Private/PhysicsProxy/GeometryCollectionPhysicsProxy.cpp b/Engine/Source/Runtime/Experimental/Chaos/Private/PhysicsProxy/GeometryCollectionPhysicsProxy.cpp
--- a/Engine/Source/Runtime/Experimental/Chaos/Private/PhysicsProxy/GeometryCollectionPhysicsProxy.cpp
+++ b/Engine/Source/Runtime/Experimental/Chaos/Private/PhysicsProxy/GeometryCollectionPhysicsProxy.cpp
@@ -3856,11 +3856,21 @@ void FGeometryCollectionPhysicsProxy::BufferPhysicsResults_External(Chaos::FDirt
 	PrepareBufferData(BufferData, PhysicsThreadCollection);
 }
 
+namespace Chaos::CVars
+{
+	extern bool bSingleThreadPushData;
+}
+
 static void UpdateParticleHandleTransform(Chaos::FPBDRigidsSolver& CurrentSolver, Chaos::FPBDRigidClusteredParticleHandle& Handle, const Chaos::FRigidTransform3& NewTransform)
 {
 	Handle.SetX(NewTransform.GetTranslation());
 	Handle.SetR(NewTransform.GetRotation());
 	Handle.UpdateWorldSpaceState(NewTransform, Chaos::FVec3{0});
+
+	if (Chaos::CVars::bSingleThreadPushData)
+	{
+		CurrentSolver.GetEvolution()->DirtyParticle(Handle);
+	}
 }
 
 static void UpdateParticleHandleTransformIfNeeded(Chaos::FPBDRigidsSolver& CurrentSolver, Chaos::FPBDRigidClusteredParticleHandle& Handle, const Chaos::FRigidTransform3& NewTransform)
