IAD Cache
- Remove LastAccessTime from pending chunks.  As implemented, this isn't really an improvement over just using current time at flush.

diff --git a/Engine/Source/Runtime/Experimental/IoStore/OnDemand/Private/OnDemandInstallCache.cpp b/Engine/Source/Runtime/Experimental/IoStore/OnDemand/Private/OnDemandInstallCache.cpp
--- a/Engine/Source/Runtime/Experimental/IoStore/OnDemand/Private/OnDemandInstallCache.cpp
+++ b/Engine/Source/Runtime/Experimental/IoStore/OnDemand/Private/OnDemandInstallCache.cpp
@@ -1660,7 +1660,6 @@ class FOnDemandInstallCache final
 			TotalSize += Chunk.GetSize();
 			ChunkHashes.Add(ChunkHash);
 			Chunks.Add(MoveTemp(Chunk));
-			LastAccessTime.store(FDateTime::UtcNow().GetTicks(), std::memory_order_relaxed);
 		}
 
 		void Reset()
@@ -1668,7 +1667,6 @@ class FOnDemandInstallCache final
 			Chunks.Reset();
 			ChunkHashes.Reset();
 			TotalSize = 0;
-			LastAccessTime.store(0, std::memory_order_relaxed);
 		}
 
 		const FIoBuffer* FindChunk(const FIoHash& ChunkHash) const
@@ -1685,7 +1683,6 @@ class FOnDemandInstallCache final
 		TArray<FIoBuffer>		Chunks;
 		TArray<FIoHash>			ChunkHashes;
 		uint64					TotalSize = 0;
-		std::atomic<uint64>		LastAccessTime{ 0 };
 		mutable FSharedMutex	SharedMutex;
 	};
 
@@ -1921,7 +1918,7 @@ void FOnDemandInstallCache::Initialize(FSharedBackendContextRef Context)
 
 void FOnDemandInstallCache::Shutdown()
 {
-	if (FResult Result = FlushPendingChunksImpl(PendingChunks, PendingChunks.LastAccessTime.load(std::memory_order_relaxed)); Result.HasError())
+	if (FResult Result = FlushPendingChunksImpl(PendingChunks); Result.HasError())
 	{
 		UE_LOG(LogIoStoreOnDemand, Error, TEXT("Failed to flush pending chunks on shutdown, reason '%s'"), *LexToString(Result.GetError()));
 	}
@@ -2102,8 +2099,6 @@ bool FOnDemandInstallCache::Resolve(FIoRequestImpl* Request)
 			check(BytesToRead.GetSize() == ChunkRequest.EncodedChunk.GetSize());
 			FMemory::Memcpy(ChunkRequest.EncodedChunk.GetData(), BytesToRead.GetData(), ChunkRequest.EncodedChunk.GetSize());
 
-			PendingChunks.LastAccessTime.store(FDateTime::UtcNow().GetTicks(), std::memory_order_relaxed);
-
 			SharedLock.Unlock();
 
 			UE::Tasks::Launch(UE_SOURCE_LOCATION, [this, Request]
@@ -2309,8 +2304,7 @@ FResult FOnDemandInstallCache::PutChunk(FIoBuffer&& Chunk, const FIoHash& ChunkH
 {
 	if (PendingChunks.TotalSize > FPendingChunks::MaxPendingBytes)
 	{
-		if (FResult Result = FlushPendingChunks(PendingChunks, PendingChunks.LastAccessTime.load(std::memory_order_relaxed));
-			Result.HasError())
+		if (FResult Result = FlushPendingChunks(PendingChunks); Result.HasError())
 		{
 			return Result;
 		}
@@ -3384,8 +3378,7 @@ FResult FOnDemandInstallCache::Defrag(
 
 FResult FOnDemandInstallCache::Flush()
 {
-	if (FResult Result = FlushPendingChunks(PendingChunks, PendingChunks.LastAccessTime.load(std::memory_order_relaxed));
-		Result.HasError())
+	if (FResult Result = FlushPendingChunks(PendingChunks); Result.HasError())
 	{
 		return Result;
 	}
