Fix FOnDemandInstallCache::FlushPendingChunksImpl moving chunks while flushing which causes a crash if we get a concurrent read from the pending chunk buffer.

diff --git a/Engine/Source/Runtime/Experimental/IoStore/OnDemand/Private/OnDemandInstallCache.cpp b/Engine/Source/Runtime/Experimental/IoStore/OnDemand/Private/OnDemandInstallCache.cpp
--- a/Engine/Source/Runtime/Experimental/IoStore/OnDemand/Private/OnDemandInstallCache.cpp
+++ b/Engine/Source/Runtime/Experimental/IoStore/OnDemand/Private/OnDemandInstallCache.cpp
@@ -1739,7 +1739,7 @@ private:
 	bool										Resolve(FIoRequestImpl* Request);
 	void										CompleteRequest(FIoRequestImpl* Request, EIoErrorCode Status);
 	FResult										FlushPendingChunks(FPendingChunks& Chunks, int64 UtcAccessTicks = 0);
-	FResult										FlushPendingChunksImpl(FPendingChunks& Chunks, int64 UtcAccessTicks = 0);
+	FResult										FlushPendingChunksImpl(const FPendingChunks& Chunks, int64 UtcAccessTicks = 0);
 	FString										GetJournalFilename() const { return CacheDirectory / TEXT("cas.jrn"); }
 	FString										GetSnapshotFilename() const { return CacheDirectory / TEXT("cas.snp"); }
 	void										CaptureCacheState(UE::UnifiedError::IoStoreOnDemand::FInstallCacheErrorContext& OutCtx);
@@ -3504,6 +3504,12 @@ FResult FOnDemandInstallCache::FlushPendingChunks(FPendingChunks& Chunks, int64
 #if UE_ONDEMANDINSTALLCACHE_EXCLUSIVE_WRITE
 	UE::Tasks::TTask<FResult> Task = ExclusivePipe.Launch(UE_SOURCE_LOCATION, [this, &Chunks, UtcAccessTicks]
 	{
+		ON_SCOPE_EXIT 
+		{ 
+			TUniqueLock Lock(Chunks.SharedMutex);
+			Chunks.Reset();
+		};
+
 		return FlushPendingChunksImpl(Chunks, UtcAccessTicks);
 	}, UE::Tasks::ETaskPriority::BackgroundHigh);
 
@@ -3512,18 +3518,18 @@ FResult FOnDemandInstallCache::FlushPendingChunks(FPendingChunks& Chunks, int64
 	return Task.GetResult();
 
 #else
+	ON_SCOPE_EXIT
+	{
+		TUniqueLock Lock(Chunks.SharedMutex);
+		Chunks.Reset();
+	};
+
 	return FlushPendingChunksImpl(Chunks, UtcAccessTicks);
 #endif // UE_ONDEMANDINSTALLCACHE_EXCLUSIVE_WRITE
 }
 
-FResult FOnDemandInstallCache::FlushPendingChunksImpl(FPendingChunks& Chunks, int64 UtcAccessTicks)
+FResult FOnDemandInstallCache::FlushPendingChunksImpl(const FPendingChunks& Chunks, int64 UtcAccessTicks)
 {
-	ON_SCOPE_EXIT 
-	{ 
-		TUniqueLock Lock(Chunks.SharedMutex);
-		Chunks.Reset();
-	};
-
 	FLargeMemoryWriter Ar(FMath::Min(Chunks.TotalSize, static_cast<uint64>(Cas.GetMaxBlockSize())));
 
 	// This should be the only thread writing to CurrentBlock
@@ -3562,9 +3568,9 @@ FResult FOnDemandInstallCache::FlushPendingChunksImpl(FPendingChunks& Chunks, in
 		}
 
 		const int64 CasBlockOffset = CasFileHandle->Tell();
-
-		TArray<FIoHash>		ChunkHashes;
+		
 		TArray<int64>		Offsets;
+		const FIoHash*		ChunkHashBegin = &Chunks.ChunkHashes[ChunkIdx];
 
 		while (ChunkIdx < Chunks.Chunks.Num())
 		{
@@ -3573,14 +3579,15 @@ FResult FOnDemandInstallCache::FlushPendingChunksImpl(FPendingChunks& Chunks, in
 				break;
 			}
 
-			FIoBuffer Chunk = MoveTemp(Chunks.Chunks[ChunkIdx]);
-			ChunkHashes.Add(MoveTemp(Chunks.ChunkHashes[ChunkIdx]));
+			const FIoBuffer& Chunk = Chunks.Chunks[ChunkIdx];
 			Offsets.Add(CasBlockOffset + Ar.Tell());
-			Ar.Serialize(Chunk.GetData(), Chunk.GetSize());
+			Ar.Serialize(const_cast<uint8*>(Chunk.GetData()), Chunk.GetSize());
 
 			++ChunkIdx;
 		}
 
+		TConstArrayView<FIoHash> ChunkHashes = MakeArrayView(ChunkHashBegin, Offsets.Num());
+
 		if (Ar.Tell() > 0)
 		{
 			UE_LOG(LogIoStoreOnDemand, Log, TEXT("Writing %.2lf MiB to CAS block %u"),
