[DSTM] Making the pull data queue transactionally safe

The default spsc queue implementation use atomics and therefore is not transactionally safe, causing an assert in transactions. This change changes the implementation to the transactionally safe version if remote object handle is defined to 1, which uses a mutex under the hood (performance in this case will be impacted though)

diff --git a/Engine/Source/Runtime/Experimental/Chaos/Public/Chaos/ChaosMarshallingManager.h b/Engine/Source/Runtime/Experimental/Chaos/Public/Chaos/ChaosMarshallingManager.h
--- a/Engine/Source/Runtime/Experimental/Chaos/Public/Chaos/ChaosMarshallingManager.h
+++ b/Engine/Source/Runtime/Experimental/Chaos/Public/Chaos/ChaosMarshallingManager.h
@@ -9,6 +9,7 @@
 #include "PhysicsProxy/SingleParticlePhysicsProxyFwd.h"
 #include "Chaos/SimCallbackObject.h"
 #include "Chaos/AsyncInitBodyHelper.h"
+#include "Framework/Threading.h"
 
 class FGeometryCollectionResults;
 
@@ -505,8 +506,8 @@ private:
 
 	//pull
 	FPullPhysicsData* CurPullData;	//the current pull data sim is writing to
-	TSpscQueue<FPullPhysicsData*> PullDataQueue;	//the results the simulation has written to. Consumed by external thread
-	TSpscQueue<FPullPhysicsData*> PullDataPool;	//the pull data pool to avoid reallocs. Pushed by external thread, popped by internal
+	TPhysicsSpscQueue<FPullPhysicsData*> PullDataQueue;	//the results the simulation has written to. Consumed by external thread
+	TPhysicsSpscQueue<FPullPhysicsData*> PullDataPool;	//the pull data pool to avoid reallocs. Pushed by external thread, popped by internal
 	TArray<TUniquePtr<FPullPhysicsData>> BackingPullBuffer;		//all pull data is cleaned up by this
 
 	std::atomic<int32> Delay;
diff --git a/Engine/Source/Runtime/Experimental/Chaos/Public/Framework/Threading.h b/Engine/Source/Runtime/Experimental/Chaos/Public/Framework/Threading.h
--- a/Engine/Source/Runtime/Experimental/Chaos/Public/Framework/Threading.h
+++ b/Engine/Source/Runtime/Experimental/Chaos/Public/Framework/Threading.h
@@ -15,6 +15,8 @@
 #include "AutoRTFM.h"
 #include "Templates/SharedPointer.h"
 #include "ChaosInsights/ChaosInsightsMacros.h"
+#include "Containers/SpscQueue.h"
+#include "Containers/TransactionallySafeSpscQueue.h"
 
 #ifndef PHYSICS_THREAD_CONTEXT
 	#if (!UE_BUILD_SHIPPING && !UE_BUILD_TEST)
@@ -1204,6 +1206,14 @@ FORCEINLINE void EnsureIsInGameThreadContext()
 	using FPhysSceneLock = FPhysSceneLockNonTransactional;
 #endif
 
+#if UE_WITH_REMOTE_OBJECT_HANDLE
+template<typename T>
+using TPhysicsSpscQueue = TTransactionallySafeSpscQueue<T>;
+#else
+template<typename T>
+using TPhysicsSpscQueue = TSpscQueue<T>;
+#endif
+
 	// Stable types to use in calling code configured by the compiler switches above
 	using FPhysicsSceneGuardScopedWrite = TPhysicsSceneGuardScopedWrite<FPhysSceneLock>;
 	using FPhysicsSceneGuardScopedRead = TPhysicsSceneGuardScopedRead<FPhysSceneLock>;
