Fix race condition detected by Tsan
This issue has been hard to reproduce after the first fix CL 45316272.
Tsan doesn't seem to detect it every time.
The following change ensure integration is finished when async acceleration is flushed. That is often the case because integration task are created before the kinematic task which have already a dependency with the async flush.

diff --git a/Engine/Source/Runtime/Experimental/Chaos/Public/Chaos/TaskDispatcherEvolution.h b/Engine/Source/Runtime/Experimental/Chaos/Public/Chaos/TaskDispatcherEvolution.h
--- a/Engine/Source/Runtime/Experimental/Chaos/Public/Chaos/TaskDispatcherEvolution.h
+++ b/Engine/Source/Runtime/Experimental/Chaos/Public/Chaos/TaskDispatcherEvolution.h
@@ -289,6 +289,7 @@ public:
 	{
 		TArray<UE::Tasks::FTask> PendingTasks{AsyncQueueDynTask, AsyncQueueKinTask};
 		PendingTasks.Append(KinematicPendingTasks);
+		PendingTasks.Append(IntegrationPendingTasks);
 		FlushAsyncQueueTask = UE::Tasks::Launch(UE_SOURCE_LOCATION, [FlushAccelerationQueueLambda]()
 		{
 			QUICK_SCOPE_CYCLE_COUNTER(STAT_FlushAsyncAccelerationQueueTask);
