Dataflow fix global in header

diff --git a/Engine/Source/Runtime/Experimental/Dataflow/Engine/Private/Dataflow/DataflowObject.cpp b/Engine/Source/Runtime/Experimental/Dataflow/Engine/Private/Dataflow/DataflowObject.cpp
--- a/Engine/Source/Runtime/Experimental/Dataflow/Engine/Private/Dataflow/DataflowObject.cpp
+++ b/Engine/Source/Runtime/Experimental/Dataflow/Engine/Private/Dataflow/DataflowObject.cpp
@@ -128,6 +128,8 @@ UE::Dataflow::FGraph* FDataflowAssetEdit::GetGraph()
 	return nullptr;
 }
 
+const FString UDataflow::SimulationTag = TEXT("DataflowSimulationTag");
+
 UDataflow::UDataflow(const FObjectInitializer& ObjectInitializer)
 	: UEdGraph(ObjectInitializer)
 	, Dataflow(new UE::Dataflow::FGraph())
diff --git a/Engine/Source/Runtime/Experimental/Dataflow/Engine/Public/Dataflow/DataflowObject.h b/Engine/Source/Runtime/Experimental/Dataflow/Engine/Public/Dataflow/DataflowObject.h
--- a/Engine/Source/Runtime/Experimental/Dataflow/Engine/Public/Dataflow/DataflowObject.h
+++ b/Engine/Source/Runtime/Experimental/Dataflow/Engine/Public/Dataflow/DataflowObject.h
@@ -124,7 +124,7 @@ public:
 	virtual bool IsEditorOnly() const { return true; }
 
 	/** Simulation tag to use in the node registry */
-	static const inline FString SimulationTag = TEXT("DataflowSimulationTag");
+	static DATAFLOWENGINE_API const FString SimulationTag;
 
 	//~ Begin UObject interface
 	virtual void BeginDestroy() override;
