[Chaos CCD]

Get ProjectCCD to account for target separation.

diff --git a/Engine/Source/Runtime/Experimental/Chaos/Private/Chaos/CCDUtilities.cpp b/Engine/Source/Runtime/Experimental/Chaos/Private/Chaos/CCDUtilities.cpp
--- a/Engine/Source/Runtime/Experimental/Chaos/Private/Chaos/CCDUtilities.cpp
+++ b/Engine/Source/Runtime/Experimental/Chaos/Private/Chaos/CCDUtilities.cpp
@@ -1046,7 +1046,14 @@ namespace Chaos
 				const FReal Phi = Constraint->GetPhi();
 				const FReal PhiTolerance = PhiToleranceScale * Constraint->GetCCDTargetPenetration();
 
-				if (Phi < -PhiTolerance)
+				FReal TargetSeparation = 0;
+				if (const FManifoldPoint* ManifoldPoint = Constraint->GetClosestManifoldPoint())
+				{
+					TargetSeparation = ManifoldPoint->TargetPhi;
+				}
+				const FReal MinPhi = TargetSeparation - PhiTolerance;
+
+				if (Phi < MinPhi)
 				{
 					bSolved = false;
 
