Chaos Deformable Solver Integration: refactor constraints using SolverParticlesOrRange

diff --git a/Engine/Source/Runtime/Experimental/Chaos/Public/Chaos/BlendedXPBDCorotatedConstraints.h b/Engine/Source/Runtime/Experimental/Chaos/Public/Chaos/BlendedXPBDCorotatedConstraints.h
--- a/Engine/Source/Runtime/Experimental/Chaos/Public/Chaos/BlendedXPBDCorotatedConstraints.h
+++ b/Engine/Source/Runtime/Experimental/Chaos/Public/Chaos/BlendedXPBDCorotatedConstraints.h
@@ -74,8 +74,8 @@ namespace Chaos::Softs
 
 		virtual ~FBlendedXPBDCorotatedConstraints() {}
 
-
-		virtual void ApplyInSerial(ParticleType& Particles, const T Dt, const int32 ElementIndex) const override
+		template<typename SolverParticlesOrRange>
+		void ApplyInSerial(SolverParticlesOrRange& Particles, const T Dt, const int32 ElementIndex) const
 		{
 			TVec4<TVector<T, 3>> PolarDelta = GetPolarDelta(Particles, Dt, ElementIndex);
 
@@ -100,8 +100,8 @@ namespace Chaos::Softs
 		}
 
 	protected:
-
-		virtual TVec4<TVector<T, 3>> GetDeterminantDelta(const ParticleType& Particles, const T Dt, const int32 ElementIndex, const T Tol = (T)1e-3) const override
+		template<typename SolverParticlesOrRange>
+		TVec4<TVector<T, 3>> GetDeterminantDelta(const SolverParticlesOrRange& Particles, const T Dt, const int32 ElementIndex, const T Tol = (T)1e-3) const
 		{
 			TRACE_CPUPROFILER_EVENT_SCOPE(STAT_ChaosBlendedXPBDCorotatedApplyDet);
 
@@ -162,8 +162,8 @@ namespace Chaos::Softs
 		}
 
 
-
-		virtual TVec4<TVector<T, 3>> GetPolarDelta(const ParticleType& Particles, const T Dt, const int32 ElementIndex, const T Tol = (T)1e-3) const override
+		template<typename SolverParticlesOrRange>
+		TVec4<TVector<T, 3>> GetPolarDelta(const SolverParticlesOrRange& Particles, const T Dt, const int32 ElementIndex, const T Tol = (T)1e-3) const
 		{
 			TRACE_CPUPROFILER_EVENT_SCOPE(STAT_ChaosBlendedXPBDCorotatedApplyPolar);
 			SCOPE_CYCLE_COUNTER(STAT_ChaosXPBDCorotatedPolar);
diff --git a/Engine/Source/Runtime/Experimental/Chaos/Public/Chaos/XPBDCorotatedConstraints.h b/Engine/Source/Runtime/Experimental/Chaos/Public/Chaos/XPBDCorotatedConstraints.h
--- a/Engine/Source/Runtime/Experimental/Chaos/Public/Chaos/XPBDCorotatedConstraints.h
+++ b/Engine/Source/Runtime/Experimental/Chaos/Public/Chaos/XPBDCorotatedConstraints.h
@@ -215,8 +215,8 @@ namespace Chaos::Softs
 			return Result;
 		}
 
-
-		PMatrix<T, 3, 3> Ds(const int e, const ParticleType& InParticles) const {
+		template<typename SolverParticlesOrRange>
+		PMatrix<T, 3, 3> Ds(const int e, const SolverParticlesOrRange& InParticles) const {
 			PMatrix<T, 3, 3> Result((T)0.);
 			for (int i = 0; i < 3; i++) {
 				for (int c = 0; c < 3; c++) {
@@ -226,8 +226,8 @@ namespace Chaos::Softs
 			return Result;
 		}
 
-
-		PMatrix<T, 3, 3> F(const int e, const ParticleType& InParticles) const {
+		template<typename SolverParticlesOrRange>
+		PMatrix<T, 3, 3> F(const int e, const SolverParticlesOrRange& InParticles) const {
 			return ElementDmInv(e) * Ds(e, InParticles);
 		}
 
@@ -256,7 +256,8 @@ namespace Chaos::Softs
 			for (T& Lambdas : LambdaArray) { Lambdas = (T)0.; }
 		}
 
-		virtual void ApplyInSerial(ParticleType& Particles, const T Dt, const int32 ElementIndex) const
+		template<typename SolverParticlesOrRange>
+		void ApplyInSerial(SolverParticlesOrRange& Particles, const T Dt, const int32 ElementIndex) const
 		{
 			TRACE_CPUPROFILER_EVENT_SCOPE(STAT_ChaosXPBDCorotatedApplySingle);
 
@@ -276,7 +277,8 @@ namespace Chaos::Softs
 
 		}
 
-		void ApplyInSerial(ParticleType& Particles, const T Dt) const
+		template<typename SolverParticlesOrRange>
+		void ApplyInSerial(SolverParticlesOrRange& Particles, const T Dt) const
 		{
 			SCOPE_CYCLE_COUNTER(STAT_ChaosXPBDCorotated);
 			TRACE_CPUPROFILER_EVENT_SCOPE(STAT_ChaosXPBDCorotatedApplySerial);
@@ -292,10 +294,10 @@ namespace Chaos::Softs
 				}
 			}
 
-
 		}
 
-		void ApplyInParallel(ParticleType& Particles, const T Dt) const
+		template<typename SolverParticlesOrRange>
+		void ApplyInParallel(SolverParticlesOrRange& Particles, const T Dt) const
 		{	
 			//code for error metric:
 			if (bRecordMetric)
@@ -490,8 +492,8 @@ namespace Chaos::Softs
 			}
 		}
 
-
-		virtual TVec4<TVector<T, 3>> GetDeterminantDelta(const ParticleType& Particles, const T Dt, const int32 ElementIndex, const T Tol = (T)1e-3) const
+		template<typename SolverParticlesOrRange>
+		TVec4<TVector<T, 3>> GetDeterminantDelta(const SolverParticlesOrRange& Particles, const T Dt, const int32 ElementIndex, const T Tol = (T)1e-3) const
 		{
 			TRACE_CPUPROFILER_EVENT_SCOPE(STAT_ChaosXPBDCorotatedApplyDet);
 			//SCOPE_CYCLE_COUNTER(STAT_ChaosXPBDCorotatedDet);
@@ -554,8 +556,8 @@ namespace Chaos::Softs
 		}
 
 
-
-		virtual TVec4<TVector<T, 3>> GetPolarDelta(const ParticleType& Particles, const T Dt, const int32 ElementIndex, const T Tol = (T)1e-3) const
+		template<typename SolverParticlesOrRange>
+		TVec4<TVector<T, 3>> GetPolarDelta(const SolverParticlesOrRange& Particles, const T Dt, const int32 ElementIndex, const T Tol = (T)1e-3) const
 		{	
 			TRACE_CPUPROFILER_EVENT_SCOPE(STAT_ChaosXPBDCorotatedApplyPolar);
 			SCOPE_CYCLE_COUNTER(STAT_ChaosXPBDCorotatedPolar);
diff --git a/Engine/Source/Runtime/Experimental/Chaos/Public/Chaos/XPBDWeakConstraints.h b/Engine/Source/Runtime/Experimental/Chaos/Public/Chaos/XPBDWeakConstraints.h
--- a/Engine/Source/Runtime/Experimental/Chaos/Public/Chaos/XPBDWeakConstraints.h
+++ b/Engine/Source/Runtime/Experimental/Chaos/Public/Chaos/XPBDWeakConstraints.h
@@ -78,7 +78,8 @@ namespace Chaos::Softs
 
 		virtual ~FXPBDWeakConstraints() {}
 
-		void ApplyInParallel(ParticleType& Particles, const T Dt) const
+		template<typename SolverParticlesOrRange>
+		void ApplyInParallel(SolverParticlesOrRange& Particles, const T Dt) const
 		{
 			TRACE_CPUPROFILER_EVENT_SCOPE(STAT_ChaosXPBDWeakConstraintApply);
 
@@ -115,7 +116,8 @@ namespace Chaos::Softs
 			return Indices;
 		}
 
-		void Init(const FSolverParticles& InParticles, const T Dt) const
+		template<typename SolverParticlesOrRange>
+		void Init(const SolverParticlesOrRange& InParticles, const T Dt) const
 		{
 			for (T& Lambdas : LambdaArray) { Lambdas = (T)0.; }
 			if (DebugDrawParams.bVisualizeBindings)
@@ -137,7 +139,8 @@ namespace Chaos::Softs
 			Chaos::ComputeWeakConstraintsColoring(Indices, SecondIndices, Particles, ConstraintsPerColor);
 		}
 
-		void ApplySingleConstraintWithoutSelfTarget(ParticleType& Particles, const T Dt, const int32 ConstraintIndex, const T Tol = UE_KINDA_SMALL_NUMBER) const
+		template<typename SolverParticlesOrRange>
+		void ApplySingleConstraintWithoutSelfTarget(SolverParticlesOrRange& Particles, const T Dt, const int32 ConstraintIndex, const T Tol = UE_KINDA_SMALL_NUMBER) const
 		{
 			if (Stiffness[ConstraintIndex] < Tol)
 			{
@@ -179,7 +182,8 @@ namespace Chaos::Softs
 			}
 		}
 
-		void ApplySingleConstraintWithSelfTarget(ParticleType& Particles, const T Dt, const int32 ConstraintIndex, const T Tol = UE_KINDA_SMALL_NUMBER) const
+		template<typename SolverParticlesOrRange>
+		void ApplySingleConstraintWithSelfTarget(SolverParticlesOrRange& Particles, const T Dt, const int32 ConstraintIndex, const T Tol = UE_KINDA_SMALL_NUMBER) const
 		{
 			if (Stiffness[ConstraintIndex] < Tol)
 			{
@@ -235,7 +239,8 @@ namespace Chaos::Softs
 			}
 		}
 
-		void VisualizeAllBindings(const FSolverParticles& InParticles, const T Dt) const
+		template<typename SolverParticlesOrRange>
+		void VisualizeAllBindings(const SolverParticlesOrRange& InParticles, const T Dt) const
 		{
 #if WITH_EDITOR
 			auto DoubleVert = [](Chaos::TVec3<T> V) { return FVector3d(V.X, V.Y, V.Z); };
