Chaos Deformable Solver Integration:
Added API macros

diff --git a/Engine/Plugins/ChaosCloth/Source/ChaosCloth/Public/ChaosCloth/ChaosClothComplexColliders.h b/Engine/Plugins/ChaosCloth/Source/ChaosCloth/Public/ChaosCloth/ChaosClothComplexColliders.h
--- a/Engine/Plugins/ChaosCloth/Source/ChaosCloth/Public/ChaosCloth/ChaosClothComplexColliders.h
+++ b/Engine/Plugins/ChaosCloth/Source/ChaosCloth/Public/ChaosCloth/ChaosClothComplexColliders.h
@@ -51,7 +51,7 @@ namespace Chaos
 		// ---- End of Collider interface ----
 
 		// -- Constraints interface ----
-		void ExtractComplexColliderBoneData(TMap<Softs::FParticleRangeIndex, Softs::FPBDComplexColliderBoneData>& Data) const;
+		CHAOSCLOTH_API void ExtractComplexColliderBoneData(TMap<Softs::FParticleRangeIndex, Softs::FPBDComplexColliderBoneData>& Data) const;
 		int32 GetNumSkinnedTriangleMeshes() const
 		{
 			return SkinnedTriangleMeshes.Num();
diff --git a/Engine/Source/Runtime/Experimental/Chaos/Public/Chaos/UniformGrid.h b/Engine/Source/Runtime/Experimental/Chaos/Public/Chaos/UniformGrid.h
--- a/Engine/Source/Runtime/Experimental/Chaos/Public/Chaos/UniformGrid.h
+++ b/Engine/Source/Runtime/Experimental/Chaos/Public/Chaos/UniformGrid.h
@@ -333,7 +333,7 @@ public:
 #endif
 	~TMPMGrid() {}
 
-	void BaseNodeIndex(const TVector<T, 3>& X, TVector<int32, 3>& Index, TVector<T, 3>& weights) const;
+	CHAOS_API void BaseNodeIndex(const TVector<T, 3>& X, TVector<int32, 3>& Index, TVector<T, 3>& weights) const;
 	
 	inline T Nijk(T w, int32 ii) const {
 		if (interp == linear)
@@ -395,15 +395,15 @@ public:
 		return Result;
 	}
 
-	TVector<T, 3> Node(int32 FlatIndexIn) const;
+	CHAOS_API TVector<T, 3> Node(int32 FlatIndexIn) const;
 
-	int32 FlatIndex(const TVector<int32, 3>& Index) const;
+	CHAOS_API int32 FlatIndex(const TVector<int32, 3>& Index) const;
 	TVector<int32, 3> Lin2MultiIndex(const int32 IndexIn) const;
 	const TVector<int32, 3> GetCells() const { return MCells; }
 	const TVector<T, 3> GetMinCorner() const { return MMinCorner; } 
 	const TVector<T, 3>& GetDx() const { return MDx; }
 	void SetDx(const TVector<T, 3> DxIn) { MDx = DxIn; }
-	void UpdateGridFromPositions(const Chaos::TDynamicParticles<T, 3>& InParticles);
+	CHAOS_API void UpdateGridFromPositions(const Chaos::TDynamicParticles<T, 3>& InParticles);
 	
 	enum InterpType { linear = 1, quadratic = 2 };
 	void SetInterp(InterpType InterpIn);
@@ -461,12 +461,19 @@ FArchive& operator<<(FArchive& Ar, TUniformGridBase<T, d>& Value)
 	return Ar;
 }
 
-#if PLATFORM_MAC || PLATFORM_LINUX
+#if !IS_MERGEDMODULES
+#if PLATFORM_COMPILER_CLANG
 extern template class CHAOS_API Chaos::TUniformGridBase<Chaos::FReal, 3>;
 extern template class CHAOS_API Chaos::TUniformGrid<Chaos::FReal, 3>;
 extern template class CHAOS_API Chaos::TMPMGrid<Chaos::FReal>;
 extern template class CHAOS_API Chaos::TMPMGrid<Chaos::FRealSingle>;
 extern template class CHAOS_API Chaos::TUniformGrid<Chaos::FReal, 2>;
+#else
+extern template class Chaos::TUniformGridBase<Chaos::FReal, 3>;
+extern template class Chaos::TUniformGrid<Chaos::FReal, 3>;
+extern template class Chaos::TMPMGrid<Chaos::FReal>;
+extern template class Chaos::TMPMGrid<Chaos::FRealSingle>;
+extern template class Chaos::TUniformGrid<Chaos::FReal, 2>;
+#endif
 #endif
-
 }
diff --git a/Engine/Source/Runtime/Experimental/Chaos/Public/GeometryCollection/Facades/CollectionTetrahedralFacade.h b/Engine/Source/Runtime/Experimental/Chaos/Public/GeometryCollection/Facades/CollectionTetrahedralFacade.h
--- a/Engine/Source/Runtime/Experimental/Chaos/Public/GeometryCollection/Facades/CollectionTetrahedralFacade.h
+++ b/Engine/Source/Runtime/Experimental/Chaos/Public/GeometryCollection/Facades/CollectionTetrahedralFacade.h
@@ -64,7 +64,7 @@ namespace GeometryCollection::Facades
 		/**
 		*  Barycentric intersections with Tetarhedron
 		*/
-		bool Intersection(
+		CHAOS_API bool Intersection(
 			const TConstArrayView<Chaos::Softs::FSolverVec3>& SamplePositions, 
 			const TConstArrayView<Chaos::Softs::FSolverVec3>& TetarhedronPositions,
 			TArray<TetrahedralParticleEmbedding>& OutIntersections) const;
