GCC warning fixes.

#rnx

diff --git a/Engine/Source/Runtime/Experimental/IoStore/HttpClient/Private/TransactionOne.inl b/Engine/Source/Runtime/Experimental/IoStore/HttpClient/Private/TransactionOne.inl
--- a/Engine/Source/Runtime/Experimental/IoStore/HttpClient/Private/TransactionOne.inl
+++ b/Engine/Source/Runtime/Experimental/IoStore/HttpClient/Private/TransactionOne.inl
@@ -689,7 +689,7 @@ FOutcome FBody::Recv(FMutableMemoryView Dest, FTlsPeer& Peer)
 		return FOutcome::Ok(Result);
 	}
 
-	if (Result == Dest.GetSize())
+	if (Result == int32(Dest.GetSize()))
 	{
 		return FOutcome::WaitBuffer(Result);
 	}
diff --git a/Engine/Source/Runtime/Experimental/IoStore/HttpClient/Private/TransactionTwo.inl b/Engine/Source/Runtime/Experimental/IoStore/HttpClient/Private/TransactionTwo.inl
--- a/Engine/Source/Runtime/Experimental/IoStore/HttpClient/Private/TransactionTwo.inl
+++ b/Engine/Source/Runtime/Experimental/IoStore/HttpClient/Private/TransactionTwo.inl
@@ -588,7 +588,7 @@ int32 FDriverNg::OnHeader(nghttp2_rcbuf* Name, nghttp2_rcbuf* Value)
 		return NGHTTP2_ERR_CALLBACK_FAILURE;
 	}
 
-	FHeaderNg Header(Name, Value);
+	FHeaderNg Header = { Name, Value };
 	ResponseHeaders->Add(MoveTemp(Header));
 	return 0;
 }
@@ -780,12 +780,12 @@ FOutcome FStatusHeaders::TryRecvResponse(FTlsPeer& Peer)
 		FAnsiStringView Key = Header.Key.Get();
 		FAnsiStringView Value = Header.Value.Get();
 
-		if ((Flags ^ Flag_Code) && (Key == ":status"))
+		if (!(Flags & Flag_Code) && (Key == ":status"))
 		{
 			StatusCode = uint32(CrudeToInt(Value));
 			Flags |= Flag_Code;
 		}
-		else if ((Flags ^ Flag_Length) && (Key == "content-length"))
+		else if (!(Flags & Flag_Length) && (Key == "content-length"))
 		{
 			ContentLength = CrudeToInt(Value);
 			Flags |= Flag_Length;
