Chaos - follow up fix to JointConstraintProxy fix for changing solver type

#rnx

diff --git a/Engine/Source/Runtime/Experimental/Chaos/Private/PhysicsProxy/JointConstraintProxy.cpp b/Engine/Source/Runtime/Experimental/Chaos/Private/PhysicsProxy/JointConstraintProxy.cpp
--- a/Engine/Source/Runtime/Experimental/Chaos/Private/PhysicsProxy/JointConstraintProxy.cpp
+++ b/Engine/Source/Runtime/Experimental/Chaos/Private/PhysicsProxy/JointConstraintProxy.cpp
@@ -127,7 +127,7 @@ void FJointConstraintPhysicsProxy::InitializeOnPhysicsThread(FPBDRigidsSolver* I
 					};
 					FixConnectorTransformsForRoot(OriginalParticleHandles_PT, RootHandles, JointSettings.ConnectorTransforms);
 
-					InSolver->GetEvolution()->CreateJointConstraint(RootHandles, JointSettings);
+					Constraint_PT = InSolver->GetEvolution()->CreateJointConstraint(RootHandles, JointSettings);
 				}
 			}
 		}
