Chaos Deformable Solver Integration: moved PMatrix identity and zero definitions from Chaos to ChaosCore

diff --git a/Engine/Source/Runtime/Experimental/Chaos/Private/Chaos/Utilities.cpp b/Engine/Source/Runtime/Experimental/Chaos/Private/Chaos/Utilities.cpp
--- a/Engine/Source/Runtime/Experimental/Chaos/Private/Chaos/Utilities.cpp
+++ b/Engine/Source/Runtime/Experimental/Chaos/Private/Chaos/Utilities.cpp
@@ -6,13 +6,6 @@
 
 namespace Chaos
 {
-	// @todo(ccaulfield): should be in ChaosCore, but that can't actually include its own headers at the moment (e.g., Matrix.h includes headers from Chaos)
-	const PMatrix<FRealSingle, 3, 3> PMatrix<FRealSingle, 3, 3>::Zero = PMatrix<FRealSingle, 3, 3>(0, 0, 0);
-	const PMatrix<FRealSingle, 3, 3> PMatrix<FRealSingle, 3, 3>::Identity = PMatrix<FRealSingle, 3, 3>(1, 1, 1);
-
-	const PMatrix<FRealDouble, 3, 3> PMatrix<FRealDouble, 3, 3>::Zero = PMatrix<FRealDouble, 3, 3>(0, 0, 0);
-	const PMatrix<FRealDouble, 3, 3> PMatrix<FRealDouble, 3, 3>::Identity = PMatrix<FRealDouble, 3, 3>(1, 1, 1);
-
 	namespace Utilities
 	{
 		FReal GetSolverPhysicsResultsTime(const FPhysicsSolverBase* Solver)
diff --git a/Engine/Source/Runtime/Experimental/ChaosCore/Private/Transform.cpp b/Engine/Source/Runtime/Experimental/ChaosCore/Private/Transform.cpp
--- a/Engine/Source/Runtime/Experimental/ChaosCore/Private/Transform.cpp
+++ b/Engine/Source/Runtime/Experimental/ChaosCore/Private/Transform.cpp
@@ -2,6 +2,7 @@
 
 #include "Chaos/Transform.h"
 #include "Chaos/Real.h"
+#include "Chaos/Matrix.h"
 
 namespace Chaos
 {
@@ -16,4 +17,10 @@ namespace Chaos
 		// LWC_TODO: Perf pessimization
 		return ToMatrixNoScale() * static_cast<const UE::Math::TMatrix<FRealDouble>&>(Matrix);
 	}
+
+	const PMatrix<FRealSingle, 3, 3> PMatrix<FRealSingle, 3, 3>::Zero = PMatrix<FRealSingle, 3, 3>(0, 0, 0);
+	const PMatrix<FRealSingle, 3, 3> PMatrix<FRealSingle, 3, 3>::Identity = PMatrix<FRealSingle, 3, 3>(1, 1, 1);
+
+	const PMatrix<FRealDouble, 3, 3> PMatrix<FRealDouble, 3, 3>::Zero = PMatrix<FRealDouble, 3, 3>(0, 0, 0);
+	const PMatrix<FRealDouble, 3, 3> PMatrix<FRealDouble, 3, 3>::Identity = PMatrix<FRealDouble, 3, 3>(1, 1, 1);
 }
diff --git a/Engine/Source/Runtime/Experimental/ChaosCore/Public/Chaos/Matrix.h b/Engine/Source/Runtime/Experimental/ChaosCore/Public/Chaos/Matrix.h
--- a/Engine/Source/Runtime/Experimental/ChaosCore/Public/Chaos/Matrix.h
+++ b/Engine/Source/Runtime/Experimental/ChaosCore/Public/Chaos/Matrix.h
@@ -810,8 +810,8 @@ namespace Chaos
 				X[2] * Y[2]);
 		}
 
-		static const PMatrix<FRealDouble, 3, 3> Zero;
-		static const PMatrix<FRealDouble, 3, 3> Identity;
+		static CHAOSCORE_API const PMatrix<FRealDouble, 3, 3> Zero;
+		static CHAOSCORE_API const PMatrix<FRealDouble, 3, 3> Identity;
 	};
 
 	template<>
@@ -1252,8 +1252,8 @@ namespace Chaos
 		}
 
 
-		static const PMatrix<FRealSingle, 3, 3> Zero;
-		static const PMatrix<FRealSingle, 3, 3> Identity;
+		static CHAOSCORE_API const PMatrix<FRealSingle, 3, 3> Zero;
+		static CHAOSCORE_API const PMatrix<FRealSingle, 3, 3> Identity;
 	};
 
 	typedef PMatrix<double, 3, 3> PMatrix33d;
