Collision Filter Refactor Part 3.2: Fixing Word3 to be consistent between sim and query.
This is a suprisingly nuanced change to make sure code works as well as possible due to many spots setting one filter as empty or with small differences on the flags.
To be as safe as possible, some flags are taken from query and some from sim where appropriate.
Also updated several call sites to fix issues with word 3.
Also fixed a few more spots to use the new API now that word 3 is consistent.

diff --git a/Engine/Plugins/Experimental/SpatialReadiness/Source/SpatialReadiness/Private/SpatialReadiness/SpatialReadinessSimCallback.cpp b/Engine/Plugins/Experimental/SpatialReadiness/Source/SpatialReadiness/Private/SpatialReadiness/SpatialReadinessSimCallback.cpp
--- a/Engine/Plugins/Experimental/SpatialReadiness/Source/SpatialReadiness/Private/SpatialReadiness/SpatialReadinessSimCallback.cpp
+++ b/Engine/Plugins/Experimental/SpatialReadiness/Source/SpatialReadiness/Private/SpatialReadiness/SpatialReadinessSimCallback.cpp
@@ -104,6 +104,7 @@ int32 FSpatialReadinessSimCallback::AddUnreadyVolume_GT(const FBox& Bounds, cons
 	QueryFilterData.Word2 |= QueryFilterData.Word1;
 	QueryFilterData.Word1 = 0;
 	QueryFilterData.Word3 |= EPDF_SimpleCollision | EPDF_ComplexCollision;
+	SimFilterData.Word3 = QueryFilterData.Word3;
 
 	// Make the geometry
 	ParticleHandle.SetGeometry(BoxGeom);
diff --git a/Engine/Source/Programs/HeadlessChaos/Private/HeadlessChaosTestPhysicsScene.cpp b/Engine/Source/Programs/HeadlessChaos/Private/HeadlessChaosTestPhysicsScene.cpp
--- a/Engine/Source/Programs/HeadlessChaos/Private/HeadlessChaosTestPhysicsScene.cpp
+++ b/Engine/Source/Programs/HeadlessChaos/Private/HeadlessChaosTestPhysicsScene.cpp
@@ -510,6 +510,7 @@ namespace ChaosTest {
 		FCollisionFilterData NewParticleQueryFilter;
 		NewParticleQueryFilter.Word1 = TNumericLimits<int32>::Max();
 		NewParticleQueryFilter.Word2 = TNumericLimits<int32>::Max();
+		NewParticleQueryFilter.Word3 = TNumericLimits<int32>::Max();
 		for (const TUniquePtr<Chaos::FPerShapeData>& Shape : Particle.ShapesArray())
 		{
 			PRAGMA_DISABLE_INTERNAL_WARNINGS
diff --git a/Engine/Source/Programs/HeadlessChaos/Private/HeadlessChaosTestUtility.cpp b/Engine/Source/Programs/HeadlessChaos/Private/HeadlessChaosTestUtility.cpp
--- a/Engine/Source/Programs/HeadlessChaos/Private/HeadlessChaosTestUtility.cpp
+++ b/Engine/Source/Programs/HeadlessChaos/Private/HeadlessChaosTestUtility.cpp
@@ -590,11 +590,6 @@ namespace ChaosTest {
 		InParticles.SetObjectStateLowLevel(EObjectStateType::Dynamic);
 
 		::ChaosTest::SetParticleSimDataToCollide({ &InParticles });
-		//for (const TUniquePtr<Chaos::FPerShapeData>& Shape : InParticles.ShapesArray())
-		//{
-		//	Shape->SimData.Word3 = 1;
-		//	Shape->SimData.Word1 = 1;
-		//}
 	}
 
 	FPBDRigidParticleHandle* AppendDynamicParticleConvexBox(FPBDRigidsSOAs& SOAs, const FVec3& Scale)
@@ -672,17 +667,22 @@ namespace ChaosTest {
 		return FVec3(FVector::UpVector);
 	}
 
+	void SetShapeFilterToCollide(Chaos::FPerShapeData& Shape)
+	{
+		// TODO @ JoshD: Update to use new API once builder API is finalized.
+		FCollisionFilterData FilterData = Shape.GetSimData();
+		FilterData.Word1 = 1;
+		FilterData.Word3 = 1;
+		Shape.SetSimData(FilterData);
+	}
+
 	void SetParticleSimDataToCollide(TArray< Chaos::FGeometryParticle* > ParticleArray)
 	{
 		for (Chaos::FGeometryParticle* Particle : ParticleArray)
 		{
 			for (const TUniquePtr<Chaos::FPerShapeData>& Shape :Particle->ShapesArray())
 			{
-				Shape->ModifySimData([](auto& SimData)
-				{
-					SimData.Word3 = 1;
-					SimData.Word1 = 1;
-				});
+				SetShapeFilterToCollide(*Shape);
 			}
 		}
 	}
@@ -692,11 +692,7 @@ namespace ChaosTest {
 		{
 			for (const TUniquePtr<Chaos::FPerShapeData>& Shape : Particle->ShapesArray())
 			{
-				Shape->ModifySimData([](auto& SimData)
-				{
-					SimData.Word3 = 1;
-					SimData.Word1 = 1;
-				});
+				SetShapeFilterToCollide(*Shape);
 			}
 		}
 	}
diff --git a/Engine/Source/Runtime/Engine/Private/PhysicsEngine/Experimental/ChaosInterfaceUtils.cpp b/Engine/Source/Runtime/Engine/Private/PhysicsEngine/Experimental/ChaosInterfaceUtils.cpp
--- a/Engine/Source/Runtime/Engine/Private/PhysicsEngine/Experimental/ChaosInterfaceUtils.cpp
+++ b/Engine/Source/Runtime/Engine/Private/PhysicsEngine/Experimental/ChaosInterfaceUtils.cpp
@@ -103,8 +103,9 @@ namespace ChaosInterface
 		auto NewShapeHelper = [&InParams, &CollisionTraceType](const Chaos::FImplicitObjectPtr& InGeom, int32 ShapeIdx, void* UserData, ECollisionEnabled::Type ShapeCollisionEnabled, bool bComplexShape = false)
 		{
 			TUniquePtr<Chaos::FPerShapeData> NewShape = Chaos::FShapeInstanceProxy::Make(ShapeIdx, InGeom);
-			NewShape->SetQueryData(bComplexShape ? InParams.CollisionData.CollisionFilterData.QueryComplexFilter : InParams.CollisionData.CollisionFilterData.QuerySimpleFilter);
-			NewShape->SetSimData(InParams.CollisionData.CollisionFilterData.SimFilter);
+			const FCollisionFilterData& QueryFilter = bComplexShape ? InParams.CollisionData.CollisionFilterData.QueryComplexFilter : InParams.CollisionData.CollisionFilterData.QuerySimpleFilter;
+			Chaos::Filter::FCombinedShapeFilterData CombinedShapeFilter = Chaos::Filter::FShapeFilterBuilder::BuildFromLegacyShapeFilter(QueryFilter, InParams.CollisionData.CollisionFilterData.SimFilter);
+			NewShape->SetCombinedShapeFilterData(CombinedShapeFilter);
 			NewShape->SetCollisionTraceType(ConvertCollisionTraceFlag(CollisionTraceType));
 			NewShape->UpdateShapeBounds(InParams.WorldTransform);
 			NewShape->SetUserData(UserData);
diff --git a/Engine/Source/Runtime/Engine/Private/PhysicsEngine/ImmediatePhysics/ImmediatePhysicsAdapters.cpp b/Engine/Source/Runtime/Engine/Private/PhysicsEngine/ImmediatePhysics/ImmediatePhysicsAdapters.cpp
--- a/Engine/Source/Runtime/Engine/Private/PhysicsEngine/ImmediatePhysics/ImmediatePhysicsAdapters.cpp
+++ b/Engine/Source/Runtime/Engine/Private/PhysicsEngine/ImmediatePhysics/ImmediatePhysicsAdapters.cpp
@@ -126,6 +126,9 @@ namespace ImmediatePhysics
 		// @todo(chaos): we need an API for setting up filters
 		BodyCollisionData.CollisionFilterData.SimFilter.Word1 = 0xFFFF;
 		BodyCollisionData.CollisionFilterData.SimFilter.Word3 = 0xFFFF;
+		// Ensure the word3s are equal (until collision filter refactor is complete)
+		BodyCollisionData.CollisionFilterData.QuerySimpleFilter.Word3 = 0xFFFF;
+		BodyCollisionData.CollisionFilterData.QueryComplexFilter.Word3 = 0xFFFF;
 		BodyCollisionData.CollisionFlags.bEnableSimCollisionSimple = true;
 
 		BodyInstance->BuildBodyCollisionFlags(BodyCollisionData.CollisionFlags, BodyInstance->GetCollisionEnabled(), BodyInstance->BodySetup->GetCollisionTraceFlag() == CTF_UseComplexAsSimple);
diff --git a/Engine/Source/Runtime/Experimental/Chaos/Private/Chaos/CollisionFilterData.cpp b/Engine/Source/Runtime/Experimental/Chaos/Private/Chaos/CollisionFilterData.cpp
--- a/Engine/Source/Runtime/Experimental/Chaos/Private/Chaos/CollisionFilterData.cpp
+++ b/Engine/Source/Runtime/Experimental/Chaos/Private/Chaos/CollisionFilterData.cpp
@@ -18,6 +18,9 @@ namespace Chaos::Filter
 	constexpr uint32 CollisionChannelMask = ((1 << NumCollisionChannelBits) - 1) << CollisionChannelBitsOffset;
 	constexpr uint32 MaskFilterMask = ((1 << NumMaskFilterBits) - 1) << MaskFilterBitsOffset;
 
+	constexpr uint32 LegacyQueryFlagsMask = (uint32)(EFilterFlags::SimpleCollision | EFilterFlags::ComplexCollision);
+	constexpr uint32 LegacySimFlagsMask = (uint32)(EFilterFlags::All) & ~LegacyQueryFlagsMask;
+
 	uint32 GetFlags(uint32 Word3)
 	{
 		const uint32 FilterFlags = (Word3 & FilterFlagsMask);
@@ -27,7 +30,7 @@ namespace Chaos::Filter
 	void SetFlags(uint32& Word3, uint32 Flags)
 	{
 		const uint32 Word3Cleared = Word3 & ~FilterFlagsMask;
-		Word3 = Word3Cleared | (Flags | FilterFlagsMask);
+		Word3 = Word3Cleared | (Flags & FilterFlagsMask);
 	}
 
 	bool HasFlag(uint32 Word3, EFilterFlags InFlag)
@@ -52,13 +55,17 @@ namespace Chaos::Filter
 		return Word3 >> (32u - NumMaskFilterBits);
 	}
 
-	void SetMaskFilter(uint32& Word3, uint8 MaskFilter)
+	void SetMaskFilter(uint32& Word3, uint8 InMaskFilter)
 	{
+		// Mask the input to make sure there's no bit overflows
+		constexpr uint32 InFilterMask = (1 << NumMaskFilterBits) - 1;
+		const uint32 MaskFilter = uint32(InMaskFilter) & InFilterMask;
+
 		static_assert(NumMaskFilterBits <= 8, "Only up to 8 extra filter bits are supported.");
 		// Clear the old mask filter bits
 		Word3 &= ~MaskFilterMask;
 		// Set the new mask filter bits
-		Word3 |= uint32(MaskFilter) << MaskFilterBitsOffset;
+		Word3 |= MaskFilter << MaskFilterBitsOffset;
 	}
 
 	FInstanceData::FInstanceData(const uint32 ActorId, const uint32 ComponentId)
@@ -105,6 +112,7 @@ namespace Chaos::Filter
 	void FShapeFilterData::SetFlags(EFilterFlags InFlags)
 	{
 		Chaos::Filter::SetFlags(Word3, (uint32)InFlags);
+		Chaos::Filter::SetFlags(SimWord3, (uint32)InFlags);
 	}
 
 	bool FShapeFilterData::HasFlag(EFilterFlags InFlag) const
@@ -120,6 +128,7 @@ namespace Chaos::Filter
 	void FShapeFilterData::SetMaskFilter(uint8 MaskFilter)
 	{
 		Chaos::Filter::SetMaskFilter(Word3, MaskFilter);
+		Chaos::Filter::SetMaskFilter(SimWord3, MaskFilter);
 	}
 
 	uint8 FShapeFilterData::GetCollisionChannelIndex() const
@@ -290,6 +299,22 @@ namespace Chaos::Filter
 
 	FCombinedShapeFilterData FShapeFilterBuilder::BuildFromLegacyShapeFilter(const FCollisionFilterData& QueryFilterData, const FCollisionFilterData& SimFilterData)
 	{
+		// This is often called with logic like: (bComplex ? SimpleFilter : ComplexFilter, SimFilter).
+		// In this case, the sim filter won't match because of the simple/complex flags.
+		// Word3 should be separate apart from these flags though, so simply pretend we're setting a new query filter.
+		// Do a basic ensure that the filter mask and collision channels are the same, if not this is a bug.
+		if (QueryFilterData.Word3 != 0 && SimFilterData.Word3 != 0)
+		{
+			const uint32 SimFilterMask = SimFilterData.Word3 & MaskFilterMask;
+			const uint32 SimCollisionChannel = SimFilterData.Word3 & CollisionChannelMask;
+			const uint32 QueryFilterMask = QueryFilterData.Word3 & MaskFilterMask;
+			const uint32 QueryCollisionChannel = QueryFilterData.Word3 & CollisionChannelMask;
+			ensure(SimFilterMask == QueryFilterMask);
+			ensure(SimCollisionChannel == QueryCollisionChannel);
+		}
+		const uint32 Word3 = GetWord3FromNewQueryFilter(QueryFilterData.Word3, SimFilterData.Word3);
+
+
 		FCombinedShapeFilterData CombinedShapeFilterData;
 		CombinedShapeFilterData.InstanceData.ActorId = QueryFilterData.Word0;
 		CombinedShapeFilterData.InstanceData.ComponentId = SimFilterData.Word2;
@@ -297,36 +322,44 @@ namespace Chaos::Filter
 
 		CombinedShapeFilterData.ShapeFilterData.QueryBlockChannels = QueryFilterData.Word1;
 		CombinedShapeFilterData.ShapeFilterData.QueryOverlapChannels = QueryFilterData.Word2;
-		CombinedShapeFilterData.ShapeFilterData.Word3 = QueryFilterData.Word3;
+		CombinedShapeFilterData.ShapeFilterData.Word3 = Word3;
 		CombinedShapeFilterData.ShapeFilterData.SimBlockChannels = SimFilterData.Word1;
-		CombinedShapeFilterData.ShapeFilterData.SimWord3 = SimFilterData.Word3;
+		CombinedShapeFilterData.ShapeFilterData.SimWord3 = Word3;
 
 		return CombinedShapeFilterData;
 	}
 
 	void FShapeFilterBuilder::SetLegacyShapeQueryFilter(FCombinedShapeFilterData& CombinedShapeFilterData, const FCollisionFilterData& QueryFilterData)
 	{
-		const FCollisionFilterData SimFilterData = GetLegacyShapeSimFilter(CombinedShapeFilterData);
-		CombinedShapeFilterData = BuildFromLegacyShapeFilter(QueryFilterData, SimFilterData);
+		FCollisionFilterData SimFilterData = GetLegacyShapeSimFilter(CombinedShapeFilterData);
+		const uint32 Word3 = GetWord3FromNewQueryFilter(QueryFilterData.Word3, SimFilterData.Word3);
+		BuildFromLegacyShapeFilter(QueryFilterData, SimFilterData, Word3, CombinedShapeFilterData);
 	}
 
 	void FShapeFilterBuilder::SetLegacyShapeSimFilter(FCombinedShapeFilterData& CombinedShapeFilterData, const FCollisionFilterData& SimFilterData)
 	{
 		const FCollisionFilterData QueryFilterData = GetLegacyShapeQueryFilter(CombinedShapeFilterData);
-		CombinedShapeFilterData = BuildFromLegacyShapeFilter(QueryFilterData, SimFilterData);
+		const uint32 Word3 = GetWord3FromNewSimFilter(SimFilterData.Word3, QueryFilterData.Word3);
+		BuildFromLegacyShapeFilter(QueryFilterData, SimFilterData, Word3, CombinedShapeFilterData);
 	}
 
 	void FShapeFilterBuilder::GetLegacyShapeFilter(const FCombinedShapeFilterData& CombinedShapeFilterData, FCollisionFilterData& OutQueryFilterData, FCollisionFilterData& OutSimFilterData)
 	{
+		// Grab whichever word3 is valid (temporary safeguard for legacy).
+		const FShapeFilterData& ShapeFilterData = CombinedShapeFilterData.ShapeFilterData;
+		const uint32 Word3 = ShapeFilterData.Word3 != 0 ? ShapeFilterData.Word3 : ShapeFilterData.SimWord3;
+		// We expect the word3's to always be equal in the transition period before removal of SimWord3.
+		ensure(ShapeFilterData.Word3 == ShapeFilterData.SimWord3);
+
 		OutQueryFilterData.Word0 = CombinedShapeFilterData.InstanceData.ActorId;
 		OutSimFilterData.Word2 = CombinedShapeFilterData.InstanceData.ComponentId;
 		OutSimFilterData.Word0 = CombinedShapeFilterData.InstanceData.BodyIndex;
 
 		OutQueryFilterData.Word1 = CombinedShapeFilterData.ShapeFilterData.QueryBlockChannels;
 		OutQueryFilterData.Word2 = CombinedShapeFilterData.ShapeFilterData.QueryOverlapChannels;
-		OutQueryFilterData.Word3 = CombinedShapeFilterData.ShapeFilterData.Word3;
+		OutQueryFilterData.Word3 = Word3;
 		OutSimFilterData.Word1 = CombinedShapeFilterData.ShapeFilterData.SimBlockChannels;
-		OutSimFilterData.Word3 = CombinedShapeFilterData.ShapeFilterData.SimWord3;
+		OutSimFilterData.Word3 = Word3;
 	}
 
 	FCollisionFilterData FShapeFilterBuilder::GetLegacyShapeQueryFilter(const FCombinedShapeFilterData& CombinedShapeFilterData)
@@ -343,6 +376,45 @@ namespace Chaos::Filter
 		return SimFilter;
 	}
 
+	uint32 FShapeFilterBuilder::GetWord3FromNewSimFilter(const uint32 SimWord3, const uint32 QueryWord3)
+	{
+		// If the sim's word3 isn't set, this is likely legacy setting an empty sim word or an ordering issue, so we just use the query word3.
+		// Otherwise, the new word3 is the sim's word3 with a few bits stolen from the flag of the query word3 (simple/complex only affect queries currently).
+		if (SimWord3 != 0)
+		{
+			const uint32 MaskFilter = SimWord3 & MaskFilterMask;
+			const uint32 ChannelIndex = SimWord3 & CollisionChannelMask;
+			const uint32 Flags = (SimWord3 & LegacySimFlagsMask) | (QueryWord3 & LegacyQueryFlagsMask);
+			return MaskFilter | ChannelIndex | Flags;
+		}
+		return QueryWord3;
+	}
+
+	uint32 FShapeFilterBuilder::GetWord3FromNewQueryFilter(const uint32 QueryWord3, const uint32 SimWord3)
+	{
+		if (QueryWord3 != 0)
+		{
+			const uint32 MaskFilter = QueryWord3 & MaskFilterMask;
+			const uint32 ChannelIndex = QueryWord3 & CollisionChannelMask;
+			const uint32 Flags = (SimWord3 & LegacySimFlagsMask) | (QueryWord3 & LegacyQueryFlagsMask);
+			return MaskFilter | ChannelIndex | Flags;
+		}
+		return SimWord3;
+	}
+
+	void FShapeFilterBuilder::BuildFromLegacyShapeFilter(const FCollisionFilterData& QueryFilterData, const FCollisionFilterData& SimFilterData, const uint32 Word3, FCombinedShapeFilterData& OutCombinedShapeFilterData)
+	{
+		OutCombinedShapeFilterData.InstanceData.ActorId = QueryFilterData.Word0;
+		OutCombinedShapeFilterData.InstanceData.ComponentId = SimFilterData.Word2;
+		OutCombinedShapeFilterData.InstanceData.BodyIndex = SimFilterData.Word0;
+
+		OutCombinedShapeFilterData.ShapeFilterData.QueryBlockChannels = QueryFilterData.Word1;
+		OutCombinedShapeFilterData.ShapeFilterData.QueryOverlapChannels = QueryFilterData.Word2;
+		OutCombinedShapeFilterData.ShapeFilterData.Word3 = Word3;
+		OutCombinedShapeFilterData.ShapeFilterData.SimBlockChannels = SimFilterData.Word1;
+		OutCombinedShapeFilterData.ShapeFilterData.SimWord3 = Word3;
+	}
+
 	FQueryFilterData FQueryFilterBuilder::CreateLegacyObjectTypeFilter(const uint32 ObjectTypesToQuery, const bool bMultiTrace, const uint8 IgnoreMask, const EFilterFlags FilterFlags)
 	{
 		const uint8 MultiTraceValue = bMultiTrace ? 0x1 : 0x0;
diff --git a/Engine/Source/Runtime/Experimental/Chaos/Private/Chaos/ParticleDirtyFlags.cpp b/Engine/Source/Runtime/Experimental/Chaos/Private/Chaos/ParticleDirtyFlags.cpp
--- a/Engine/Source/Runtime/Experimental/Chaos/Private/Chaos/ParticleDirtyFlags.cpp
+++ b/Engine/Source/Runtime/Experimental/Chaos/Private/Chaos/ParticleDirtyFlags.cpp
@@ -15,6 +15,7 @@ namespace Chaos
 	{
 		PRAGMA_DISABLE_DEPRECATION_WARNINGS
 		QueryData = InQueryData;
+		SimData.Word3 = QueryData.Word3 = Chaos::Filter::FShapeFilterBuilder::GetWord3FromNewQueryFilter(QueryData.Word3, SimData.Word3);
 		PRAGMA_ENABLE_DEPRECATION_WARNINGS
 	}
 
@@ -29,6 +30,7 @@ namespace Chaos
 	{
 		PRAGMA_DISABLE_DEPRECATION_WARNINGS
 		SimData = InSimData;
+		SimData.Word3 = QueryData.Word3 = Chaos::Filter::FShapeFilterBuilder::GetWord3FromNewSimFilter(SimData.Word3, QueryData.Word3);
 		PRAGMA_ENABLE_DEPRECATION_WARNINGS
 	}
 
diff --git a/Engine/Source/Runtime/Experimental/Chaos/Private/EventDefaults.cpp b/Engine/Source/Runtime/Experimental/Chaos/Private/EventDefaults.cpp
--- a/Engine/Source/Runtime/Experimental/Chaos/Private/EventDefaults.cpp
+++ b/Engine/Source/Runtime/Experimental/Chaos/Private/EventDefaults.cpp
@@ -148,8 +148,8 @@ namespace Chaos
 									const FPerShapeData* Shape1 = Constraint.GetShape1();
 
 									// If we don't have a filter - allow the notify, otherwise obey the filter flag
-									const bool bFilter0Notify = Shape0 ? Shape0->GetSimData().HasFlag(EFilterFlags::ContactNotify) : true;
-									const bool bFilter1Notify = Shape1 ? Shape1->GetSimData().HasFlag(EFilterFlags::ContactNotify) : true;
+									const bool bFilter0Notify = Shape0 ? Shape0->GetShapeFilterData().HasFlag(EFilterFlags::ContactNotify) : true;
+									const bool bFilter1Notify = Shape1 ? Shape1->GetShapeFilterData().HasFlag(EFilterFlags::ContactNotify) : true;
 
 									if (!bFilter0Notify && !bFilter1Notify)
 									{
diff --git a/Engine/Source/Runtime/Experimental/Chaos/Private/PhysicsProxy/GeometryCollectionPhysicsProxy.cpp b/Engine/Source/Runtime/Experimental/Chaos/Private/PhysicsProxy/GeometryCollectionPhysicsProxy.cpp
--- a/Engine/Source/Runtime/Experimental/Chaos/Private/PhysicsProxy/GeometryCollectionPhysicsProxy.cpp
+++ b/Engine/Source/Runtime/Experimental/Chaos/Private/PhysicsProxy/GeometryCollectionPhysicsProxy.cpp
@@ -378,24 +378,26 @@ void SetImplicitToPTParticles(Chaos::TPBDRigidParticleHandle<Chaos::FReal, 3>* H
 	{
 		// Override collision filters and make this body collide with everything.
 		int32 CurrShape = 0;
-		FCollisionFilterData FilterData;
-		FilterData.Word1 = 0xFFFF; // this body channel
-		FilterData.Word3 = 0xFFFF; // collision candidate channels
+		FCollisionFilterData SimFilterData;
+		SimFilterData.Word1 = 0xFFFF; // this body channel
+		SimFilterData.Word3 = 0xFFFF; // collision candidate channels
+		FCollisionFilterData QueryFilterData;
+		QueryFilterData.Word3 = SimFilterData.Word3;
+		const Chaos::Filter::FCombinedShapeFilterData CombinedFilterData = Chaos::Filter::FShapeFilterBuilder::BuildFromLegacyShapeFilter(QueryFilterData, SimFilterData);
 		for (const TUniquePtr<Chaos::FPerShapeData>& Shape : Handle->ShapesArray())
 		{
 			Shape->SetSimEnabled(true);
 			Shape->SetCollisionTraceType(Chaos::EChaosCollisionTraceFlag::Chaos_CTF_UseDefault);
 			//Shape->CollisionTraceType = Chaos::EChaosCollisionTraceFlag::Chaos_CTF_UseSimpleAndComplex;
-			Shape->SetSimData(FilterData);
-			Shape->SetQueryData(FCollisionFilterData());
+			Shape->SetCombinedShapeFilterData(CombinedFilterData);
 		}
 	}
 	else
 	{
+		const Chaos::Filter::FCombinedShapeFilterData CombinedFilterData = Chaos::Filter::FShapeFilterBuilder::BuildFromLegacyShapeFilter(QueryFilterIn, SimFilterIn);
 		for (const TUniquePtr<Chaos::FPerShapeData>& Shape : Handle->ShapesArray())
 		{
-			Shape->SetSimData(SimFilterIn);
-			Shape->SetQueryData(QueryFilterIn);
+			Shape->SetCombinedShapeFilterData(CombinedFilterData);
 		}
 	}
 }
diff --git a/Engine/Source/Runtime/Experimental/Chaos/Public/Chaos/CollisionFilterData.h b/Engine/Source/Runtime/Experimental/Chaos/Public/Chaos/CollisionFilterData.h
--- a/Engine/Source/Runtime/Experimental/Chaos/Public/Chaos/CollisionFilterData.h
+++ b/Engine/Source/Runtime/Experimental/Chaos/Public/Chaos/CollisionFilterData.h
@@ -222,7 +222,11 @@ namespace Chaos::Filter
 		UE_INTERNAL CHAOS_API static void GetLegacyShapeFilter(const FCombinedShapeFilterData& CombinedShapeFilterData, FCollisionFilterData& OutQueryFilterData, FCollisionFilterData& OutSimFilterData);
 		UE_INTERNAL CHAOS_API static FCollisionFilterData GetLegacyShapeQueryFilter(const FCombinedShapeFilterData& CombinedShapeFilterData);
 		UE_INTERNAL CHAOS_API static FCollisionFilterData GetLegacyShapeSimFilter(const FCombinedShapeFilterData& CombinedShapeFilterData);
+
+		UE_INTERNAL static uint32 GetWord3FromNewSimFilter(const uint32 SimWord3, const uint32 QueryWord3);
+		UE_INTERNAL static uint32 GetWord3FromNewQueryFilter(const uint32 QueryWord3, const uint32 SimWord3);
 	private:
+		static void BuildFromLegacyShapeFilter(const FCollisionFilterData& QueryFilterData, const FCollisionFilterData& SimFilterData, const uint32 Word3, FCombinedShapeFilterData& OutCombinedShapeFilterData);
 	};
 
 	class UE_INTERNAL FQueryFilterBuilder
diff --git a/Engine/Source/Runtime/Experimental/Chaos/Public/Chaos/ShapeInstance.h b/Engine/Source/Runtime/Experimental/Chaos/Public/Chaos/ShapeInstance.h
--- a/Engine/Source/Runtime/Experimental/Chaos/Public/Chaos/ShapeInstance.h
+++ b/Engine/Source/Runtime/Experimental/Chaos/Public/Chaos/ShapeInstance.h
@@ -442,6 +442,7 @@ namespace Chaos
 		}
 
 		template <typename Lambda>
+		UE_DEPRECATED(5.8, "Get/SetCombinedShapeFilterData instead")
 		void ModifySimData(const Lambda& LambdaFunc)
 		{
 			PRAGMA_DISABLE_DEPRECATION_WARNINGS
@@ -690,7 +691,9 @@ namespace Chaos
 		void SetProxy(IPhysicsProxyBase* InProxy) { check(false); }
 
 		PRAGMA_DISABLE_DEPRECATION_WARNINGS
-		template <typename Lambda> void ModifySimData(const Lambda& LambdaFunc) { LambdaFunc(CollisionData.SimData); }
+		template <typename Lambda>
+		UE_DEPRECATED(5.8, "Get/SetCombinedShapeFilterData instead")
+		void ModifySimData(const Lambda& LambdaFunc) { LambdaFunc(CollisionData.SimData); }
 		PRAGMA_ENABLE_DEPRECATION_WARNINGS
 		template <typename Lambda> void ModifyMaterials(const Lambda& LambdaFunc) { LambdaFunc(GetMaterialDataImpl().Materials); }
 		template <typename Lambda> void ModifyMaterialMasks(const Lambda& LambdaFunc) { LambdaFunc(GetMaterialDataImpl().MaterialMasks); }
