Improve profiling tag  for CSV capture
Distinguish StepSolver from PushData and GTPreSimCallbacks

diff --git a/Engine/Source/Runtime/Experimental/Chaos/Private/Chaos/Framework/PhysicsSolverBase.cpp b/Engine/Source/Runtime/Experimental/Chaos/Private/Chaos/Framework/PhysicsSolverBase.cpp
--- a/Engine/Source/Runtime/Experimental/Chaos/Private/Chaos/Framework/PhysicsSolverBase.cpp
+++ b/Engine/Source/Runtime/Experimental/Chaos/Private/Chaos/Framework/PhysicsSolverBase.cpp
@@ -102,7 +102,7 @@ namespace Chaos
 		LLM_SCOPE(ELLMTag::ChaosUpdate);
 		SCOPE_CYCLE_COUNTER(STAT_ChaosTick);
 		CSV_SCOPED_TIMING_STAT_EXCLUSIVE(Physics);
-		CSV_SCOPED_TIMING_STAT(PhysicsVerbose, StepSolver);
+		CSV_SCOPED_TIMING_STAT(PhysicsVerbose, StepSolver_PushData);
 		
 		CVD_SCOPE_CONTEXT(Solver.GetChaosVDContextData());
 
@@ -122,7 +122,7 @@ namespace Chaos
 		LLM_SCOPE(ELLMTag::ChaosUpdate);
 		SCOPE_CYCLE_COUNTER(STAT_ChaosTick);
 		CSV_SCOPED_TIMING_STAT_EXCLUSIVE(Physics);
-		CSV_SCOPED_TIMING_STAT(PhysicsVerbose, StepSolver);
+		CSV_SCOPED_TIMING_STAT(PhysicsVerbose, StepSolver_GTPreSimCallbacks);
 
 		//We are on GT, but we know PhysicsThread is waiting so we're actually going to operate on PT data
 #if PHYSICS_THREAD_CONTEXT
