Collision Filter Refactor Part 3.3: Adding internal builders and fixing a lot of call sites that were directly setting words.
Most of these were just setting words to full masks for blocking all.
I tried to preserve the original bit pattern where possible, such as only blocking for sim vs. everything.

diff --git a/Engine/Plugins/Experimental/ControlRigPhysics/Source/ControlRigPhysics/Private/RigPhysicsSimulation.cpp b/Engine/Plugins/Experimental/ControlRigPhysics/Source/ControlRigPhysics/Private/RigPhysicsSimulation.cpp
--- a/Engine/Plugins/Experimental/ControlRigPhysics/Source/ControlRigPhysics/Private/RigPhysicsSimulation.cpp
+++ b/Engine/Plugins/Experimental/ControlRigPhysics/Source/ControlRigPhysics/Private/RigPhysicsSimulation.cpp
@@ -404,9 +404,11 @@ static bool CreateGeometry(
 	// Set the filter to collide with everything (we use a broad phase that only contains particle
 	// pairs that are explicitly set to collide)
 	FBodyCollisionData BodyCollisionData;
-	// @todo(chaos): we need an API for setting up filters
-	BodyCollisionData.CollisionFilterData.SimFilter.Word1 = 0xFFFF;
-	BodyCollisionData.CollisionFilterData.SimFilter.Word3 = 0xFFFF;
+
+	const Chaos::Filter::FShapeFilterData ShapeFilter = Chaos::Filter::FShapeFilterBuilder::BuildSimBlockAll(Chaos::EFilterFlags::All);
+	Chaos::Filter::FCombinedShapeFilterData CombinedShapeFilter(ShapeFilter, Chaos::Filter::FInstanceData());
+	Chaos::Filter::FShapeFilterBuilder::GetLegacyShapeFilter(CombinedShapeFilter, BodyCollisionData.CollisionFilterData.QuerySimpleFilter, BodyCollisionData.CollisionFilterData.SimFilter);
+	BodyCollisionData.CollisionFilterData.QueryComplexFilter = BodyCollisionData.CollisionFilterData.QuerySimpleFilter;
 
 	// See FBodyInstance::BuildBodyCollisionFlags
 	BodyCollisionData.CollisionFlags.bEnableQueryCollision = false;
diff --git a/Engine/Source/Programs/HeadlessChaos/Private/GeometryCollection/GeometryCollectionTestFramework.cpp b/Engine/Source/Programs/HeadlessChaos/Private/GeometryCollection/GeometryCollectionTestFramework.cpp
--- a/Engine/Source/Programs/HeadlessChaos/Private/GeometryCollection/GeometryCollectionTestFramework.cpp
+++ b/Engine/Source/Programs/HeadlessChaos/Private/GeometryCollection/GeometryCollectionTestFramework.cpp
@@ -206,10 +206,11 @@ namespace GeometryCollectionTest
 		FCollisionFilterData SimFilterData;
 		FCollisionFilterData QueryFilterData;
 
-		// Enable all collisions
-		SimFilterData.Word1 = 0xFFFF; // this body channel
-		SimFilterData.Word3 = 0xFFFF; // collision candidate channels (also enables EPDF_ContactNotify)
-		QueryFilterData.Word3 = 0xFFFF;
+		PRAGMA_DISABLE_INTERNAL_WARNINGS
+		const Chaos::Filter::FShapeFilterData ShapeFilterData = Chaos::Filter::FShapeFilterBuilder::BuildBlockAll(EFilterFlags::All);
+		Chaos::Filter::FCombinedShapeFilterData CombinedShapeFilter(ShapeFilterData, Chaos::Filter::FInstanceData());
+		Chaos::Filter::FShapeFilterBuilder::GetLegacyShapeFilter(CombinedShapeFilter, QueryFilterData, SimFilterData);
+		PRAGMA_ENABLE_INTERNAL_WARNINGS
 
 		FGeometryCollectionPhysicsProxy* PhysObject =
 			new FGeometryCollectionPhysicsProxy(
@@ -272,11 +273,10 @@ namespace GeometryCollectionTest
 		auto& Particle = Proxy->GetGameThreadAPI();
 		Particle.SetGeometry(MakeImplicitObjectPtr<Chaos::TPlane<FReal, 3>>(FVector(0), FVector(0, 0, 1)));
 
-		FCollisionFilterData FilterData;
-		FilterData.Word1 = 0xFFFF;
-		FilterData.Word3 = 0xFFFF;
-		Particle.SetShapeSimData(0, FilterData);
-
+		PRAGMA_DISABLE_INTERNAL_WARNINGS
+		const Chaos::Filter::FShapeFilterData ShapeFilterData = Chaos::Filter::FShapeFilterBuilder::BuildBlockAll(EFilterFlags::All);
+		Particle.SetShapeFilterData(0, ShapeFilterData);
+		PRAGMA_ENABLE_INTERNAL_WARNINGS
 
 		return new RigidBodyWrapper(PhysicalMaterial, Proxy);
 	}
diff --git a/Engine/Source/Programs/HeadlessChaos/Private/HeadlessChaosTestBroadphase.cpp b/Engine/Source/Programs/HeadlessChaos/Private/HeadlessChaosTestBroadphase.cpp
--- a/Engine/Source/Programs/HeadlessChaos/Private/HeadlessChaosTestBroadphase.cpp
+++ b/Engine/Source/Programs/HeadlessChaos/Private/HeadlessChaosTestBroadphase.cpp
@@ -1094,11 +1094,12 @@ namespace ChaosTest
 		const auto& ParticlesView = Particles.GetAllParticlesView();
 
 		// ensure these can't be filtered out.
-		FCollisionFilterData FilterData;
-		FilterData.Word0 = TNumericLimits<uint32>::Max();
-		FilterData.Word1 = TNumericLimits<uint32>::Max();
-		FilterData.Word2 = TNumericLimits<uint32>::Max();
-		FilterData.Word3 = TNumericLimits<uint32>::Max();
+		PRAGMA_DISABLE_INTERNAL_WARNINGS
+		Chaos::Filter::FShapeFilterBuilder Builder;
+		Builder.SetBlockChannelMask(TNumericLimits<uint64>::Max());
+		Builder.SetOverlapChannelMask(TNumericLimits<uint64>::Max());
+		const Chaos::Filter::FShapeFilterData ShapeFilterData = Builder.Build();
+		PRAGMA_ENABLE_INTERNAL_WARNINGS
 
 		Chaos::FImplicitObjectPtr Box( new TBox<FReal, 3>(FVec3(0, 0, 0), FVec3(BoxSize, BoxSize, BoxSize)));
 
@@ -1117,9 +1118,9 @@ namespace ChaosTest
 					Handle->SetR(FRotation3::Identity);
 					GTParticle->SetR(FRotation3::Identity);
 					Handle->SetGeometry(Box);
-					Handle->ShapesArray()[0]->SetQueryData(FilterData);
+					Handle->ShapesArray()[0]->SetShapeFilterData(ShapeFilterData);
 					GTParticle->SetGeometry(Box);
-					GTParticle->ShapesArray()[0]->SetQueryData(FilterData);
+					GTParticle->ShapesArray()[0]->SetShapeFilterData(ShapeFilterData);
 					Handle->SetUniqueIdx(FUniqueIdx(Idx));
 					GTParticle->SetUniqueIdx(FUniqueIdx(Idx));
 					++Idx;
diff --git a/Engine/Source/Programs/HeadlessChaos/Private/HeadlessChaosTestPhysicsScene.cpp b/Engine/Source/Programs/HeadlessChaos/Private/HeadlessChaosTestPhysicsScene.cpp
--- a/Engine/Source/Programs/HeadlessChaos/Private/HeadlessChaosTestPhysicsScene.cpp
+++ b/Engine/Source/Programs/HeadlessChaos/Private/HeadlessChaosTestPhysicsScene.cpp
@@ -473,16 +473,14 @@ namespace ChaosTest {
 		Scene.StartFrame();
 		Scene.EndFrame();
 
-
-		// Make query filter that will allow query against particle that blocks/touches all channels.
-		// Filter will fail against particle that has no query allowed (default query filter).
-		FCollisionFilterData QueryFilter;
-		QueryFilter.Word0 = 1; // Setting to non-zero to set query type that will filter
-
-		// This is setting a somewhat arbritrary trace channels. It's very hard to make sense of these bitfields at this level of API.
-		// Below particle uses a filter that touches/blocks anything, so these bits are enough to make filter pass.
-		QueryFilter.Word3 = 7 << 21; 
-
+		// Make a query filter that will block against a particle that blocks all channels but fails on the default particle.
+		// The channel index of 7 is arbitrary.
+		PRAGMA_DISABLE_INTERNAL_WARNINGS
+		Chaos::Filter::FQueryTraceFilterBuilder Builder;
+		Builder.SetCollisionChannelIndex(7);
+		const Chaos::Filter::FQueryFilterData QueryFilterData = Builder.Build();
+		FCollisionFilterData QueryFilter = Chaos::Filter::FQueryFilterBuilder::GetLegacyQueryFilter(QueryFilterData);
+		PRAGMA_ENABLE_INTERNAL_WARNINGS
 
 		// Get collision data off shape
 		for (const TUniquePtr<Chaos::FPerShapeData>& Shape : Particle.ShapesArray())
@@ -507,24 +505,25 @@ namespace ChaosTest {
 		}
 
 		// Change filter data on game thread to contain touch/block on all channels.
-		FCollisionFilterData NewParticleQueryFilter;
-		NewParticleQueryFilter.Word1 = TNumericLimits<int32>::Max();
-		NewParticleQueryFilter.Word2 = TNumericLimits<int32>::Max();
-		NewParticleQueryFilter.Word3 = TNumericLimits<int32>::Max();
+		PRAGMA_DISABLE_INTERNAL_WARNINGS
+		Chaos::Filter::FShapeFilterBuilder ShapeFilterBuilder;
+		ShapeFilterBuilder.SetBlockChannelMask(TNumericLimits<uint64>::Max());
+		ShapeFilterBuilder.SetOverlapChannelMask(TNumericLimits<uint64>::Max());
+		ShapeFilterBuilder.SetFilterFlags(Chaos::EFilterFlags::All);
+		const Chaos::Filter::FShapeFilterData NewParticleQueryFilter = ShapeFilterBuilder.Build();
+		// Temporarily get the legacy query filter from the new filter until PrePrepQueryFilter is updated.
+		Chaos::Filter::FCombinedShapeFilterData CombinedShapeFilter(NewParticleQueryFilter, Chaos::Filter::FInstanceData());
+		const FCollisionFilterData LegacyQueryFilter = Chaos::Filter::FShapeFilterBuilder::GetLegacyShapeQueryFilter(CombinedShapeFilter);
+		PRAGMA_ENABLE_INTERNAL_WARNINGS
+
 		for (const TUniquePtr<Chaos::FPerShapeData>& Shape : Particle.ShapesArray())
 		{
-			PRAGMA_DISABLE_INTERNAL_WARNINGS
-			const FCollisionData& CollisionData = Shape->GetCollisionData();
-
 			// Update filter
-			FCollisionData NewCollisionData = CollisionData;
-			NewCollisionData.SetQueryData(NewParticleQueryFilter);
-			Shape->SetCollisionData(NewCollisionData);
+			Shape->SetShapeFilterData(NewParticleQueryFilter);
 
 			// Filter with new data, ensuring we pass and are not filtered out.
-			bool bFiltered = PrePreQueryFilterImp(QueryFilter, NewCollisionData.GetQueryData());
+			bool bFiltered = PrePreQueryFilterImp(QueryFilter, LegacyQueryFilter);
 			EXPECT_EQ(bFiltered, false);
-			PRAGMA_ENABLE_INTERNAL_WARNINGS
 		}
 
 		// Update particle in GT accel structure so cached PrePreFilter updates
diff --git a/Engine/Source/Programs/HeadlessChaos/Private/HeadlessChaosTestUtility.cpp b/Engine/Source/Programs/HeadlessChaos/Private/HeadlessChaosTestUtility.cpp
--- a/Engine/Source/Programs/HeadlessChaos/Private/HeadlessChaosTestUtility.cpp
+++ b/Engine/Source/Programs/HeadlessChaos/Private/HeadlessChaosTestUtility.cpp
@@ -669,11 +669,11 @@ namespace ChaosTest {
 
 	void SetShapeFilterToCollide(Chaos::FPerShapeData& Shape)
 	{
-		// TODO @ JoshD: Update to use new API once builder API is finalized.
-		FCollisionFilterData FilterData = Shape.GetSimData();
-		FilterData.Word1 = 1;
-		FilterData.Word3 = 1;
-		Shape.SetSimData(FilterData);
+		// Change the filter to block everything and to be marked as simple collision
+		PRAGMA_DISABLE_INTERNAL_WARNINGS
+		const Chaos::Filter::FShapeFilterData ShapeFilterData = Chaos::Filter::FShapeFilterBuilder::BuildBlockAll(Chaos::EFilterFlags::SimpleCollision);
+		Shape.SetShapeFilterData(ShapeFilterData);
+		PRAGMA_ENABLE_INTERNAL_WARNINGS
 	}
 
 	void SetParticleSimDataToCollide(TArray< Chaos::FGeometryParticle* > ParticleArray)
diff --git a/Engine/Source/Runtime/Engine/Private/Components/PrimitiveComponent.cpp b/Engine/Source/Runtime/Engine/Private/Components/PrimitiveComponent.cpp
--- a/Engine/Source/Runtime/Engine/Private/Components/PrimitiveComponent.cpp
+++ b/Engine/Source/Runtime/Engine/Private/Components/PrimitiveComponent.cpp
@@ -3540,8 +3540,7 @@ bool UPrimitiveComponent::LineTraceComponent(FHitResult& OutHit, const FVector S
 		{
 			bHaveHit = true;
 
-			FCollisionFilterData QueryFilter;
-			QueryFilter.Word1 = 0xFFFFF;
+			const FCollisionFilterData QueryFilter = Chaos::Filter::FQueryFilterBuilder::GetLegacyQueryFilter(Chaos::Filter::FQueryObjectFilterBuilder::BuildBlockAll());
 			ChaosInterface::SetFlags(BestHit, EHitFlags::Distance | EHitFlags::Normal | EHitFlags::Position);
 
 			ConvertQueryImpactHit(GetWorld(), BestHit, OutHit, (End - Start).Size(), QueryFilter, Start, End, nullptr, FTransform{ Start }, true, Params.bReturnPhysicalMaterial);
@@ -3605,8 +3604,7 @@ bool UPrimitiveComponent::SweepComponent(FHitResult& OutHit, const FVector Start
 	SweepParams.bComputeMTD = true;
 	if (CollisionInterface.ShapeSweep(Objects, Geometry, FTransform{ ShapeWorldRotation, Start }, End, SweepParams, BestHit))
 	{
-		FCollisionFilterData QueryFilter;
-		QueryFilter.Word1 = 0xFFFFF;
+		const FCollisionFilterData QueryFilter = Chaos::Filter::FQueryFilterBuilder::GetLegacyQueryFilter(Chaos::Filter::FQueryObjectFilterBuilder::BuildBlockAll());
 		ChaosInterface::SetFlags(BestHit, EHitFlags::Distance | EHitFlags::Normal | EHitFlags::Position | EHitFlags::FaceIndex);
 
 		bool bHasHit = false;
@@ -3714,8 +3712,7 @@ bool UPrimitiveComponent::ComponentOverlapComponentWithResultImpl(const class UP
 			{
 				TArray<FOverlapResult> Overlaps;
 
-				FCollisionFilterData QueryFilter;
-				QueryFilter.Word1 = 0xFFFFF;
+				const FCollisionFilterData QueryFilter = Chaos::Filter::FQueryFilterBuilder::GetLegacyQueryFilter(Chaos::Filter::FQueryObjectFilterBuilder::BuildBlockAll());
 				ConvertOverlapResults(OverlapHits.Num(), OverlapHits.GetData(), QueryFilter, Overlaps);
 
 				if (!Overlaps.IsEmpty())
@@ -3763,8 +3760,7 @@ bool UPrimitiveComponent::OverlapComponentWithResult(const FVector& Pos, const F
 	{
 		TArray<FOverlapResult> Overlaps;
 
-		FCollisionFilterData QueryFilter;
-		QueryFilter.Word1 = 0xFFFFF;
+		const FCollisionFilterData QueryFilter = Chaos::Filter::FQueryFilterBuilder::GetLegacyQueryFilter(Chaos::Filter::FQueryObjectFilterBuilder::BuildBlockAll());
 		ConvertOverlapResults(OverlapHits.Num(), OverlapHits.GetData(), QueryFilter, Overlaps);
 
 		if (!Overlaps.IsEmpty())
diff --git a/Engine/Source/Runtime/Engine/Private/PhysicsEngine/Experimental/PhysInterface_Chaos.cpp b/Engine/Source/Runtime/Engine/Private/PhysicsEngine/Experimental/PhysInterface_Chaos.cpp
--- a/Engine/Source/Runtime/Engine/Private/PhysicsEngine/Experimental/PhysInterface_Chaos.cpp
+++ b/Engine/Source/Runtime/Engine/Private/PhysicsEngine/Experimental/PhysInterface_Chaos.cpp
@@ -983,9 +983,8 @@ bool FPhysInterface_Chaos::LineTrace_Geom(FHitResult& OutHit, const FBodyInstanc
 						BestHit.WorldPosition = WorldTM.TransformPositionNoScale(BestHit.WorldPosition);
 						SetFlags(BestHit, EHitFlags::Distance | EHitFlags::Normal | EHitFlags::Position);
 
-						// we just like to make sure if the hit is made, set to test touch
-						FCollisionFilterData QueryFilter;
-						QueryFilter.Word2 = 0xFFFFF;
+						// Verify the hit is made by doing an overlap test.
+						const FCollisionFilterData QueryFilter = Chaos::Filter::FQueryFilterBuilder::GetLegacyQueryFilter(Chaos::Filter::FQueryTraceFilterBuilder::BuildOverlapAll());
 
 						FTransform StartTM(WorldStart);
 						const UPrimitiveComponent* OwnerComponentInst = InInstance->OwnerComponent.Get();
@@ -1082,8 +1081,7 @@ bool FPhysInterface_Chaos::Sweep_Geom(FHitResult& OutHit, const FBodyInstance* I
 						if (bShapeHit && FloatDistance < Hit.Distance)
 						{
 							// we just like to make sure if the hit is made
-							FCollisionFilterData QueryFilter;
-							QueryFilter.Word2 = 0xFFFFF;
+							const FCollisionFilterData QueryFilter = Chaos::Filter::FQueryFilterBuilder::GetLegacyQueryFilter(Chaos::Filter::FQueryTraceFilterBuilder::BuildOverlapAll());
 
 							// we don't get Shape information when we access via PShape, so I filled it up
 							Hit.Shape = Shape;
diff --git a/Engine/Source/Runtime/Engine/Private/PhysicsEngine/ImmediatePhysics/ImmediatePhysicsAdapters.cpp b/Engine/Source/Runtime/Engine/Private/PhysicsEngine/ImmediatePhysics/ImmediatePhysicsAdapters.cpp
--- a/Engine/Source/Runtime/Engine/Private/PhysicsEngine/ImmediatePhysics/ImmediatePhysicsAdapters.cpp
+++ b/Engine/Source/Runtime/Engine/Private/PhysicsEngine/ImmediatePhysics/ImmediatePhysicsAdapters.cpp
@@ -123,12 +123,11 @@ namespace ImmediatePhysics
 
 		// Set the filter to collide with everything (we use a broad phase that only contains particle pairs that are explicitly set to collide)
 		FBodyCollisionData BodyCollisionData;
-		// @todo(chaos): we need an API for setting up filters
-		BodyCollisionData.CollisionFilterData.SimFilter.Word1 = 0xFFFF;
-		BodyCollisionData.CollisionFilterData.SimFilter.Word3 = 0xFFFF;
-		// Ensure the word3s are equal (until collision filter refactor is complete)
-		BodyCollisionData.CollisionFilterData.QuerySimpleFilter.Word3 = 0xFFFF;
-		BodyCollisionData.CollisionFilterData.QueryComplexFilter.Word3 = 0xFFFF;
+
+		const Chaos::Filter::FShapeFilterData ShapeFilterData = Chaos::Filter::FShapeFilterBuilder::BuildSimBlockAll(EFilterFlags::All);
+		const Chaos::Filter::FCombinedShapeFilterData CombinedShapeFilter(ShapeFilterData, Chaos::Filter::FInstanceData());
+		Chaos::Filter::FShapeFilterBuilder::GetLegacyShapeFilter(CombinedShapeFilter, BodyCollisionData.CollisionFilterData.QuerySimpleFilter, BodyCollisionData.CollisionFilterData.SimFilter);
+		BodyCollisionData.CollisionFilterData.QueryComplexFilter = BodyCollisionData.CollisionFilterData.QuerySimpleFilter;
 		BodyCollisionData.CollisionFlags.bEnableSimCollisionSimple = true;
 
 		BodyInstance->BuildBodyCollisionFlags(BodyCollisionData.CollisionFlags, BodyInstance->GetCollisionEnabled(), BodyInstance->BodySetup->GetCollisionTraceFlag() == CTF_UseComplexAsSimple);
diff --git a/Engine/Source/Runtime/Engine/Private/PhysicsEngine/ImmediatePhysics/ImmediatePhysicsChaos/ImmediatePhysicsActorHandle_Chaos.cpp b/Engine/Source/Runtime/Engine/Private/PhysicsEngine/ImmediatePhysics/ImmediatePhysicsChaos/ImmediatePhysicsActorHandle_Chaos.cpp
--- a/Engine/Source/Runtime/Engine/Private/PhysicsEngine/ImmediatePhysics/ImmediatePhysicsChaos/ImmediatePhysicsActorHandle_Chaos.cpp
+++ b/Engine/Source/Runtime/Engine/Private/PhysicsEngine/ImmediatePhysics/ImmediatePhysicsChaos/ImmediatePhysicsActorHandle_Chaos.cpp
@@ -88,7 +88,7 @@ namespace ImmediatePhysics_Chaos
 			// Set the collision filter data for the shapes to collide with everything.
 			// Even though we already tried to do this when we created the original shapes array, 
 			// that gets thrown away and we need to do it here. This is not a good API
-			Chaos::Filter::FShapeFilterData ShapeFilter = Chaos::Filter::FShapeFilterBuilder::BuildLegacyBlockAllSimShapeFilter();
+			const Chaos::Filter::FShapeFilterData ShapeFilter = Chaos::Filter::FShapeFilterBuilder::BuildSimBlockAll(Chaos::EFilterFlags::All);
 
 			FCollisionData CollisionData;
 			CollisionData.SetShapeFilterData(ShapeFilter);
diff --git a/Engine/Source/Runtime/Experimental/Chaos/Private/Chaos/CollisionFilterData.cpp b/Engine/Source/Runtime/Experimental/Chaos/Private/Chaos/CollisionFilterData.cpp
--- a/Engine/Source/Runtime/Experimental/Chaos/Private/Chaos/CollisionFilterData.cpp
+++ b/Engine/Source/Runtime/Experimental/Chaos/Private/Chaos/CollisionFilterData.cpp
@@ -18,6 +18,9 @@ namespace Chaos::Filter
 	constexpr uint32 CollisionChannelMask = ((1 << NumCollisionChannelBits) - 1) << CollisionChannelBitsOffset;
 	constexpr uint32 MaskFilterMask = ((1 << NumMaskFilterBits) - 1) << MaskFilterBitsOffset;
 
+	// Note: The full mask is currently limited to 32 bits during the transition period
+	constexpr uint64 FullChannelMask = 0xFFFFFFFF;
+
 	constexpr uint32 LegacyQueryFlagsMask = (uint32)(EFilterFlags::SimpleCollision | EFilterFlags::ComplexCollision);
 	constexpr uint32 LegacySimFlagsMask = (uint32)(EFilterFlags::All) & ~LegacyQueryFlagsMask;
 
@@ -50,6 +53,13 @@ namespace Chaos::Filter
 		return (uint64)1 << (uint64)GetCollisionChannelIndex(Word3);
 	}
 
+	void SetCollisionChannel(uint32& Word3, uint32 Value)
+	{
+		const uint32 Word3Cleared = Word3 & ~CollisionChannelMask;
+		const uint32 NewCollisionChannelBits = CollisionChannelMask & (Value << CollisionChannelBitsOffset);
+		Word3 = Word3Cleared | NewCollisionChannelBits;
+	}
+
 	uint8 GetMaskFilter(uint32 Word3)
 	{
 		return Word3 >> (32u - NumMaskFilterBits);
@@ -259,6 +269,73 @@ namespace Chaos::Filter
 		return SimFilter.Word0 != 0 || SimFilter.Word1 != 0 || SimFilter.Word2 != 0 || SimFilter.Word3 != 0;
 	}
 
+	FShapeFilterBuilder::FShapeFilterBuilder(const FShapeFilterData& InFilterData)
+		: ShapeFilterData(InFilterData)
+	{
+	}
+
+	FShapeFilterBuilder& FShapeFilterBuilder::SetCollisionChannelIndex(const uint8 ChannelIndex)
+	{
+		SetCollisionChannel(ShapeFilterData.Word3, (uint32)ChannelIndex);
+		ShapeFilterData.SimWord3 = ShapeFilterData.Word3;
+		return *this;
+	}
+
+	FShapeFilterBuilder& FShapeFilterBuilder::SetBlockChannelMask(const uint64 BlockMask)
+	{
+		ShapeFilterData.SimBlockChannels = (uint32)BlockMask;
+		ShapeFilterData.QueryBlockChannels = (uint32)BlockMask;
+		return *this;
+	}
+
+	FShapeFilterBuilder& FShapeFilterBuilder::SetBlockChannelMasks(const uint64 SimBlockMask, const uint64 QueryBlockMask)
+	{
+		ShapeFilterData.SimBlockChannels = (uint32)SimBlockMask;
+		ShapeFilterData.QueryBlockChannels = (uint32)QueryBlockMask;
+		return *this;
+	}
+
+	FShapeFilterBuilder& FShapeFilterBuilder::SetOverlapChannelMask(const uint64 OverlapMask)
+	{
+		ShapeFilterData.QueryOverlapChannels = (uint32)OverlapMask;
+		return *this;
+	}
+
+	FShapeFilterBuilder& FShapeFilterBuilder::SetMaskFilter(const uint8 InMaskFilter = 0)
+	{
+		ShapeFilterData.SetMaskFilter(InMaskFilter);
+		ShapeFilterData.SimWord3 = ShapeFilterData.Word3;
+		return *this;
+	}
+
+	FShapeFilterBuilder& FShapeFilterBuilder::SetFilterFlags(const EFilterFlags InFilterFlags)
+	{
+		ShapeFilterData.SetFlags(InFilterFlags);
+		ShapeFilterData.SimWord3 = ShapeFilterData.Word3;
+		return *this;
+	}
+
+	FShapeFilterData FShapeFilterBuilder::Build() const
+	{
+		return ShapeFilterData;
+	}
+
+	FShapeFilterData FShapeFilterBuilder::BuildBlockAll(const EFilterFlags InFilterFlags)
+	{
+		FShapeFilterBuilder Builder;
+		Builder.SetBlockChannelMask(FullChannelMask);
+		Builder.SetFilterFlags(InFilterFlags);
+		return Builder.Build();
+	}
+
+	FShapeFilterData FShapeFilterBuilder::BuildSimBlockAll(const EFilterFlags InFilterFlags)
+	{
+		FShapeFilterBuilder Builder;
+		Builder.SetBlockChannelMasks(FullChannelMask, 0);
+		Builder.SetFilterFlags(InFilterFlags);
+		return Builder.Build();
+	}
+
 	FShapeFilterData FShapeFilterBuilder::BuildLegacyShapeFilter(const uint8 InChannelIndex, const uint32 InQueryOverlapChannels, const uint32 InQueryBlockChannels, const uint32 InSimBlockChannels, const uint8 InMaskFilter, const EFilterFlags InFilterFlags)
 	{
 		// (Mask Filter 6 bits) | (Channel Index 5 bits) | (Flags 21 bits)
@@ -266,16 +343,16 @@ namespace Chaos::Filter
 		Result.QueryBlockChannels = InQueryBlockChannels;
 		Result.QueryOverlapChannels = InQueryOverlapChannels;
 		Result.SimBlockChannels = InSimBlockChannels;
-		Result.Word3 = FilterFlagsMask & (uint32)InFilterFlags;
-		Result.Word3 |= CollisionChannelMask & ((uint32)InChannelIndex << CollisionChannelBitsOffset);
-		Result.Word3 |= MaskFilterMask & ((uint32)InMaskFilter << MaskFilterBitsOffset);
+		Filter::SetFlags(Result.Word3, (uint32)InFilterFlags);
+		Filter::SetCollisionChannel(Result.Word3, (uint32)InChannelIndex);
+		Filter::SetMaskFilter(Result.Word3, InMaskFilter);
 		Result.SimWord3 = Result.Word3;
 		return Result;
 	}
 
 	FShapeFilterData FShapeFilterBuilder::BuildLegacyBlockAllSimShapeFilter(const uint8 InChannelIndex, const uint8 InMaskFilter, const EFilterFlags InFilterFlags)
 	{
-		return BuildLegacyShapeFilter(InChannelIndex, 0, 0, (uint32)-1, 0, InFilterFlags);
+		return BuildLegacyShapeFilter(InChannelIndex, 0, 0, (uint32)-1, InMaskFilter, InFilterFlags);
 	}
 
 	FShapeFilterData FShapeFilterBuilder::BuildFromLegacyQueryFilter(const FCollisionFilterData& QueryFilter)
@@ -455,4 +532,100 @@ namespace Chaos::Filter
 		Result.Word3 = QueryFilterData.Word3;
 		return Result;
 	}
+
+	FQueryObjectFilterBuilder::FQueryObjectFilterBuilder()
+	{
+		QueryFilterData.Word0 = 0;
+	}
+
+	FQueryObjectFilterBuilder::FQueryObjectFilterBuilder(const FQueryFilterData& InFilter)
+		: QueryFilterData(InFilter)
+	{
+		QueryFilterData.Word0 = 0;
+	}
+
+	FQueryObjectFilterBuilder& FQueryObjectFilterBuilder::SetObjectTypes(uint64 ObjectTypes)
+	{
+		QueryFilterData.Word1 = (uint32)ObjectTypes;
+		return *this;
+	}
+
+	FQueryObjectFilterBuilder& FQueryObjectFilterBuilder::SetMultiQuery(bool bMultiQuery)
+	{
+		const uint32 MultiQueryValue = bMultiQuery ? 0x1 : 0x0;
+		SetCollisionChannel(QueryFilterData.Word3, MultiQueryValue);
+		return *this;
+	}
+
+	FQueryObjectFilterBuilder& FQueryObjectFilterBuilder::SetMaskFilter(const uint8 InMaskFilter)
+	{
+		Filter::SetMaskFilter(QueryFilterData.Word3, InMaskFilter);
+		return *this;
+	}
+
+	FQueryFilterData FQueryObjectFilterBuilder::Build() const
+	{
+		return QueryFilterData;
+	}
+
+	FQueryFilterData FQueryObjectFilterBuilder::BuildBlockAll()
+	{
+		FQueryObjectFilterBuilder Builder;
+		Builder.SetObjectTypes(FullChannelMask);
+		return Builder.Build();
+	}
+
+	FQueryTraceFilterBuilder::FQueryTraceFilterBuilder()
+	{
+		QueryFilterData.Word0 = 1;
+	}
+
+	FQueryTraceFilterBuilder::FQueryTraceFilterBuilder(const FQueryFilterData& InFilter)
+		: QueryFilterData(InFilter)
+	{
+		QueryFilterData.Word0 = 1;
+	}
+
+	FQueryTraceFilterBuilder& FQueryTraceFilterBuilder::SetCollisionChannelIndex(const uint8 CollisionChannelIndex)
+	{
+		SetCollisionChannel(QueryFilterData.Word3, (uint32)CollisionChannelIndex);
+		return *this;
+	}
+
+	FQueryTraceFilterBuilder& FQueryTraceFilterBuilder::SetBlockChannelMask(const uint64 BlockChannelMask)
+	{
+		QueryFilterData.Word1 = (uint32)BlockChannelMask;
+		return *this;
+	}
+
+	FQueryTraceFilterBuilder& FQueryTraceFilterBuilder::SetOverlapChannelMask(const uint64 OverlapChannelMask)
+	{
+		QueryFilterData.Word2 = (uint32)OverlapChannelMask;
+		return *this;
+	}
+
+	FQueryTraceFilterBuilder& FQueryTraceFilterBuilder::SetMaskFilter(const uint8 InMaskFilter)
+	{
+		Filter::SetMaskFilter(QueryFilterData.Word3, InMaskFilter);
+		return *this;
+	}
+
+	FQueryTraceFilterBuilder& FQueryTraceFilterBuilder::SetFilterFlags(const EFilterFlags InFilterFlags)
+	{
+		QueryFilterData.SetFlags(InFilterFlags);
+		return *this;
+	}
+
+	FQueryFilterData FQueryTraceFilterBuilder::Build() const
+	{
+		return QueryFilterData;
+	}
+
+	FQueryFilterData FQueryTraceFilterBuilder::BuildOverlapAll(const EFilterFlags InFilterFlags)
+	{
+		Chaos::Filter::FQueryTraceFilterBuilder Builder;
+		Builder.SetOverlapChannelMask(FullChannelMask);
+		Builder.SetFilterFlags(InFilterFlags);
+		return Builder.Build();
+	}
 } // namespace Chaos::Filter
diff --git a/Engine/Source/Runtime/Experimental/Chaos/Private/PhysicsProxy/GeometryCollectionPhysicsProxy.cpp b/Engine/Source/Runtime/Experimental/Chaos/Private/PhysicsProxy/GeometryCollectionPhysicsProxy.cpp
--- a/Engine/Source/Runtime/Experimental/Chaos/Private/PhysicsProxy/GeometryCollectionPhysicsProxy.cpp
+++ b/Engine/Source/Runtime/Experimental/Chaos/Private/PhysicsProxy/GeometryCollectionPhysicsProxy.cpp
@@ -378,18 +378,14 @@ void SetImplicitToPTParticles(Chaos::TPBDRigidParticleHandle<Chaos::FReal, 3>* H
 	{
 		// Override collision filters and make this body collide with everything.
 		int32 CurrShape = 0;
-		FCollisionFilterData SimFilterData;
-		SimFilterData.Word1 = 0xFFFF; // this body channel
-		SimFilterData.Word3 = 0xFFFF; // collision candidate channels
-		FCollisionFilterData QueryFilterData;
-		QueryFilterData.Word3 = SimFilterData.Word3;
-		const Chaos::Filter::FCombinedShapeFilterData CombinedFilterData = Chaos::Filter::FShapeFilterBuilder::BuildFromLegacyShapeFilter(QueryFilterData, SimFilterData);
+		
+		const Chaos::Filter::FShapeFilterData ShapeFilterData = Chaos::Filter::FShapeFilterBuilder::BuildSimBlockAll(Chaos::EFilterFlags::All);
 		for (const TUniquePtr<Chaos::FPerShapeData>& Shape : Handle->ShapesArray())
 		{
 			Shape->SetSimEnabled(true);
 			Shape->SetCollisionTraceType(Chaos::EChaosCollisionTraceFlag::Chaos_CTF_UseDefault);
 			//Shape->CollisionTraceType = Chaos::EChaosCollisionTraceFlag::Chaos_CTF_UseSimpleAndComplex;
-			Shape->SetCombinedShapeFilterData(CombinedFilterData);
+			Shape->SetShapeFilterData(ShapeFilterData);
 		}
 	}
 	else
diff --git a/Engine/Source/Runtime/Experimental/Chaos/Private/PhysicsProxy/SingleParticlePhysicsProxy.cpp b/Engine/Source/Runtime/Experimental/Chaos/Private/PhysicsProxy/SingleParticlePhysicsProxy.cpp
--- a/Engine/Source/Runtime/Experimental/Chaos/Private/PhysicsProxy/SingleParticlePhysicsProxy.cpp
+++ b/Engine/Source/Runtime/Experimental/Chaos/Private/PhysicsProxy/SingleParticlePhysicsProxy.cpp
@@ -545,4 +545,17 @@ void FSingleParticlePhysicsProxy::ClearEvents()
 		Rigid->ClearEvents();
 	}
 }
+
+void FRigidBodyHandle_External::SetShapeFilterData(int32 InShapeIndex, const Filter::FShapeFilterData& FilterData)
+{
+	VerifyContext();
+	GetParticle_LowLevel()->SetShapeFilterData(InShapeIndex, FilterData);
+}
+
+void FRigidBodyHandle_External::SetShapeFilterInstanceData(int32 InShapeIndex, const Filter::FInstanceData& InstanceData)
+{
+	VerifyContext();
+	GetParticle_LowLevel()->SetShapeFilterInstanceData(InShapeIndex, InstanceData);
+}
+
 }
diff --git a/Engine/Source/Runtime/Experimental/Chaos/Public/Chaos/CollisionFilterData.h b/Engine/Source/Runtime/Experimental/Chaos/Public/Chaos/CollisionFilterData.h
--- a/Engine/Source/Runtime/Experimental/Chaos/Public/Chaos/CollisionFilterData.h
+++ b/Engine/Source/Runtime/Experimental/Chaos/Public/Chaos/CollisionFilterData.h
@@ -174,6 +174,8 @@ namespace Chaos::Filter
 
 	private:
 		friend class FQueryFilterBuilder;
+		friend class FQueryObjectFilterBuilder;
+		friend class FQueryTraceFilterBuilder;
 
 		void SetMaskChannelAndFlags(const uint8 IgnoreMask, const uint8 ChannelIndex, const EFilterFlags FilterFlags);
 
@@ -209,6 +211,21 @@ namespace Chaos::Filter
 	class UE_INTERNAL FShapeFilterBuilder
 	{
 	public:
+		FShapeFilterBuilder() = default;
+		CHAOS_API FShapeFilterBuilder(const FShapeFilterData& InFilterData);
+
+		// TODO @ JoshD: Eventually deprecated and replace with SetChannelMask
+		CHAOS_API FShapeFilterBuilder& SetCollisionChannelIndex(const uint8 ChannelIndex);
+		CHAOS_API FShapeFilterBuilder& SetBlockChannelMask(const uint64 BlockMask);
+		UE_INTERNAL CHAOS_API FShapeFilterBuilder& SetBlockChannelMasks(const uint64 SimBlockMask, const uint64 QueryBlockMask);
+		CHAOS_API FShapeFilterBuilder& SetOverlapChannelMask(const uint64 OverlapMask);
+		CHAOS_API FShapeFilterBuilder& SetMaskFilter(const uint8 InMaskFilter);
+		CHAOS_API FShapeFilterBuilder& SetFilterFlags(const EFilterFlags InFilterFlags);
+
+		CHAOS_API FShapeFilterData Build() const;
+		CHAOS_API static FShapeFilterData BuildBlockAll(const EFilterFlags InFilterFlags = EFilterFlags::None);
+		CHAOS_API static FShapeFilterData BuildSimBlockAll(const EFilterFlags InFilterFlags = EFilterFlags::None);
+
 		UE_INTERNAL CHAOS_API static FShapeFilterData BuildLegacyShapeFilter(const uint8 InChannelIndex, const uint32 InQueryOverlapChannels, const uint32 InQueryBlockChannels, const uint32 InSimBlockChannels, const uint8 InMaskFilter = 0, const EFilterFlags InFilterFlags = EFilterFlags::None);
 		UE_INTERNAL CHAOS_API static FShapeFilterData BuildLegacyBlockAllSimShapeFilter(const uint8 InChannelIndex = 0, const uint8 InMaskFilter = 0, const EFilterFlags InFilterFlags = EFilterFlags::All);
 
@@ -225,8 +242,11 @@ namespace Chaos::Filter
 
 		UE_INTERNAL static uint32 GetWord3FromNewSimFilter(const uint32 SimWord3, const uint32 QueryWord3);
 		UE_INTERNAL static uint32 GetWord3FromNewQueryFilter(const uint32 QueryWord3, const uint32 SimWord3);
+
 	private:
 		static void BuildFromLegacyShapeFilter(const FCollisionFilterData& QueryFilterData, const FCollisionFilterData& SimFilterData, const uint32 Word3, FCombinedShapeFilterData& OutCombinedShapeFilterData);
+
+		FShapeFilterData ShapeFilterData;
 	};
 
 	class UE_INTERNAL FQueryFilterBuilder
@@ -239,4 +259,41 @@ namespace Chaos::Filter
 		UE_INTERNAL CHAOS_API static FCollisionFilterData GetLegacyQueryFilter(const FQueryFilterData& QueryFilterData);
 	private:
 	};
+
+	class UE_INTERNAL FQueryObjectFilterBuilder
+	{
+	public:
+		CHAOS_API FQueryObjectFilterBuilder();
+		CHAOS_API FQueryObjectFilterBuilder(const FQueryFilterData& InFilter);
+
+		CHAOS_API FQueryObjectFilterBuilder& SetObjectTypes(uint64 ObjectTypes);
+		CHAOS_API FQueryObjectFilterBuilder& SetMultiQuery(bool bMultiQuery);
+		CHAOS_API FQueryObjectFilterBuilder& SetMaskFilter(const uint8 InMaskFilter);
+
+		CHAOS_API FQueryFilterData Build() const;
+		CHAOS_API static FQueryFilterData BuildBlockAll();
+
+	private:
+		FQueryFilterData QueryFilterData;
+	};
+
+	class UE_INTERNAL FQueryTraceFilterBuilder
+	{
+	public:
+		CHAOS_API FQueryTraceFilterBuilder();
+		CHAOS_API FQueryTraceFilterBuilder(const FQueryFilterData& InFilter);
+
+		// TODO @ JoshD: Eventually deprecated and replace with SetChannelMask
+		CHAOS_API FQueryTraceFilterBuilder& SetCollisionChannelIndex(const uint8 CollisionChannelIndex);
+		CHAOS_API FQueryTraceFilterBuilder& SetBlockChannelMask(const uint64 BlockChannelMask);
+		CHAOS_API FQueryTraceFilterBuilder& SetOverlapChannelMask(const uint64 OverlapChannelMask);
+		CHAOS_API FQueryTraceFilterBuilder& SetMaskFilter(const uint8 InMaskFilter);
+		CHAOS_API FQueryTraceFilterBuilder& SetFilterFlags(const EFilterFlags InFilterFlags);
+
+		CHAOS_API FQueryFilterData Build() const;
+		CHAOS_API static FQueryFilterData BuildOverlapAll(const EFilterFlags InFilterFlags = EFilterFlags::None);
+
+	private:
+		FQueryFilterData QueryFilterData;
+	};
 } // namespace Chaos::Filter
diff --git a/Engine/Source/Runtime/Experimental/Chaos/Public/Chaos/ParticleHandle.h b/Engine/Source/Runtime/Experimental/Chaos/Public/Chaos/ParticleHandle.h
--- a/Engine/Source/Runtime/Experimental/Chaos/Public/Chaos/ParticleHandle.h
+++ b/Engine/Source/Runtime/Experimental/Chaos/Public/Chaos/ParticleHandle.h
@@ -2916,6 +2916,24 @@ public:
 		}
 	}
 
+	void SetShapeFilterData(int32 InShapeIndex, const Chaos::Filter::FShapeFilterData& FilterData)
+	{
+		const Chaos::Filter::FShapeFilterData& Current = MShapesArray[InShapeIndex]->GetShapeFilterData();
+		if (Current != FilterData)
+		{
+			MShapesArray[InShapeIndex]->SetShapeFilterData(FilterData);
+		}
+	}
+
+	void SetShapeFilterInstanceData(int32 InShapeIndex, const Chaos::Filter::FInstanceData& InstanceData)
+	{
+		const Chaos::Filter::FInstanceData& Current = MShapesArray[InShapeIndex]->GetFilterInstanceData();
+		if (Current != InstanceData)
+		{
+			MShapesArray[InShapeIndex]->SetFilterInstanceData(InstanceData);
+		}
+	}
+
 #if CHAOS_DEBUG_NAME
 	const TSharedPtr<FString, ESPMode::ThreadSafe>& DebugName() const { return MNonFrequentData.Read().DebugName(); }
 	void SetDebugName(const TSharedPtr<FString, ESPMode::ThreadSafe>& InDebugName)
diff --git a/Engine/Source/Runtime/Experimental/Chaos/Public/PhysicsProxy/SingleParticlePhysicsProxy.h b/Engine/Source/Runtime/Experimental/Chaos/Public/PhysicsProxy/SingleParticlePhysicsProxy.h
--- a/Engine/Source/Runtime/Experimental/Chaos/Public/PhysicsProxy/SingleParticlePhysicsProxy.h
+++ b/Engine/Source/Runtime/Experimental/Chaos/Public/PhysicsProxy/SingleParticlePhysicsProxy.h
@@ -1168,6 +1168,8 @@ public:
 	void SetShapeSimCollisionEnabled(int32 InShapeIndex, bool bInEnabled) { VerifyContext(); GetParticle_LowLevel()->SetShapeSimCollisionEnabled(InShapeIndex, bInEnabled); }
 	void SetShapeQueryCollisionEnabled(int32 InShapeIndex, bool bInEnabled) { VerifyContext(); GetParticle_LowLevel()->SetShapeQueryCollisionEnabled(InShapeIndex, bInEnabled); }
 	void SetShapeSimData(int32 InShapeIndex, const FCollisionFilterData& SimData) { VerifyContext(); GetParticle_LowLevel()->SetShapeSimData(InShapeIndex, SimData); }
+	CHAOS_API void SetShapeFilterData(int32 InShapeIndex, const Filter::FShapeFilterData& FilterData);
+	CHAOS_API void SetShapeFilterInstanceData(int32 InShapeIndex, const Filter::FInstanceData& InstanceData);
 
 	void SetParticleID(const FParticleID& ParticleID)
 	{
diff --git a/Engine/Source/Runtime/Experimental/ChaosSolverEngine/Private/Chaos/ChaosSolverActor.cpp b/Engine/Source/Runtime/Experimental/ChaosSolverEngine/Private/Chaos/ChaosSolverActor.cpp
--- a/Engine/Source/Runtime/Experimental/ChaosSolverEngine/Private/Chaos/ChaosSolverActor.cpp
+++ b/Engine/Source/Runtime/Experimental/ChaosSolverEngine/Private/Chaos/ChaosSolverActor.cpp
@@ -457,10 +457,9 @@ void AChaosSolverActor::MakeFloor()
 		const FVector BoxMax(100000, 100000, 0);
 		FloorParticle->SetGeometry(MakeImplicitObjectPtr<Chaos::TBox<Chaos::FReal, 3>>(BoxMin, BoxMax));
 		FloorParticle->SetX(Chaos::FVec3(SolverLocation.X, SolverLocation.Y, FloorHeight));
-		FCollisionFilterData FilterData;
-		FilterData.Word1 = 0xFFFF;
-		FilterData.Word3 = 0xFFFF;
-		FloorParticle->SetShapeSimData(0, FilterData);
+
+		const Chaos::Filter::FShapeFilterData ShapeFilterData = Chaos::Filter::FShapeFilterBuilder::BuildSimBlockAll(Chaos::EFilterFlags::All);
+		FloorParticle->SetShapeFilterData(0, ShapeFilterData);
 		Proxy = Chaos::FSingleParticlePhysicsProxy::Create(MoveTemp(FloorParticle));
 		RigidSolverProxy.Solver->RegisterObject(Proxy);
 	}
