Collision Filters: Small fix to rename ActorId to OwnerId to pre-emptively unify the API.

diff --git a/Engine/Plugins/Experimental/ChaosMover/Source/ChaosMover/Private/ChaosMoverSimulation.cpp b/Engine/Plugins/Experimental/ChaosMover/Source/ChaosMover/Private/ChaosMoverSimulation.cpp
--- a/Engine/Plugins/Experimental/ChaosMover/Source/ChaosMover/Private/ChaosMoverSimulation.cpp
+++ b/Engine/Plugins/Experimental/ChaosMover/Source/ChaosMover/Private/ChaosMoverSimulation.cpp
@@ -279,7 +279,7 @@ void UChaosMoverSimulation::ModifyContacts(const FMoverTimeStep& TimeStep, const
 							continue;
 						}
 
-						uint32 ActorID = InstanceData.GetActorId();
+						uint32 ActorID = InstanceData.GetOwnerId();
 						if (SimInputs->CollisionQueryParams.GetIgnoredSourceObjects().Contains(ActorID))
 						{
 							PairModifier.Disable();
diff --git a/Engine/Source/Runtime/Engine/Private/Collision/SceneQuery.cpp b/Engine/Source/Runtime/Engine/Private/Collision/SceneQuery.cpp
--- a/Engine/Source/Runtime/Engine/Private/Collision/SceneQuery.cpp
+++ b/Engine/Source/Runtime/Engine/Private/Collision/SceneQuery.cpp
@@ -706,7 +706,7 @@ bool TSceneCastCommonImpWithRetryRequest(const UWorld* World, typename Traits::O
 								OutRetryParams = Params;
 								// Ignore the actor that was hit (Check the shape data to be sure)
 								const Chaos::Filter::FInstanceData& FilterInstanceData = Traits::GetHits(HitBufferSync)->Shape->GetFilterInstanceData();
-								uint32 ActorIDFromShape = FilterInstanceData.GetActorId();
+								uint32 ActorIDFromShape = FilterInstanceData.GetOwnerId();
 								OutRetryParams.AddIgnoredActor(ActorIDFromShape);
 								if (ActorIDFromShape != ClusterUnionActor->GetUniqueID())
 								{
diff --git a/Engine/Source/Runtime/Engine/Private/Collision/WorldCollision.cpp b/Engine/Source/Runtime/Engine/Private/Collision/WorldCollision.cpp
--- a/Engine/Source/Runtime/Engine/Private/Collision/WorldCollision.cpp
+++ b/Engine/Source/Runtime/Engine/Private/Collision/WorldCollision.cpp
@@ -521,7 +521,7 @@ bool UWorld::ComponentSweepMultiByChannel(TArray<struct FHitResult>& OutHits, cl
 		for (Chaos::FShapeInstanceProxy* Shape : Shapes)
 		{
 			const Chaos::Filter::FInstanceData InstanceData = ChaosInterface::GetFilterInstanceData(*Shape);
-			uint32 ShapeActorID = InstanceData.GetActorId();
+			uint32 ShapeActorID = InstanceData.GetOwnerId();
 			ActorsToExclude.Add(ShapeActorID);
 		}
 
diff --git a/Engine/Source/Runtime/Experimental/Chaos/Private/Chaos/CollisionFilterData.cpp b/Engine/Source/Runtime/Experimental/Chaos/Private/Chaos/CollisionFilterData.cpp
--- a/Engine/Source/Runtime/Experimental/Chaos/Private/Chaos/CollisionFilterData.cpp
+++ b/Engine/Source/Runtime/Experimental/Chaos/Private/Chaos/CollisionFilterData.cpp
@@ -78,25 +78,35 @@ namespace Chaos::Filter
 		Word3 |= MaskFilter << MaskFilterBitsOffset;
 	}
 
-	FInstanceData::FInstanceData(const uint32 ActorId, const uint32 ComponentId)
-		: ActorId(ActorId)
+	FInstanceData::FInstanceData(const uint32 OwnerId, const uint32 ComponentId)
+		: OwnerId(OwnerId)
 		, ComponentId(ComponentId)
 	{
 	}
 
 	bool FInstanceData::IsValid() const
 	{
-		return ActorId != 0 || ComponentId != 0 || BodyIndex != 0;
+		return OwnerId != 0 || ComponentId != 0 || BodyIndex != 0;
 	}
 
 	uint32 FInstanceData::GetActorId() const
 	{
-		return ActorId;
+		return GetOwnerId();
 	}
 
 	void FInstanceData::SetActorId(const uint32 InActorId)
 	{
-		ActorId = InActorId;
+		SetOwnerId(InActorId);
+	}
+	
+	uint32 FInstanceData::GetOwnerId() const
+	{
+		return OwnerId;
+	}
+
+	void FInstanceData::SetOwnerId(const uint32 InOwnerId)
+	{
+		OwnerId = InOwnerId;
 	}
 
 	uint32 FInstanceData::GetComponentId() const
@@ -393,7 +403,7 @@ namespace Chaos::Filter
 
 
 		FCombinedShapeFilterData CombinedShapeFilterData;
-		CombinedShapeFilterData.InstanceData.ActorId = QueryFilterData.Word0;
+		CombinedShapeFilterData.InstanceData.OwnerId = QueryFilterData.Word0;
 		CombinedShapeFilterData.InstanceData.ComponentId = SimFilterData.Word2;
 		CombinedShapeFilterData.InstanceData.BodyIndex = SimFilterData.Word0;
 
@@ -428,7 +438,7 @@ namespace Chaos::Filter
 		// We expect the word3's to always be equal in the transition period before removal of SimWord3.
 		ensure(ShapeFilterData.Word3 == ShapeFilterData.SimWord3);
 
-		OutQueryFilterData.Word0 = CombinedShapeFilterData.InstanceData.ActorId;
+		OutQueryFilterData.Word0 = CombinedShapeFilterData.InstanceData.OwnerId;
 		OutSimFilterData.Word2 = CombinedShapeFilterData.InstanceData.ComponentId;
 		OutSimFilterData.Word0 = CombinedShapeFilterData.InstanceData.BodyIndex;
 
@@ -481,7 +491,7 @@ namespace Chaos::Filter
 
 	void FShapeFilterBuilder::BuildFromLegacyShapeFilter(const FCollisionFilterData& QueryFilterData, const FCollisionFilterData& SimFilterData, const uint32 Word3, FCombinedShapeFilterData& OutCombinedShapeFilterData)
 	{
-		OutCombinedShapeFilterData.InstanceData.ActorId = QueryFilterData.Word0;
+		OutCombinedShapeFilterData.InstanceData.OwnerId = QueryFilterData.Word0;
 		OutCombinedShapeFilterData.InstanceData.ComponentId = SimFilterData.Word2;
 		OutCombinedShapeFilterData.InstanceData.BodyIndex = SimFilterData.Word0;
 
diff --git a/Engine/Source/Runtime/Experimental/Chaos/Public/Chaos/CollisionFilterData.h b/Engine/Source/Runtime/Experimental/Chaos/Public/Chaos/CollisionFilterData.h
--- a/Engine/Source/Runtime/Experimental/Chaos/Public/Chaos/CollisionFilterData.h
+++ b/Engine/Source/Runtime/Experimental/Chaos/Public/Chaos/CollisionFilterData.h
@@ -80,15 +80,17 @@ namespace Chaos::Filter
 	struct FInstanceData
 	{
 		FInstanceData() = default;
-		CHAOS_API FInstanceData(const uint32 ActorId, const uint32 ComponentId);
+		CHAOS_API FInstanceData(const uint32 InOwnerId, const uint32 ComponentId);
 
 		bool operator==(const FInstanceData&) const = default;
 		bool operator!=(const FInstanceData&) const = default;
 
 		CHAOS_API bool IsValid() const;
 
-		CHAOS_API uint32 GetActorId() const;
-		CHAOS_API void SetActorId(const uint32 InActorId);
+		UE_DEPRECATED(5.8, "Use GetOwnerId instead") CHAOS_API uint32 GetActorId() const;
+		UE_DEPRECATED(5.8, "Use SetOwnerId instead") CHAOS_API void SetActorId(const uint32 InActorId);
+		CHAOS_API uint32 GetOwnerId() const;
+		CHAOS_API void SetOwnerId(const uint32 InOwnerId);
 
 		UE_INTERNAL CHAOS_API uint32 GetComponentId() const;
 		UE_INTERNAL CHAOS_API void SetComponentId(const uint32 InComponentId);
@@ -97,7 +99,7 @@ namespace Chaos::Filter
 		friend class FShapeFilterBuilder;
 		friend FCombinedShapeFilterData;
 
-		uint32 ActorId = 0;
+		uint32 OwnerId = 0;
 		uint32 ComponentId = 0;
 		// TODO @ JoshD: Remove if possible
 		uint32 BodyIndex = 0;
