Dstm: Fix crash when dependent primitive component is migrated before the constraint component does.

diff --git a/Engine/Source/Runtime/Experimental/Chaos/Private/ChaosVisualDebugger/ChaosVDDataWrapperUtils.cpp b/Engine/Source/Runtime/Experimental/Chaos/Private/ChaosVisualDebugger/ChaosVDDataWrapperUtils.cpp
--- a/Engine/Source/Runtime/Experimental/Chaos/Private/ChaosVisualDebugger/ChaosVDDataWrapperUtils.cpp
+++ b/Engine/Source/Runtime/Experimental/Chaos/Private/ChaosVisualDebugger/ChaosVDDataWrapperUtils.cpp
@@ -370,8 +370,14 @@ FChaosVDJointConstraint FChaosVDDataWrapperUtils::BuildJointDataWrapper(const Ch
 	{
 		WrappedJointData.ConstraintIndex = ConstaintHanlde->GetConstraintIndex();
 
-		WrappedJointData.ParticleParIndexes[0] = ConstaintHanlde->GetConstrainedParticles()[0]->UniqueIdx().Idx;
-		WrappedJointData.ParticleParIndexes[1] = ConstaintHanlde->GetConstrainedParticles()[1]->UniqueIdx().Idx;
+		if (Chaos::FGeometryParticleHandle* Particle0 = ConstaintHanlde->GetConstrainedParticles()[0])
+		{
+			WrappedJointData.ParticleParIndexes[0] = Particle0->UniqueIdx().Idx;
+		}
+		if (Chaos::FGeometryParticleHandle* Particle1 = ConstaintHanlde->GetConstrainedParticles()[1])
+		{
+			WrappedJointData.ParticleParIndexes[1] = Particle1->UniqueIdx().Idx;
+		}
 
 		WrappedJointData.PhysicsThreadJointState.bBroken = ConstaintHanlde->IsConstraintBroken();
 		WrappedJointData.PhysicsThreadJointState.bBreaking = ConstaintHanlde->IsConstraintBreaking();
