Chaos - prevent duplicate mesh edge collisions

Triangle edge contacts are added twice to FMeshContactGenerator (once from each triangle sharing the edge) but they were not being filtered. This leads to unnecessary contacts with the same position, normal and depth which is wasteful and makes it hard for the solver to track friction.

#rnx

diff --git a/Engine/Source/Runtime/Experimental/Chaos/Private/Chaos/Collision/MeshContactGenerator.cpp b/Engine/Source/Runtime/Experimental/Chaos/Private/Chaos/Collision/MeshContactGenerator.cpp
--- a/Engine/Source/Runtime/Experimental/Chaos/Private/Chaos/Collision/MeshContactGenerator.cpp
+++ b/Engine/Source/Runtime/Experimental/Chaos/Private/Chaos/Collision/MeshContactGenerator.cpp
@@ -38,7 +38,8 @@ namespace Chaos::Private
 	FMeshContactGenerator::FMeshContactGenerator(const FMeshContactGeneratorSettings& InSettings)
 		: Settings(InSettings)
 		, EdgeTriangleIndicesMap(InSettings.HashSize)
-		, VertexContactIndicesMap(InSettings.HashSize)
+		, VertexContactDataMap(InSettings.HashSize)
+		, EdgeContactDataMap(InSettings.HashSize)
 	{
 	}
 
@@ -53,8 +54,9 @@ namespace Chaos::Private
 		// but not all are shared so lets go with NumEdges = (2 * NumTriangles)
 		EdgeTriangleIndicesMap.Reset(InMaxTriangles * 2);
 
-		// Worst case - assume all contacts are vertex contacts
-		VertexContactIndicesMap.Reset(InMaxContacts);
+		// Worst case - assume all contacts are vertex contacts or edge contacts
+		VertexContactDataMap.Reset(InMaxContacts);
+		EdgeContactDataMap.Reset(InMaxContacts);
 	}
 
 	void FMeshContactGenerator::AddTriangleContacts(const int32 LocalTriangleIndex, const TArrayView<FContactPoint>& TriangleContactPoints)
@@ -93,12 +95,69 @@ namespace Chaos::Private
 			// Register vertex and edge collisions
 			if ((VertexID0 != INDEX_NONE) && (VertexID1 != INDEX_NONE))
 			{
-				// Edge collisions - if it's a face normal, tell each triangle using the edge it has a face contact
+				// Edge collisions - keep the contacts from only one of the triangles sharing the edge.
+				// NOTES:
+				//  - for convex-vs-triangle collisions we can have up to 2 collisions per triangle edge.
+				//	- each triangle sharing the edge should report the same 0, 1, or 2 collisions which
+				//	  means we can ignore contacts from one of the triangles.
+				//	- except that convex-triangle manifold creation limits contacts to 4 so it may have 
+				//	  thrown away different edge contacts from each triangle. This should be fine though
+				//	  since it only happens when we already have a lot of contacts on the triangles so
+				//	  we don't need those edge contacts.
+				FContactEdgeID EdgeID = FContactEdgeID(VertexID0, VertexID1);
+
+				// If we already have collisions with this edge on another triangle,
+				// keep those with the normal closest to its face normal
+				FEdgeContactData* EdgeContact = EdgeContactDataMap.Find(EdgeID);
+				if (EdgeContact != nullptr)
+				{
+					// We have an existing contact on this edge. Is it on another triangle?
+					if (EdgeContact->LocalTriangleIndex != LocalTriangleIndex)
+					{
+						// If existing contact was a face contact just keep it and ignore the new contact.
+						if (EdgeContact->bIsFaceContact)
+						{
+							continue;
+						}
+							
+						// If the existing contact had a normal closer to its face normal, keep it and ignore the new contact
+						if (!bIsFaceContact && (ContactNormalDotTriangleNormal < EdgeContact->ContactNormalDotTriangleNormal))
+						{
+							continue;
+						}
+
+						// Throw away the existing edge contacts
+						if (EdgeContact->ContactIndices[0] != INDEX_NONE)
+						{
+							ContactDatas[EdgeContact->ContactIndices[0]].SetDisabled();
+						}
+						if (EdgeContact->ContactIndices[1] != INDEX_NONE)
+						{
+							ContactDatas[EdgeContact->ContactIndices[1]].SetDisabled();
+						}
+
+						// Change existing edge contact to refer to this triangle
+						EdgeContact->Init(EdgeID, LocalTriangleIndex);
+					}
+				}
+				else
+				{
+					// First time we have seen this edge - start tracking its contacts
+					EdgeContact = &EdgeContactDataMap.Emplace(EdgeID, EdgeID, LocalTriangleIndex);
+				}
+
+				// Track this edge contact
+				check(EdgeContact != nullptr);
+				EdgeContact->AddContact(NewContactIndex, ContactNormalDotTriangleNormal, bIsFaceContact);
+
+				// If it's a face contact, tell each triangle using the edge it has a face contact. 
+				// This is used in the outer loop to skip triangles that have multiple contacts from 
+				// adjacent triangle edges or verts.
 				if (bIsFaceContact)
 				{
 					Triangles[LocalTriangleIndex].AddFaceEdgeCollision();
 
-					const int32 OtherLocalTriangleIndex = GetOtherTriangleIndexForEdge(LocalTriangleIndex, FContactEdgeID(VertexID0, VertexID1));
+					const int32 OtherLocalTriangleIndex = GetOtherTriangleIndexForEdge(LocalTriangleIndex, EdgeID);
 					if (OtherLocalTriangleIndex != INDEX_NONE)
 					{
 						Triangles[OtherLocalTriangleIndex].AddFaceEdgeCollision();
@@ -108,15 +167,15 @@ namespace Chaos::Private
 			else if ((VertexID0 != INDEX_NONE) && (VertexID1 == INDEX_NONE))
 			{
 				// Vertex collisions - we only keep one vertex collision per vertex
-				if (FVertexContactIndex* ExistingContactIndex = VertexContactIndicesMap.Find(VertexID0))
+				if (FVertexContactData* ExistingVertexContact = VertexContactDataMap.Find(VertexID0))
 				{
 					// We have an existing contact at this vertex. 
 					// If it was a face contact just ignore the new contact, 
 					// otherwise keep the contact with the most face-pointing normal.
-					if (!ExistingContactIndex->bIsFaceContact)
+					if (!ExistingVertexContact->bIsFaceContact)
 					{
-						FContactPoint& ExistingContact = Contacts[ExistingContactIndex->ContactIndex];
-						FTriangleContactPointData& ExistingContactData = ContactDatas[ExistingContactIndex->ContactIndex];
+						FContactPoint& ExistingContact = Contacts[ExistingVertexContact->ContactIndex];
+						FTriangleContactPointData& ExistingContactData = ContactDatas[ExistingVertexContact->ContactIndex];
 						if (ContactNormalDotTriangleNormal > ExistingContactData.GetContactNormalDotTriangleNormal())
 						{
 							ExistingContact = ContactPoint;
@@ -126,7 +185,7 @@ namespace Chaos::Private
 							ExistingContactData.SetContactNormalDotTriangleNormal(ContactNormalDotTriangleNormal);
 							ExistingContactData.SetTriangleIndex(LocalTriangleIndex);
 
-							ExistingContactIndex->bIsFaceContact = bIsFaceContact;
+							ExistingVertexContact->bIsFaceContact = bIsFaceContact;
 						}
 					}
 
@@ -135,7 +194,7 @@ namespace Chaos::Private
 				}
 
 				// This is the first time we hit this vertex so store the contact index with the vertex
-				VertexContactIndicesMap.Emplace(VertexID0, VertexID0, NewContactIndex, bIsFaceContact);
+				VertexContactDataMap.Emplace(VertexID0, VertexID0, NewContactIndex, bIsFaceContact);
 			}
 
 			// Store the contact, set the face index
@@ -153,6 +212,11 @@ namespace Chaos::Private
 
 	void FMeshContactGenerator::ProcessGeneratedContacts(const FRigidTransform3& ConvexTransform, const FRigidTransform3& MeshToConvexTransform)
 	{
+		if (Contacts.Num() == 0)
+		{
+			return;
+		}
+
 		if (!!Settings.bFixNormals)
 		{
 			// Contacts that get pruned or corrected will show as green
@@ -186,6 +250,12 @@ namespace Chaos::Private
 		{
 			FTriangleContactPointData& ContactPointData = ContactDatas[ContactIndex];
 
+			// Some contacts may be disabled during the AddContacts phase
+			if (!ContactPointData.IsEnabled())
+			{
+				continue;
+			}
+
 			// Reject back-faces
 			if (!!Settings.bCullBackFaces && (ContactPointData.GetContactNormalDotTriangleNormal() < -Settings.BackFaceCullTolerance))
 			{
diff --git a/Engine/Source/Runtime/Experimental/Chaos/Public/Chaos/Collision/MeshContactGenerator.h b/Engine/Source/Runtime/Experimental/Chaos/Public/Chaos/Collision/MeshContactGenerator.h
--- a/Engine/Source/Runtime/Experimental/Chaos/Public/Chaos/Collision/MeshContactGenerator.h
+++ b/Engine/Source/Runtime/Experimental/Chaos/Public/Chaos/Collision/MeshContactGenerator.h
@@ -309,14 +309,16 @@ namespace Chaos::Private
 	private:
 		using FTriangleExt = FMeshContactGeneratorTriangle;
 
-		// A contact index combined with a flag to indicate if the normal is roughly along the triangle face
-		struct FVertexContactIndex
+		// Info about a vertex contact that allows us to prevent adding duplicates from all the triangles that 
+		// share the vertex. Contains a contact index combined with a flag to indicate if the normal is roughly 
+		// along the triangle face
+		struct FVertexContactData
 		{
-			FVertexContactIndex()
+			FVertexContactData()
 			{
 			}
 			
-			FVertexContactIndex(const FContactVertexID InID, const int32 InContactIndex, const bool bInIsFaceContact)
+			FVertexContactData(const FContactVertexID& InID, const int32 InContactIndex, const bool bInIsFaceContact)
 				: ID(InID)
 				, ContactIndex(InContactIndex)
 				, bIsFaceContact(bInIsFaceContact)
@@ -328,6 +330,50 @@ namespace Chaos::Private
 			bool bIsFaceContact;
 		};
 
+		// Info about an edge contact that allows us to prevent adding duplicate edge contacts from the two
+		// triangles that share the edge.
+		struct FEdgeContactData
+		{
+			FEdgeContactData()
+			{
+			}
+
+			FEdgeContactData(const FContactEdgeID& InID, const int32 InLocalTriangleIndex)
+			{
+				Init(InID, InLocalTriangleIndex);
+			}
+
+			void Init(const FContactEdgeID& InID, const int32 InLocalTriangleIndex)
+			{
+				ID = InID;
+				LocalTriangleIndex = InLocalTriangleIndex;
+				ContactIndices[0] = INDEX_NONE;
+				ContactIndices[1] = INDEX_NONE;
+				ContactNormalDotTriangleNormal = 0.0f;
+				bIsFaceContact = false;
+			}
+
+			void AddContact(int32 InContactIndex, const FRealSingle InContactNormalDotTriangleNormal, const bool bInIsFaceContact)
+			{
+				if (ContactIndices[0] == INDEX_NONE)
+				{
+					ContactIndices[0] = InContactIndex;
+				}
+				else
+				{
+					ContactIndices[1] = InContactIndex;
+				}
+				ContactNormalDotTriangleNormal = FMath::Max(ContactNormalDotTriangleNormal, InContactNormalDotTriangleNormal);
+				bIsFaceContact = bIsFaceContact || bInIsFaceContact;
+			}
+
+			FContactEdgeID ID;
+			int32 LocalTriangleIndex;
+			int32 ContactIndices[2];
+			FRealSingle ContactNormalDotTriangleNormal;
+			bool bIsFaceContact;
+		};
+
 		// The triangle indices that share an edge (assumes only 2)
 		struct FEdgeTriangleIndices
 		{
@@ -419,9 +465,9 @@ namespace Chaos::Private
 
 		bool HasFaceVertexCollision(const FContactVertexID VertexID) const
 		{
-			if (const FVertexContactIndex* ContactIndex = VertexContactIndicesMap.Find(VertexID))
+			if (const FVertexContactData* ContactData = VertexContactDataMap.Find(VertexID))
 			{
-				return ContactIndex->bIsFaceContact;
+				return ContactData->bIsFaceContact;
 			}
 			return false;
 		}
@@ -499,13 +545,21 @@ namespace Chaos::Private
 		};
 		THashMappedArray<FContactEdgeID, FEdgeTriangleIndices, FEdgeTriangleIndicesMapTraits> EdgeTriangleIndicesMap;
 
-		// A map of VertexID to contact index on that vertex - we only ever keep one contact per vertex
-		struct VertexContactIndicesMapTraits
+		// A map of VertexID to contact data for that vertex - we only ever keep one contact per vertex
+		struct FVertexContactIndicesMapTraits
 		{
 			static uint32 GetIDHash(const FContactVertexID& VertexID) { return MurmurFinalize32(VertexID); }
-			static FContactVertexID GetElementID(const FVertexContactIndex& VertexIndex) { return VertexIndex.ID; }
+			static FContactVertexID GetElementID(const FVertexContactData& VertexContactData) { return VertexContactData.ID; }
+		};
+		THashMappedArray<FContactVertexID, FVertexContactData, FVertexContactIndicesMapTraits> VertexContactDataMap;
+
+		// A map of EdgeID to contact data for that edge - we only ever keep contacts from one triangle sharing the edge
+		struct FEdgeContactIndicesMapTraits
+		{
+			static uint32 GetIDHash(const FContactEdgeID& EdgeID) { return uint32(MurmurFinalize64(EdgeID.EdgeID)); }
+			static FContactEdgeID GetElementID(const FEdgeContactData& EdgeContactData) { return EdgeContactData.ID; }
 		};
-		THashMappedArray<FContactVertexID, FVertexContactIndex, VertexContactIndicesMapTraits> VertexContactIndicesMap;
+		THashMappedArray<FContactEdgeID, FEdgeContactData, FEdgeContactIndicesMapTraits> EdgeContactDataMap;
 	};
 
 }
\ No newline at end of file
diff --git a/Engine/Source/Runtime/Experimental/Chaos/Public/Chaos/Framework/HashMappedArray.h b/Engine/Source/Runtime/Experimental/Chaos/Public/Chaos/Framework/HashMappedArray.h
--- a/Engine/Source/Runtime/Experimental/Chaos/Public/Chaos/Framework/HashMappedArray.h
+++ b/Engine/Source/Runtime/Experimental/Chaos/Public/Chaos/Framework/HashMappedArray.h
@@ -95,7 +95,7 @@ namespace Chaos::Private
 		}
 
 		// Add an element with the specified ID to the map. Asserts if an Element with the same ID is already in the map.
-		FORCEINLINE void Add(const FIDType ID, const FElementType& Element)
+		FORCEINLINE FElementType& Add(const FIDType ID, const FElementType& Element)
 		{
 			checkSlow(Find(ID) == nullptr);
 
@@ -103,13 +103,15 @@ namespace Chaos::Private
 			const FHashType Key = FTraits::GetIDHash(ID);
 
 			HashTable.Add(Key, Index);
+
+			return Elements[Index];
 		}
 
 		// Add an element with the specified ID to the map. 
 		// NOTE: since your element type will also need to contain the ID, you usually have to pass the ID twice
 		// to emplace, which is a little annoying, but shouldn't affect much.
 		template <typename... ArgsType>
-		FORCEINLINE void Emplace(const FIDType ID, ArgsType&&... Args)
+		FORCEINLINE FElementType& Emplace(const FIDType ID, ArgsType&&... Args)
 		{
 			checkSlow(Find(ID) == nullptr);
 
@@ -117,6 +119,8 @@ namespace Chaos::Private
 			const FHashType Key = FTraits::GetIDHash(ID);
 
 			HashTable.Add(Key, Index);
+
+			return Elements[Index];
 		}
 
 		void Remove(const FIDType ID)
