Chaos - fix comment in HashMappedArray

#rnx

diff --git a/Engine/Source/Runtime/Experimental/Chaos/Public/Chaos/Framework/HashMappedArray.h b/Engine/Source/Runtime/Experimental/Chaos/Public/Chaos/Framework/HashMappedArray.h
--- a/Engine/Source/Runtime/Experimental/Chaos/Public/Chaos/Framework/HashMappedArray.h
+++ b/Engine/Source/Runtime/Experimental/Chaos/Public/Chaos/Framework/HashMappedArray.h
@@ -94,7 +94,9 @@ namespace Chaos::Private
 			return false;
 		}
 
-		// Add an element with the specified ID to the map. Asserts if an Element with the same ID is already in the map.
+		// Add an element with the specified ID to the map. Returns a reference to the newly inserted element.
+		// It is an error to insert an item twice, but this is only detected in debug builds. Use TryAdd
+		// if you do not know whether the items already exists in the map.
 		FORCEINLINE FElementType& Add(const FIDType ID, const FElementType& Element)
 		{
 			checkSlow(Find(ID) == nullptr);
@@ -107,7 +109,9 @@ namespace Chaos::Private
 			return Elements[Index];
 		}
 
-		// Add an element with the specified ID to the map. 
+		// Add an element with the specified ID to the map. Returns a reference to the newly inserted element.
+		// It is an error to insert an item twice, but this is only detected in debug builds. Use TryEmplace
+		// if you do not know whether the items already exists in the map.
 		// NOTE: since your element type will also need to contain the ID, you usually have to pass the ID twice
 		// to emplace, which is a little annoying, but shouldn't affect much.
 		template <typename... ArgsType>
