Make JsonObjectGraph TObjectPtr correct, rather than assuming UObject* compatability

#rnx

diff --git a/Engine/Source/Runtime/Experimental/JsonObjectGraph/Private/JsonStringifyImpl.cpp b/Engine/Source/Runtime/Experimental/JsonObjectGraph/Private/JsonStringifyImpl.cpp
--- a/Engine/Source/Runtime/Experimental/JsonObjectGraph/Private/JsonStringifyImpl.cpp
+++ b/Engine/Source/Runtime/Experimental/JsonObjectGraph/Private/JsonStringifyImpl.cpp
@@ -554,9 +554,11 @@ void FJsonStringifyImpl::WriteValueToJson(const void* Value, const void* Default
 		// write any tagged data for the struct - but be sure to delta serialize
 		WriteStructToJson(Value, DefaultValue, StructProperty->Struct, StructProperty->Struct);
 	}
-	else if (CastField<FObjectProperty>(Property) || CastField<FClassProperty>(Property))
+	else if (const FObjectProperty* ObjectProperty = CastField<FObjectProperty>(Property))
 	{
-		const UObject* Object = *(UObject**)Value;
+		const TObjectPtr<const UObject>* ObjectPtr = 
+			static_cast<const TObjectPtr<const UObject>*>(Value);
+		const UObject* Object = ObjectPtr->Get();
 		WriteObjectAsJsonToWriter(CurrentObject, Object, Writer);
 	}
 	// containers:
