Enable everywhere now.

#rnx

diff --git a/Engine/Source/Runtime/Experimental/IoStore/HttpClient/IoStoreHttpClient.Build.cs b/Engine/Source/Runtime/Experimental/IoStore/HttpClient/IoStoreHttpClient.Build.cs
--- a/Engine/Source/Runtime/Experimental/IoStore/HttpClient/IoStoreHttpClient.Build.cs
+++ b/Engine/Source/Runtime/Experimental/IoStore/HttpClient/IoStoreHttpClient.Build.cs
@@ -14,9 +14,6 @@ public class IoStoreHttpClient : ModuleRules
 
 		CppCompileWarningSettings.UnsafeTypeCastWarningLevel = WarningLevel.Error; 
 
-		if (Target.Platform.IsInGroup(UnrealPlatformGroup.Windows))
-		{
-			PrivateDependencyModuleNames.Add("nghttp2");
-		}
+		PrivateDependencyModuleNames.Add("nghttp2");
 	}
 }
