Updated internals IasTool utility.

#rnx

diff --git a/Engine/Source/Runtime/Experimental/IoStore/OnDemand/Private/Tool/Purl.cpp b/Engine/Source/Runtime/Experimental/IoStore/OnDemand/Private/Tool/Purl.cpp
--- a/Engine/Source/Runtime/Experimental/IoStore/OnDemand/Private/Tool/Purl.cpp
+++ b/Engine/Source/Runtime/Experimental/IoStore/OnDemand/Private/Tool/Purl.cpp
@@ -110,18 +110,45 @@ static int32 PurlCommandEntry(const FContext& Context)
 		}
 	};
 
+	bool bRedirect = Context.Get<bool>(TEXT("-Redirect"));
+	bool bUseHttp2 = Context.Get<bool>(TEXT("-HttpTwo"));
+	int32 BufSizeKiB = Context.Get<int32>(TEXT("-BufSize"), -1);
+	uint32 ThrottleKiBps = Context.Get<uint32>(TEXT("-Throttle"), 0);
+
+	FConnectionPool::FParams PoolParams = {
+		.RecvBufSize = BufSizeKiB << 10,
+		.ConnectionCount = 1,
+		.HttpVersion = (bUseHttp2 ? EHttpVersion::Two : EHttpVersion::One),
+	};
+	if (!PoolParams.SetHostFromUrl(AnsiUrl))
+	{
+		std::printf("ERROR: invalid url\n");
+		return 1;
+	}
+	FConnectionPool Pool(PoolParams);
+
 	FEventLoop Loop;
-	FEventLoop::FRequestParams RequestParams;
-	if (Context.Get<bool>(TEXT("-Redirect")))
+
+	FEventLoop::FRequestParams RequestParams = {
+		.bAutoRedirect = bRedirect,
+	};
+
+	FAnsiStringView AnsiPath = AnsiUrl;
+	for (int32 i = 2; i >= 0; i -= 2)
 	{
-		RequestParams.bAutoRedirect = true;
+		int32 Index;
+		AnsiPath.FindChar('/', Index);
+		AnsiPath = AnsiPath.Mid(Index + i);
 	}
-	FRequest Request = Loop.Request(AnsiMethod, AnsiUrl, &RequestParams);
+
+	FRequest Request = Loop.Request(AnsiMethod, AnsiPath, Pool, &RequestParams);
+	Request.Header("User-Agent", "purl");
 	Loop.Send(MoveTemp(Request), Sink);
-	
+
+	Loop.SetFailTimeout(4000);
+	Loop.Throttle(ThrottleKiBps);
 	while (Loop.Tick(-1))
 	{
-		FPlatformProcess::SleepNoStats(0.1f);
 	}
 
 	std::printf("Data: %u bytes\n", ContentSize);
@@ -138,6 +165,9 @@ static FCommand PurlCommand(
 		TArgument<FStringView>(TEXT("Url"), TEXT("Url to download")),
 		TArgument<FStringView>(TEXT("-Method"), TEXT("Request method")),
 		TArgument<bool>(TEXT("-Redirect"), TEXT("Follow 30x redirects")),
+		TArgument<bool>(TEXT("-HttpTwo"), TEXT("Use HTTP/2")),
+		TArgument<int32>(TEXT("-BufSize"), TEXT("Recv buffer size in KiB")),
+		TArgument<uint32>(TEXT("-Throttle"), TEXT("Bandwidth throttle in KiB/s")),
 	}
 );
 
