Updated Chaos and CVD to draw the dirty aabb tree.

diff --git a/Engine/Source/Runtime/Experimental/Chaos/Private/ChaosVisualDebugger/ChaosVDDataWrapperUtils.cpp b/Engine/Source/Runtime/Experimental/Chaos/Private/ChaosVisualDebugger/ChaosVDDataWrapperUtils.cpp
--- a/Engine/Source/Runtime/Experimental/Chaos/Private/ChaosVisualDebugger/ChaosVDDataWrapperUtils.cpp
+++ b/Engine/Source/Runtime/Experimental/Chaos/Private/ChaosVisualDebugger/ChaosVDDataWrapperUtils.cpp
@@ -711,11 +711,27 @@ void FChaosVDDataWrapperUtils::BuildDataWrapperFromAABBStructure(const Chaos::IS
 		{
 			BuildDataWrapperFromAABBStructure(*AABBTree, AABBTreDataWrapper);
 			OutAABBTrees.Emplace(AABBTreDataWrapper);
+
+			if (AABBTree->DirtyElementTree)
+			{
+				FChaosVDAABBTreeDataWrapper DirtyTreeWrapper;
+				DirtyTreeWrapper.SolverId = OwnerSolverID;
+				BuildDataWrapperFromAABBStructure(*AABBTree->DirtyElementTree, DirtyTreeWrapper);
+				OutAABBTrees.Emplace(DirtyTreeWrapper);
+			}
 		}
 		else if (const TAABBTree<FAccelerationStructureHandle, TBoundingVolume<FAccelerationStructureHandle>>* AABBTreeBV = AccelerationStructure->template As<TAABBTree<FAccelerationStructureHandle, TBoundingVolume<FAccelerationStructureHandle>>>())
 		{
 			BuildDataWrapperFromAABBStructure(*AABBTreeBV, AABBTreDataWrapper);
 			OutAABBTrees.Emplace(AABBTreDataWrapper);
+
+			if (AABBTreeBV->DirtyElementTree)
+			{
+				FChaosVDAABBTreeDataWrapper DirtyTreeWrapper;
+				DirtyTreeWrapper.SolverId = OwnerSolverID;
+				BuildDataWrapperFromAABBStructure(*AABBTreeBV->DirtyElementTree, DirtyTreeWrapper);
+				OutAABBTrees.Emplace(DirtyTreeWrapper);
+			}
 		}
 	}
 }
diff --git a/Engine/Source/Runtime/Experimental/Chaos/Public/Chaos/AABBTree.h b/Engine/Source/Runtime/Experimental/Chaos/Public/Chaos/AABBTree.h
--- a/Engine/Source/Runtime/Experimental/Chaos/Public/Chaos/AABBTree.h
+++ b/Engine/Source/Runtime/Experimental/Chaos/Public/Chaos/AABBTree.h
@@ -3841,6 +3841,10 @@ private:
 				const TLeafType& Leaf = Leaves[LeafIndex];
 				Leaf.DebugDrawLeaf(*InInterface, FLinearColor::MakeRandomColor(), 10.f);
 			}
+			if (DirtyElementTree)
+			{
+				DirtyElementTree->DebugDraw(InInterface);
+			}
 		}
 	}
 #endif
