ChaosFlesh:
-refactor flesh constraints to be template free using FSolverParticles and FSolverReal
-modify HeadlessChaosFlesh tests to work with FSolverReal (float) instead of double. Loosen precision checks for finite difference tests (Error(eps)/eps issue)

diff --git a/Engine/Plugins/ChaosCloth/Source/ChaosCloth/Private/ChaosCloth/ChaosClothConstraints.cpp b/Engine/Plugins/ChaosCloth/Source/ChaosCloth/Private/ChaosCloth/ChaosClothConstraints.cpp
--- a/Engine/Plugins/ChaosCloth/Source/ChaosCloth/Private/ChaosCloth/ChaosClothConstraints.cpp
+++ b/Engine/Plugins/ChaosCloth/Source/ChaosCloth/Private/ChaosCloth/ChaosClothConstraints.cpp
@@ -1,6 +1,7 @@
 // Copyright Epic Games, Inc. All Rights Reserved.
 #include "ChaosCloth/ChaosClothConstraints.h"
 #include "ChaosCloth/ChaosClothingAccessoryMeshData.h"
+#include "ChaosCloth/ChaosClothingAccessoryMeshData.h"
 #include "ChaosCloth/ChaosClothingSimulationConfig.h"
 #include "ChaosCloth/ChaosWeightMapTarget.h"
 #include "ChaosCloth/ChaosClothingPatternData.h"
@@ -541,7 +542,7 @@ void FClothConstraints::CreateGSRules()
 
 	int32 ConstraintPostProcessingIndex = 0;
 
-	GSMainConstraint = MakeShared<Chaos::Softs::FGaussSeidelMainConstraint<Softs::FSolverReal, Softs::FSolverParticles>>(PBDEvolution->Particles(), bClothDoQuasistatics, bUseSOR, SOROmega, 100);
+	GSMainConstraint = MakeShared<Chaos::Softs::FGaussSeidelMainConstraints>(PBDEvolution->Particles(), bClothDoQuasistatics, bUseSOR, SOROmega, 100);
 
 	if (bClothDoQuasistatics)
 	{
@@ -675,7 +676,7 @@ void FClothConstraints::CreateGSRules()
 
 			const TArray<TVec3<int32>> CodimensionalMesh = XStretchBiasConstraints->GetConstraints();
 
-			GSCorotatedCodimensionalConstraint = MakeShared<Chaos::Softs::FGaussSeidelCorotatedCodimensionalConstraints<Softs::FSolverReal, Softs::FSolverParticles>>(PBDEvolution->Particles(), CodimensionalMesh, false, YoungsModulus);
+			GSCorotatedCodimensionalConstraint = MakeShared<Chaos::Softs::FGaussSeidelCorotatedTriangleConstraints>(PBDEvolution->Particles(), CodimensionalMesh, false, YoungsModulus);
 
 			GSMainConstraint->AddStaticConstraints(XStretchBiasConstraints->GetConstraintsArray(), IncidentElements, IncidentElementsLocal);
 			const int32 StaticIndex = GSMainConstraint->AddStaticConstraintResidualAndHessianRange(1);
@@ -2111,7 +2112,7 @@ void FClothConstraints::CreatePBDRules()
 	#if !UE_BUILD_SHIPPING	
 	if (bDisplayResidual || bWriteFinalResiduals)
 	{
-		GSMainConstraint = MakeShared<Chaos::Softs::FGaussSeidelMainConstraint<Softs::FSolverReal, Softs::FSolverParticles>>(PBDEvolution->Particles(), false, false, (Softs::FSolverReal)1.2, 100);
+		GSMainConstraint = MakeShared<Chaos::Softs::FGaussSeidelMainConstraints>(PBDEvolution->Particles(), false, false, (Softs::FSolverReal)1.2, 100);
 
 		ConstraintInits[ConstraintInitIndex++] =
 			[this](Softs::FSolverParticles& InParticles, const Softs::FSolverReal Dt)
diff --git a/Engine/Plugins/ChaosCloth/Source/ChaosCloth/Public/ChaosCloth/ChaosClothConstraints.h b/Engine/Plugins/ChaosCloth/Source/ChaosCloth/Public/ChaosCloth/ChaosClothConstraints.h
--- a/Engine/Plugins/ChaosCloth/Source/ChaosCloth/Public/ChaosCloth/ChaosClothConstraints.h
+++ b/Engine/Plugins/ChaosCloth/Source/ChaosCloth/Public/ChaosCloth/ChaosClothConstraints.h
@@ -232,8 +232,8 @@ namespace Chaos
 		TSharedPtr<Softs::FPBDCollisionSpringConstraints> SelfCollisionConstraints;
 		TSharedPtr<Softs::FPBDTriangleMeshIntersections> SelfIntersectionConstraints;
 		TSharedPtr<Softs::FPBDSelfCollisionSphereConstraints> SelfCollisionSphereConstraints;
-		TSharedPtr<Softs::FGaussSeidelMainConstraint<Softs::FSolverReal, Softs::FSolverParticles>> GSMainConstraint;
-		TSharedPtr<Softs::FGaussSeidelCorotatedCodimensionalConstraints<Softs::FSolverReal, Softs::FSolverParticles>> GSCorotatedCodimensionalConstraint;
+		TSharedPtr<Softs::FGaussSeidelMainConstraints> GSMainConstraint;
+		TSharedPtr<Softs::FGaussSeidelCorotatedTriangleConstraints> GSCorotatedCodimensionalConstraint;
 		TSharedPtr<Softs::FMultiResConstraints> MultiResConstraints;
 		TSharedPtr<Softs::FPBDExtremeDeformationConstraints> ExtremeDeformationConstraints;
 		//~ Begin Force-based solver only constraints
diff --git a/Engine/Source/Runtime/Experimental/Chaos/Private/Chaos/Deformable/ChaosDeformableSolver.cpp b/Engine/Source/Runtime/Experimental/Chaos/Private/Chaos/Deformable/ChaosDeformableSolver.cpp
--- a/Engine/Source/Runtime/Experimental/Chaos/Private/Chaos/Deformable/ChaosDeformableSolver.cpp
+++ b/Engine/Source/Runtime/Experimental/Chaos/Private/Chaos/Deformable/ChaosDeformableSolver.cpp
@@ -99,7 +99,7 @@ namespace Chaos::Softs
 	FAutoConsoleVariableRef CVarDeformableXPBDCorotatedBatchThreshold(TEXT("p.Chaos.Deformable.XPBDBatchThreshold"), GDeformableXPBDCorotatedParams.XPBDCorotatedBatchThreshold, TEXT("Batch threshold for physics parallel for. [def: 5]"));
 	FAutoConsoleVariableRef CVarDeformableXPBDCorotatedBatchNumLogExtremeParticle(TEXT("p.Chaos.Deformable.NumLogExtremeParticle"), GDeformableXPBDCorotatedParams.NumLogExtremeParticle, TEXT("Number of most deformed particles logged. [def: 0]"));
 
-	FDeformableXPBDWeakConstraintParams GDeformableXPBDWeakConstraintParams;
+	FDeformableXPBDWeightedSpringConstraintParams GDeformableXPBDWeakConstraintParams;
 	FAutoConsoleVariableRef CVarDeformableXPBDWeakConstraintLineWidth(TEXT("p.Chaos.Deformable.XPBDWeakConstraintLineWidth"), GDeformableXPBDWeakConstraintParams.DebugLineWidth, TEXT("Line width for visualizing the double bindings in XPBD weak constraints. [def: 5]"));
 	FAutoConsoleVariableRef CVarDeformableXPBDWeakConstraintParticleWidth(TEXT("p.Chaos.Deformable.XPBDWeakConstraintParticleWidth"), GDeformableXPBDWeakConstraintParams.DebugParticleWidth, TEXT("Line width for visualizing the double bindings in XPBD weak constraints. [def: 20]"));
 	FAutoConsoleVariableRef CVarDeformableXPBDWeakConstraintDebugDraw(TEXT("p.Chaos.Deformable.XPBDWeakConstraintEnableDraw"), GDeformableXPBDWeakConstraintParams.bVisualizeBindings, TEXT("Debug draw the double bindings in XPBD weak constraints. [def: false]"));
@@ -183,12 +183,12 @@ namespace Chaos::Softs
 			AllTetAlphaJArray.Reset(new TArray<FSolverReal>());
 			AllCorotatedCodEMeshArray.Reset(new TArray<FSolverReal>());
 			AllSkinEMeshArray.Reset(new TArray<FSolverReal>());
-			GSWeakConstraints.Reset(new FGaussSeidelWeakConstraints<FSolverReal, FSolverParticles>({}, {}, {}, {}, {}, GDeformableXPBDWeakConstraintParams));
-			GSDynamicWeakConstraints.Reset(new FGaussSeidelDynamicWeakConstraints<FSolverReal, FSolverParticles>(GDeformableXPBDWeakConstraintParams));
-			GSSphereRepulsionConstraints.Reset(new FGaussSeidelSphereRepulsionConstraints<FSolverReal, FSolverParticles>(Property.SphereRepulsionRadius, Property.SphereRepulsionStiffness, Evolution->Particles(), GDeformableXPBDWeakConstraintParams));
-			GSVolumeConstraints.Reset(new FGaussSeidelUnilateralTetConstraints<Softs::FSolverReal, Softs::FSolverParticles>(Evolution->Particles(), {}, {}));
+			GSWeakConstraints.Reset(new FGaussSeidelSpringConstraints({}, {}, {}, {}, {}, GDeformableXPBDWeakConstraintParams));
+			GSDynamicWeakConstraints.Reset(new FGaussSeidelDynamicSpringConstraints(GDeformableXPBDWeakConstraintParams));
+			GSSphereRepulsionConstraints.Reset(new FGaussSeidelParticleSphereRepulsionConstraints(Property.SphereRepulsionRadius, Property.SphereRepulsionStiffness, Evolution->Particles(), GDeformableXPBDWeakConstraintParams));
+			GSVolumeConstraints.Reset(new FGaussSeidelUnilateralVolumeConstraints(Evolution->Particles(), {}, {}));
 			MuscleIndexOffset.Empty();
-			MuscleActivationConstraints.Reset(new FMuscleActivationConstraints<FSolverReal, FSolverParticles>());
+			MuscleActivationConstraints.Reset(new FMuscleActivationTetrahedralConstraints());
 		}
 		AllUnconstrainedSurfaceElementsCorotatedCod.Reset(new TArray<Chaos::TVec3<int32>>());
 		AllUnconstrainedSurfaceElementsSkin.Reset(new TArray<Chaos::TVec3<int32>>());
@@ -724,8 +724,8 @@ namespace Chaos::Softs
 				int32 InitIndex = Evolution->AddConstraintInitRange(1, true);
 				int32 ConstraintIndex = Evolution->AddConstraintRuleRange(1, true);
 
-				FXPBDWeakConstraints<FSolverReal, FSolverParticles>* WeakConstraint =
-					new FXPBDWeakConstraints<FSolverReal, FSolverParticles>(Evolution->Particles(),
+				FXPBDWeightedSpringConstraints* WeakConstraint =
+					new FXPBDWeightedSpringConstraints(Evolution->Particles(),
 						PositionTargetIndices, PositionTargetWeights, PositionTargetStiffness, PositionTargetSecondIndices, PositionTargetSecondWeights, GDeformableXPBDWeakConstraintParams);
 
 				Evolution->ConstraintInits()[InitIndex] =
@@ -740,7 +740,7 @@ namespace Chaos::Softs
 					WeakConstraint->ApplyInParallel(InParticles, Dt);
 				};
 
-				WeakConstraints.Add(TUniquePtr<FXPBDWeakConstraints<FSolverReal, FSolverParticles>>(WeakConstraint));
+				WeakConstraints.Add(TUniquePtr<FXPBDWeightedSpringConstraints>(WeakConstraint));
 			}
 		}
 
@@ -760,7 +760,7 @@ namespace Chaos::Softs
 			}
 			InStiffnessArray[ConstraintIdx] = VolumeConstraint.GetStiffness(ConstraintIdx);
 		}
-		GSVolumeConstraints.Reset(new FGaussSeidelUnilateralTetConstraints<Softs::FSolverReal, Softs::FSolverParticles>(Evolution->Particles(), MoveTemp(InConstraints), MoveTemp(InStiffnessArray)));
+		GSVolumeConstraints.Reset(new FGaussSeidelUnilateralVolumeConstraints(Evolution->Particles(), MoveTemp(InConstraints), MoveTemp(InStiffnessArray)));
 		if (GSVolumeConstraints->NumConstraints() && !Property.bUseGaussSeidelConstraints)
 		{
 			UE_LOG(LogChaosDeformableSolver, Error, TEXT("Error: must check [Use Gauss Seidel constraints] for volume constraints."));
@@ -909,7 +909,7 @@ namespace Chaos::Softs
 									if (GSDynamicWeakConstraints && Property.bEnableDynamicSprings)
 									{
 										float Stiffness = AddConstraints.Parameters.Stiffness;
-										TArray<const FGaussSeidelWeakConstraints<Softs::FSolverReal, Softs::FSolverParticles>::FGaussSeidelConstraintHandle*> ConstraintHandles = GSDynamicWeakConstraints -> AddParticleTetrahedraConstraints(Geom, Evolution->Particles(), Intersections, SampleRange, TargetRange, Stiffness);
+										TArray<const FGaussSeidelSpringConstraints::FGaussSeidelConstraintHandle*> ConstraintHandles = GSDynamicWeakConstraints -> AddParticleTetrahedraConstraints(Geom, Evolution->Particles(), Intersections, SampleRange, TargetRange, Stiffness);
 										Chaos::Softs::FConstraintObjectParticleHandel& HandleValue = Proxy.Constraints.FindOrAdd(AddConstraints);
 										HandleValue.Handles = ConstraintHandles;
 										bDynamicConstraintIsUpdated = true;
@@ -1104,8 +1104,8 @@ namespace Chaos::Softs
 
 						if (Property.bDoBlended)
 						{
-							FBlendedXPBDCorotatedConstraints<FSolverReal, FSolverParticles>* BlendedCorotatedConstraint =
-								new FBlendedXPBDCorotatedConstraints<FSolverReal, FSolverParticles>(
+							FBlendedXPBDCorotatedTetrahedralConstraints* BlendedCorotatedConstraint =
+								new FBlendedXPBDCorotatedTetrahedralConstraints(
 									Evolution->Particles(), Elements, TetStiffness, (FSolverReal).3,/*bRecordMetric = */false, Property.BlendedZeta);
 
 							Evolution->ConstraintInits()[InitIndex] =
@@ -1120,13 +1120,13 @@ namespace Chaos::Softs
 									BlendedCorotatedConstraint->ApplyInParallel(InParticles, Dt);
 								};
 
-							BlendedCorotatedConstraints.Add(TUniquePtr<FBlendedXPBDCorotatedConstraints<FSolverReal, FSolverParticles>>(BlendedCorotatedConstraint));
+							BlendedCorotatedConstraints.Add(TUniquePtr<FBlendedXPBDCorotatedTetrahedralConstraints>(BlendedCorotatedConstraint));
 
 						}
 						else
 						{
-							FXPBDCorotatedConstraints<FSolverReal, FSolverParticles>* CorotatedConstraint =
-								new FXPBDCorotatedConstraints<FSolverReal, FSolverParticles>(
+							FXPBDCorotatedTetrahedralConstraints* CorotatedConstraint =
+								new FXPBDCorotatedTetrahedralConstraints(
 									Evolution->Particles(), Elements, TetStiffness, TetNu, MoveTemp(AlphaJMesh), GDeformableXPBDCorotatedParams);
 
 							Evolution->ConstraintInits()[InitIndex] =
@@ -1141,7 +1141,7 @@ namespace Chaos::Softs
 									CorotatedConstraint->ApplyInParallel(InParticles, Dt);
 								};
 
-							CorotatedConstraints.Add(TUniquePtr<FXPBDCorotatedConstraints<FSolverReal, FSolverParticles>>(CorotatedConstraint));
+							CorotatedConstraints.Add(TUniquePtr<FXPBDCorotatedTetrahedralConstraints>(CorotatedConstraint));
 						}
 					}
 
@@ -1647,7 +1647,7 @@ namespace Chaos::Softs
 	{
 		PERF_SCOPE(STAT_ChaosDeformableSolver_InitializeGridBasedConstraintVariables);
 
-		GridBasedCorotatedConstraint.Reset(new Chaos::Softs::FXPBDGridBasedCorotatedConstraints<FSolverReal, FSolverParticles>(
+		GridBasedCorotatedConstraint.Reset(new Chaos::Softs::FXPBDGridBasedCorotatedTetrahedralConstraints(
 			Evolution->Particles(), *AllElements, Property.GridDx, /*bRecordMetric = */false, (Chaos::Softs::FSolverReal).1, (Chaos::Softs::FSolverReal).01, (Chaos::Softs::FSolverReal).4, (Chaos::Softs::FSolverReal)1000.0));
 		Evolution->ResetConstraintRules();
 		int32 InitIndex1 = Evolution->AddConstraintInitRange(1, true);
@@ -1674,11 +1674,11 @@ namespace Chaos::Softs
 	{
 		PERF_SCOPE(STAT_ChaosDeformableSolver_InitializeGaussSeidelConstraintVariables);
 
-		GSMainConstraint.Reset(new Chaos::Softs::FGaussSeidelMainConstraint<FSolverReal, FSolverParticles>(Evolution->Particles(), Property.bDoQuasistatics, Property.bUseSOR, Property.OmegaSOR, GSParallelMax, MaxDxRatio));
+		GSMainConstraint.Reset(new Chaos::Softs::FGaussSeidelMainConstraints(Evolution->Particles(), Property.bDoQuasistatics, Property.bUseSOR, Property.OmegaSOR, GSParallelMax, MaxDxRatio));
 
 		if (AllUnconstrainedSurfaceElementsCorotatedCod->Num() > 0) 
 		{
-			GSCorotatedCodConstraints.Reset(new Chaos::Softs::FGaussSeidelCorotatedCodimensionalConstraints<FSolverReal, FSolverParticles>(
+			GSCorotatedCodConstraints.Reset(new Chaos::Softs::FGaussSeidelCorotatedTriangleConstraints(
 				Evolution->Particles(), *AllUnconstrainedSurfaceElementsCorotatedCod, *AllCorotatedCodEMeshArray));
 			TArray<TArray<int32>> IncidentElements, IncidentElementsLocal;
 			GSMainConstraint->AddStaticConstraints(GSCorotatedCodConstraints->GetConstraintsArray(), IncidentElements, IncidentElementsLocal);
@@ -1693,7 +1693,7 @@ namespace Chaos::Softs
 
 		if (AllUnconstrainedSurfaceElementsSkin->Num() > 0)
 		{
-			GSLinearCodConstraints.Reset(new Chaos::Softs::FGaussSeidelLinearCodimensionalConstraints<FSolverReal, FSolverParticles>(
+			GSLinearCodConstraints.Reset(new Chaos::Softs::FGaussSeidelLinearTriangleConstraints(
 				Evolution->Particles(), *AllUnconstrainedSurfaceElementsSkin, *AllSkinEMeshArray));
 			TArray<TArray<int32>> IncidentElements, IncidentElementsLocal;
 			GSMainConstraint->AddStaticConstraints(GSLinearCodConstraints->GetConstraintsArray(), IncidentElements, IncidentElementsLocal);
@@ -1708,7 +1708,7 @@ namespace Chaos::Softs
 
 		if (Property.bUseGSNeohookean)
 		{
-			GSNeohookeanConstraints.Reset(new Chaos::Softs::FGaussSeidelNeohookeanConstraints<FSolverReal, FSolverParticles>(
+			GSNeohookeanConstraints.Reset(new Chaos::Softs::FGaussSeidelNeohookeanTetrahedralConstraints(
 				Evolution->Particles(), *AllElements, *AllTetEMeshArray, *AllTetNuMeshArray, MoveTemp(*AllTetAlphaJArray), MoveTemp(*AllIncidentElements), MoveTemp(*AllIncidentElementsLocal), 0, Evolution->Particles().Size(), Property.bDoQuasistatics, Property.bUseSOR, Property.OmegaSOR, GDeformableXPBDCorotatedParams));
 			Evolution->ResetConstraintRules();
 			
@@ -1738,7 +1738,7 @@ namespace Chaos::Softs
 		} 
 		else
 		{
-			GSCorotatedConstraints.Reset(new Chaos::Softs::FGaussSeidelCorotatedConstraints<FSolverReal, FSolverParticles>(
+			GSCorotatedConstraints.Reset(new Chaos::Softs::FGaussSeidelCorotatedTetrahedralConstraints(
 				Evolution->Particles(), *AllElements, *AllTetEMeshArray, *AllTetNuMeshArray, MoveTemp(*AllTetAlphaJArray), MoveTemp(*AllIncidentElements), MoveTemp(*AllIncidentElementsLocal), 0, Evolution->Particles().Size(), Property.bDoQuasistatics, Property.bUseSOR, Property.OmegaSOR, GDeformableXPBDCorotatedParams));
 			Evolution->ResetConstraintRules();
 
@@ -1886,7 +1886,7 @@ namespace Chaos::Softs
 		}
 		if (Property.bDoSphereRepulsion)
 		{
-			GSSphereRepulsionConstraints.Reset(new FGaussSeidelSphereRepulsionConstraints<FSolverReal, FSolverParticles>(Property.SphereRepulsionRadius, Property.SphereRepulsionStiffness, Evolution->Particles(), GDeformableXPBDWeakConstraintParams));
+			GSSphereRepulsionConstraints.Reset(new FGaussSeidelParticleSphereRepulsionConstraints(Property.SphereRepulsionRadius, Property.SphereRepulsionStiffness, Evolution->Particles(), GDeformableXPBDWeakConstraintParams));
 			int32 InitIndex = Evolution->AddConstraintInitRange(1, true);
 			
 			Evolution->ConstraintInits()[InitIndex] 
@@ -2041,11 +2041,11 @@ namespace Chaos::Softs
 		MuscleIndexOffset.Empty();
 		AllUnconstrainedSurfaceElementsCorotatedCod.Reset(new TArray<Chaos::TVec3<int32>>());
 		AllUnconstrainedSurfaceElementsSkin.Reset(new TArray<Chaos::TVec3<int32>>());
-		GSWeakConstraints.Reset(new FGaussSeidelWeakConstraints<FSolverReal, FSolverParticles>({}, {}, {}, {}, {}, GDeformableXPBDWeakConstraintParams));
-		GSDynamicWeakConstraints.Reset(new FGaussSeidelDynamicWeakConstraints<FSolverReal, FSolverParticles>(GDeformableXPBDWeakConstraintParams));
-		GSSphereRepulsionConstraints.Reset(new FGaussSeidelSphereRepulsionConstraints<FSolverReal, FSolverParticles>(Property.SphereRepulsionRadius, Property.SphereRepulsionStiffness, Evolution->Particles(), GDeformableXPBDWeakConstraintParams));
-		GSVolumeConstraints.Reset(new FGaussSeidelUnilateralTetConstraints<Softs::FSolverReal, Softs::FSolverParticles>(Evolution->Particles(), {}, {}));
-		MuscleActivationConstraints.Reset(new FMuscleActivationConstraints<FSolverReal, FSolverParticles>());
+		GSWeakConstraints.Reset(new FGaussSeidelSpringConstraints({}, {}, {}, {}, {}, GDeformableXPBDWeakConstraintParams));
+		GSDynamicWeakConstraints.Reset(new FGaussSeidelDynamicSpringConstraints(GDeformableXPBDWeakConstraintParams));
+		GSSphereRepulsionConstraints.Reset(new FGaussSeidelParticleSphereRepulsionConstraints(Property.SphereRepulsionRadius, Property.SphereRepulsionStiffness, Evolution->Particles(), GDeformableXPBDWeakConstraintParams));
+		GSVolumeConstraints.Reset(new FGaussSeidelUnilateralVolumeConstraints(Evolution->Particles(), {}, {}));
+		MuscleActivationConstraints.Reset(new FMuscleActivationTetrahedralConstraints());
 	}
 
 	void FDeformableSolver::RemoveSimulationObjects()
diff --git a/Engine/Source/Runtime/Experimental/Chaos/Private/Chaos/Deformable/GaussSeidelMainConstraint.cpp b/Engine/Source/Runtime/Experimental/Chaos/Private/Chaos/Deformable/GaussSeidelMainConstraint.cpp
--- a/Engine/Source/Runtime/Experimental/Chaos/Private/Chaos/Deformable/GaussSeidelMainConstraint.cpp
+++ b/Engine/Source/Runtime/Experimental/Chaos/Private/Chaos/Deformable/GaussSeidelMainConstraint.cpp
@@ -13,8 +13,7 @@ namespace Chaos::Softs
 
 	FAutoConsoleVariableRef CVarClothCGTol(TEXT("p.Chaos.Cloth.CGTol"), MaxItCG, TEXT("CG Tolerance [def: 1e-4]"));
 
-	template <typename T, typename ParticleType>
-	void FGaussSeidelMainConstraint<T, ParticleType>::AddStaticConstraints(const TArray<TArray<int32>>& ExtraConstraints, TArray<TArray<int32>>& ExtraIncidentElements, TArray<TArray<int32>>& ExtraIncidentElementsLocal)
+	void FGaussSeidelMainConstraints::AddStaticConstraints(const TArray<TArray<int32>>& ExtraConstraints, TArray<TArray<int32>>& ExtraIncidentElements, TArray<TArray<int32>>& ExtraIncidentElementsLocal)
 	{	
 		if (!IsClean(ExtraConstraints, ExtraIncidentElements, ExtraIncidentElementsLocal))
 		{
@@ -44,8 +43,7 @@ namespace Chaos::Softs
 		StaticIncidentElementsOffsets.Add(StaticConstraints.Num());
 	}
 
-	template <typename T, typename ParticleType>
-	void FGaussSeidelMainConstraint<T, ParticleType>::AddTransientConstraints(const TArray<TArray<int32>>& ExtraConstraints, TArray<TArray<int32>>& ExtraIncidentElements, TArray<TArray<int32>>& ExtraIncidentElementsLocal, bool CheckIncidentElements)
+	void FGaussSeidelMainConstraints::AddTransientConstraints(const TArray<TArray<int32>>& ExtraConstraints, TArray<TArray<int32>>& ExtraIncidentElements, TArray<TArray<int32>>& ExtraIncidentElementsLocal, bool CheckIncidentElements)
 	{
 		if (CheckIncidentElements)
 		{
@@ -79,8 +77,7 @@ namespace Chaos::Softs
 	
 	}
 
-	template <typename T, typename ParticleType>
-	void FGaussSeidelMainConstraint<T, ParticleType>::AddDynamicConstraints(const TArray<TArray<int32>>& ExtraConstraints, TArray<TArray<int32>>& ExtraIncidentElements, TArray<TArray<int32>>& ExtraIncidentElementsLocal, bool CheckIncidentElements)
+	void FGaussSeidelMainConstraints::AddDynamicConstraints(const TArray<TArray<int32>>& ExtraConstraints, TArray<TArray<int32>>& ExtraIncidentElements, TArray<TArray<int32>>& ExtraIncidentElementsLocal, bool CheckIncidentElements)
 	{
 		if (CheckIncidentElements)
 		{
@@ -113,17 +110,16 @@ namespace Chaos::Softs
 		DynamicIncidentElementsOffsets.Add(DynamicConstraints.Num());
 	}
 
-	template <typename T, typename ParticleType>
-	TArray<TVec3<T>> FGaussSeidelMainConstraint<T, ParticleType>::ComputeNewtonResiduals(const ParticleType& Particles, const T Dt, const bool Write2File, TArray<PMatrix<T, 3, 3>>* AllParticleHessian)
+	TArray<TVec3<FSolverReal>> FGaussSeidelMainConstraints::ComputeNewtonResiduals(const FSolverParticles& Particles, const FSolverReal Dt, const bool Write2File, TArray<PMatrix<FSolverReal, 3, 3>>* AllParticleHessian)
 	{
 		TRACE_CPUPROFILER_EVENT_SCOPE(STAT_ChaosGaussSeidelComputeNewtonResidual);
 
-		TArray<TVec3<T>> NewtonResidual;
-		NewtonResidual.Init(TVec3<T>((T)0.), xtilde.Num());
+		TArray<TVec3<FSolverReal>> NewtonResidual;
+		NewtonResidual.Init(TVec3<FSolverReal>((FSolverReal)0.), xtilde.Num());
 
 		if (AllParticleHessian)
 		{
-			AllParticleHessian->Init(PMatrix<T, 3, 3>((T)0.), Particles.Size());
+			AllParticleHessian->Init(PMatrix<FSolverReal, 3, 3>((FSolverReal)0.), Particles.Size());
 		}
 
 		for (int32 k = 0; k < ParticlesPerColor.Num(); k++)
@@ -132,10 +128,10 @@ namespace Chaos::Softs
 				{
 					const int32 p = ParticlesPerColor[k][j];
 
-					if (Particles.InvM(p) != (T)0.)
+					if (Particles.InvM(p) != (FSolverReal)0.)
 					{
 						int32 ConstraintIndex = 0;
-						Chaos::PMatrix<T, 3, 3> ParticleHessian((T)0., (T)0., (T)0.);
+						Chaos::PMatrix<FSolverReal, 3, 3> ParticleHessian((FSolverReal)0., (FSolverReal)0., (FSolverReal)0.);
 
 						this->ComputeInitialResidualAndHessian(Particles, p, Dt, NewtonResidual[p], ParticleHessian);
 
@@ -169,7 +165,7 @@ namespace Chaos::Softs
 			}, ParticlesPerColor[k].Num() < 1000);
 		}
 
-		T NewtonNorm = (T)0.;
+		FSolverReal NewtonNorm = (FSolverReal)0.;
 		for (int32 p = 0; p < NewtonResidual.Num(); p++)
 		{
 			NewtonNorm += NewtonResidual[p].SizeSquared();
@@ -199,8 +195,7 @@ namespace Chaos::Softs
 
 	}
 
-	template <typename T, typename ParticleType>
-	void FGaussSeidelMainConstraint<T, ParticleType>::ApplyCG(ParticleType& Particles, const T Dt)
+	void FGaussSeidelMainConstraints::ApplyCG(FSolverParticles& Particles, const FSolverReal Dt)
 	{
 		TRACE_CPUPROFILER_EVENT_SCOPE(STAT_ChaosGaussSeidelApplyCG);
 		TFunction<void(TArray<FSolverVec3>&)> ProjectBCs = [&Particles](TArray<FSolverVec3>& y)
@@ -254,13 +249,4 @@ namespace Chaos::Softs
 				});
 		}
 	}
-}
-
-
-template class Chaos::Softs::FGaussSeidelMainConstraint<Chaos::Softs::FSolverReal, Chaos::Softs::FSolverParticles>;
-
-template CHAOS_API void Chaos::Softs::FGaussSeidelMainConstraint<Chaos::FRealDouble, Chaos::TDynamicParticles<Chaos::FRealDouble, 3>>::AddStaticConstraints(const TArray<TArray<int32>>& ExtraConstraints, TArray<TArray<int32>>& ExtraIncidentElements, TArray<TArray<int32>>& ExtraIncidentElementsLocal);
-
-template CHAOS_API void Chaos::Softs::FGaussSeidelMainConstraint<Chaos::FRealDouble, Chaos::TDynamicParticles<Chaos::FRealDouble, 3>>::AddTransientConstraints(const TArray<TArray<int32>>& ExtraConstraints, TArray<TArray<int32>>& ExtraIncidentElements, TArray<TArray<int32>>& ExtraIncidentElementsLocal, bool CheckIncidentElements);
-
-template CHAOS_API void Chaos::Softs::FGaussSeidelMainConstraint<Chaos::FRealDouble, Chaos::TDynamicParticles<Chaos::FRealDouble, 3>>::AddDynamicConstraints(const TArray<TArray<int32>>& ExtraConstraints, TArray<TArray<int32>>& ExtraIncidentElements, TArray<TArray<int32>>& ExtraIncidentElementsLocal, bool CheckIncidentElements);
+}
\ No newline at end of file
diff --git a/Engine/Source/Runtime/Experimental/Chaos/Public/Chaos/BlendedXPBDCorotatedConstraints.h b/Engine/Source/Runtime/Experimental/Chaos/Public/Chaos/BlendedXPBDCorotatedConstraints.h
--- a/Engine/Source/Runtime/Experimental/Chaos/Public/Chaos/BlendedXPBDCorotatedConstraints.h
+++ b/Engine/Source/Runtime/Experimental/Chaos/Public/Chaos/BlendedXPBDCorotatedConstraints.h
@@ -11,11 +11,10 @@
 
 namespace Chaos::Softs
 {
-	template <typename T, typename ParticleType>
-	class FBlendedXPBDCorotatedConstraints : public FXPBDCorotatedConstraints <T, ParticleType>
+	class FBlendedXPBDCorotatedTetrahedralConstraints : public FXPBDCorotatedTetrahedralConstraints
 	{
 
-		typedef FXPBDCorotatedConstraints<T, ParticleType> Base;
+		typedef FXPBDCorotatedTetrahedralConstraints Base;
 		using Base::MeshConstraints;
 		using Base::LambdaArray;
 		using Base::Measure;
@@ -25,66 +24,65 @@ namespace Chaos::Softs
 
 	public:
 		//this one only accepts tetmesh input and mesh
-		FBlendedXPBDCorotatedConstraints(
-			const ParticleType& InParticles,
+		FBlendedXPBDCorotatedTetrahedralConstraints(
+			const FSolverParticles& InParticles,
 			const TArray<TVector<int32, 4>>& InMesh,
 			const bool bRecordMetricIn = true,
-			const T& EMesh = (T)10.0,
-			const T& NuMesh = (T).3,
-			const T& InZeta = (T)1.
+			const FSolverReal& EMesh = (FSolverReal)10.0,
+			const FSolverReal& NuMesh = (FSolverReal).3,
+			const FSolverReal& InZeta = (FSolverReal)1.
 		)
 			: Base(InParticles, InMesh, bRecordMetricIn, EMesh, NuMesh), Zeta(InZeta)
 		{
-			if (Zeta < (T)0.) 
+			if (Zeta < (FSolverReal)0.) 
 			{
-				Zeta = (T)0.;
+				Zeta = (FSolverReal)0.;
 			}
-			if (Zeta > (T)1.) 
+			if (Zeta > (FSolverReal)1.) 
 			{
-				Zeta = (T)1.;
+				Zeta = (FSolverReal)1.;
 			}
 
-			C1Contribution.Init(TVec4<TVector<T, 3>>(TVector<T, 3>((T)0.)), InMesh.Num());
-			C2Contribution.Init(TVec4<TVector<T, 3>>(TVector<T, 3>((T)0.)), InMesh.Num());
+			C1Contribution.Init(TVec4<TVector<FSolverReal, 3>>(TVector<FSolverReal, 3>((FSolverReal)0.)), InMesh.Num());
+			C2Contribution.Init(TVec4<TVector<FSolverReal, 3>>(TVector<FSolverReal, 3>((FSolverReal)0.)), InMesh.Num());
 		}
 
 		//this one only accepts tetmesh input and mesh
-		FBlendedXPBDCorotatedConstraints(
-			const ParticleType& InParticles,
+		FBlendedXPBDCorotatedTetrahedralConstraints(
+			const FSolverParticles& InParticles,
 			const TArray<TVector<int32, 4>>& InMesh,
-			const TArray<T>& EMeshArray,
-			const T& NuMesh = (T).3,
+			const TArray<FSolverReal>& EMeshArray,
+			const FSolverReal& NuMesh = (FSolverReal).3,
 			const bool bRecordMetricIn = false,
-			const T& InZeta = (T)1.
+			const FSolverReal& InZeta = (FSolverReal)1.
 		)
 			: Base(InParticles, InMesh, EMeshArray, NuMesh, bRecordMetricIn), Zeta(InZeta)
 		{
-			if (Zeta < (T)0.)
+			if (Zeta < (FSolverReal)0.)
 			{
-				Zeta = (T)0.;
+				Zeta = (FSolverReal)0.;
 			}
-			if (Zeta > (T)1.)
+			if (Zeta > (FSolverReal)1.)
 			{
-				Zeta = (T)1.;
+				Zeta = (FSolverReal)1.;
 			}
 
-			C1Contribution.Init(TVec4<TVector<T, 3>>(TVector<T, 3>((T)0.)), InMesh.Num());
-			C2Contribution.Init(TVec4<TVector<T, 3>>(TVector<T, 3>((T)0.)), InMesh.Num());
+			C1Contribution.Init(TVec4<TVector<FSolverReal, 3>>(TVector<FSolverReal, 3>((FSolverReal)0.)), InMesh.Num());
+			C2Contribution.Init(TVec4<TVector<FSolverReal, 3>>(TVector<FSolverReal, 3>((FSolverReal)0.)), InMesh.Num());
 		}
 
-		virtual ~FBlendedXPBDCorotatedConstraints() {}
+		virtual ~FBlendedXPBDCorotatedTetrahedralConstraints() {}
 
-		template<typename SolverParticlesOrRange>
-		void ApplyInSerial(SolverParticlesOrRange& Particles, const T Dt, const int32 ElementIndex) const
+		virtual void ApplyInSerial(FSolverParticlesRange& Particles, const FSolverReal Dt, const int32 ElementIndex) const override
 		{
-			TVec4<TVector<T, 3>> PolarDelta = GetPolarDelta(Particles, Dt, ElementIndex);
+			TVec4<TVector<FSolverReal, 3>> PolarDelta = GetPolarDelta(Particles, Dt, ElementIndex);
 
 			for (int i = 0; i < 4; i++)
 			{
 				Particles.P(MeshConstraints[ElementIndex][i]) += PolarDelta[i];
 			}
 
-			TVec4<TVector<T, 3>> DetDelta = GetDeterminantDelta(Particles, Dt, ElementIndex);
+			TVec4<TVector<FSolverReal, 3>> DetDelta = GetDeterminantDelta(Particles, Dt, ElementIndex);
 
 			for (int i = 0; i < 4; i++)
 			{
@@ -94,33 +92,32 @@ namespace Chaos::Softs
 
 		virtual void Init() const override
 		{
-			for (T& Lambdas : LambdaArray) { Lambdas = (T)0.; }
-			C1Contribution.Init(TVec4<TVector<T, 3>>(TVector<T, 3>((T)0.)), MeshConstraints.Num());
-			C2Contribution.Init(TVec4<TVector<T, 3>>(TVector<T, 3>((T)0.)), MeshConstraints.Num());
+			for (FSolverReal& Lambdas : LambdaArray) { Lambdas = (FSolverReal)0.; }
+			C1Contribution.Init(TVec4<TVector<FSolverReal, 3>>(TVector<FSolverReal, 3>((FSolverReal)0.)), MeshConstraints.Num());
+			C2Contribution.Init(TVec4<TVector<FSolverReal, 3>>(TVector<FSolverReal, 3>((FSolverReal)0.)), MeshConstraints.Num());
 		}
 
 	protected:
-		template<typename SolverParticlesOrRange>
-		TVec4<TVector<T, 3>> GetDeterminantDelta(const SolverParticlesOrRange& Particles, const T Dt, const int32 ElementIndex, const T Tol = (T)1e-3) const
+		virtual TVec4<TVector<FSolverReal, 3>> GetDeterminantDelta(const FSolverParticlesRange& Particles, const FSolverReal Dt, const int32 ElementIndex, const FSolverReal Tol = (FSolverReal)1e-3) const override
 		{
 			TRACE_CPUPROFILER_EVENT_SCOPE(STAT_ChaosBlendedXPBDCorotatedApplyDet);
 
-			const PMatrix<T, 3, 3> Fe = Base::F(ElementIndex, Particles);
-			PMatrix<T, 3, 3> DmInvT = Base::ElementDmInv(ElementIndex).GetTransposed();
+			const PMatrix<FSolverReal, 3, 3> Fe = Base::F(ElementIndex, Particles);
+			PMatrix<FSolverReal, 3, 3> DmInvT = Base::ElementDmInv(ElementIndex).GetTransposed();
 
-			T J = Fe.Determinant();
-			if (J - (T)1. < Tol)
+			FSolverReal J = Fe.Determinant();
+			if (J - (FSolverReal)1. < Tol)
 			{
-				return TVec4<TVector<T, 3>>(TVector<T, 3>((T)0.));
+				return TVec4<TVector<FSolverReal, 3>>(TVector<FSolverReal, 3>((FSolverReal)0.));
 			}
 
-			TVec4<TVector<T, 3>> dC2 = Base::GetDeterminantGradient(Fe, DmInvT);
+			TVec4<TVector<FSolverReal, 3>> dC2 = Base::GetDeterminantGradient(Fe, DmInvT);
 
-			T AlphaTilde = (T)2. / (Dt * Dt * Base::Lambda * Measure[ElementIndex]);
+			FSolverReal AlphaTilde = (FSolverReal)2. / (Dt * Dt * Base::Lambda * Measure[ElementIndex]);
 
-			T DLambda = (1 - J) - AlphaTilde * LambdaArray[2 * ElementIndex + 1];
+			FSolverReal DLambda = (1 - J) - AlphaTilde * LambdaArray[2 * ElementIndex + 1];
 
-			T Denom = AlphaTilde;
+			FSolverReal Denom = AlphaTilde;
 			for (int i = 0; i < 4; i++)
 			{
 				for (int j = 0; j < 3; j++)
@@ -130,7 +127,7 @@ namespace Chaos::Softs
 			}
 			DLambda /= Denom;
 			LambdaArray[2 * ElementIndex + 1] += DLambda;
-			TVec4<TVector<T, 3>> Delta2(TVector<T, 3>((T)0.));
+			TVec4<TVector<FSolverReal, 3>> Delta2(TVector<FSolverReal, 3>((FSolverReal)0.));
 			for (int i = 0; i < 4; i++)
 			{
 				for (int j = 0; j < 3; j++)
@@ -138,7 +135,7 @@ namespace Chaos::Softs
 					Delta2[i][j] = Particles.InvM(MeshConstraints[ElementIndex][i]) * dC2[i][j] * DLambda;
 				}
 			}
-			TVec4<TVector<T, 3>> Delta1(TVector<T, 3>((T)0.));
+			TVec4<TVector<FSolverReal, 3>> Delta1(TVector<FSolverReal, 3>((FSolverReal)0.));
 			for (int i = 0; i < 4; i++)
 			{
 				for (int j = 0; j < 3; j++)
@@ -147,10 +144,10 @@ namespace Chaos::Softs
 						* dC2[i][j] * LambdaArray[2 * ElementIndex + 1] - C2Contribution[ElementIndex][i][j];
 				}
 			}
-			TVec4<TVector<T, 3>> Delta(TVector<T, 3>((T)0.));
+			TVec4<TVector<FSolverReal, 3>> Delta(TVector<FSolverReal, 3>((FSolverReal)0.));
 			for (int i = 0; i < 4; i++)
 			{
-				Delta[i] = Zeta * Delta1[i] + ((T)1. - Zeta) * Delta2[i];
+				Delta[i] = Zeta * Delta1[i] + ((FSolverReal)1. - Zeta) * Delta2[i];
 			}
 
 			for (int i = 0; i < 4; i++)
@@ -161,19 +158,17 @@ namespace Chaos::Softs
 			return Delta;
 		}
 
-
-		template<typename SolverParticlesOrRange>
-		TVec4<TVector<T, 3>> GetPolarDelta(const SolverParticlesOrRange& Particles, const T Dt, const int32 ElementIndex, const T Tol = (T)1e-3) const
+		virtual TVec4<TVector<FSolverReal, 3>> GetPolarDelta(const FSolverParticlesRange& Particles, const FSolverReal Dt, const int32 ElementIndex, const FSolverReal Tol = (FSolverReal)1e-3) const override
 		{
 			TRACE_CPUPROFILER_EVENT_SCOPE(STAT_ChaosBlendedXPBDCorotatedApplyPolar);
 			SCOPE_CYCLE_COUNTER(STAT_ChaosXPBDCorotatedPolar);
-			const PMatrix<T, 3, 3> Fe = Base::F(ElementIndex, Particles);
+			const PMatrix<FSolverReal, 3, 3> Fe = Base::F(ElementIndex, Particles);
 
-			PMatrix<T, 3, 3> Re((T)0.), Se((T)0.);
+			PMatrix<FSolverReal, 3, 3> Re((FSolverReal)0.), Se((FSolverReal)0.);
 
 			Chaos::PolarDecomposition(Fe, Re, Se);
 
-			T C1 = (T)0.;
+			FSolverReal C1 = (FSolverReal)0.;
 			for (int i = 0; i < 3; i++)
 			{
 				for (int j = 0; j < 3; j++)
@@ -185,18 +180,18 @@ namespace Chaos::Softs
 
 			if (C1 < Tol)
 			{
-				return TVec4<TVector<T, 3>>(TVector<T, 3>((T)0.));
+				return TVec4<TVector<FSolverReal, 3>>(TVector<FSolverReal, 3>((FSolverReal)0.));
 			}
 
-			PMatrix<T, 3, 3> DmInvT = Base::ElementDmInv(ElementIndex).GetTransposed();
+			PMatrix<FSolverReal, 3, 3> DmInvT = Base::ElementDmInv(ElementIndex).GetTransposed();
 
-			TVec4<TVector<T, 3>> dC1 = Base::GetPolarGradient(Fe, Re, DmInvT, C1);
+			TVec4<TVector<FSolverReal, 3>> dC1 = Base::GetPolarGradient(Fe, Re, DmInvT, C1);
 
-			T AlphaTilde = (T)1. / (Dt * Dt * Base::Mu * Measure[ElementIndex]);
+			FSolverReal AlphaTilde = (FSolverReal)1. / (Dt * Dt * Base::Mu * Measure[ElementIndex]);
 
-			T DLambda = -C1 - AlphaTilde * LambdaArray[2 * ElementIndex + 0];
+			FSolverReal DLambda = -C1 - AlphaTilde * LambdaArray[2 * ElementIndex + 0];
 
-			T Denom = AlphaTilde;
+			FSolverReal Denom = AlphaTilde;
 			for (int i = 0; i < 4; i++)
 			{
 				for (int j = 0; j < 3; j++)
@@ -207,7 +202,7 @@ namespace Chaos::Softs
 			DLambda /= Denom;
 			LambdaArray[2 * ElementIndex + 0] += DLambda;
 
-			TVec4<TVector<T, 3>> Delta2(TVector<T, 3>((T)0.));
+			TVec4<TVector<FSolverReal, 3>> Delta2(TVector<FSolverReal, 3>((FSolverReal)0.));
 			for (int i = 0; i < 4; i++)
 			{
 				for (int j = 0; j < 3; j++)
@@ -215,7 +210,7 @@ namespace Chaos::Softs
 					Delta2[i][j] = Particles.InvM(MeshConstraints[ElementIndex][i]) * dC1[i][j] * DLambda;
 				}
 			}
-			TVec4<TVector<T, 3>> Delta1(TVector<T, 3>((T)0.));
+			TVec4<TVector<FSolverReal, 3>> Delta1(TVector<FSolverReal, 3>((FSolverReal)0.));
 			for (int i = 0; i < 4; i++)
 			{
 				for (int j = 0; j < 3; j++)
@@ -224,10 +219,10 @@ namespace Chaos::Softs
 						* dC1[i][j] * LambdaArray[2 * ElementIndex + 0] - C1Contribution[ElementIndex][i][j];
 				}
 			}
-			TVec4<TVector<T, 3>> Delta(TVector<T, 3>((T)0.));
+			TVec4<TVector<FSolverReal, 3>> Delta(TVector<FSolverReal, 3>((FSolverReal)0.));
 			for (int i = 0; i < 4; i++)
 			{
-				Delta[i] = Zeta * Delta1[i] + ((T)1. - Zeta) * Delta2[i];
+				Delta[i] = Zeta * Delta1[i] + ((FSolverReal)1. - Zeta) * Delta2[i];
 			}
 
 			for (int i = 0; i < 4; i++)
@@ -242,10 +237,13 @@ namespace Chaos::Softs
 
 	private:
 
-		T Zeta;
-		mutable TArray<TVec4<TVector<T, 3>>> C1Contribution;
-		mutable TArray<TVec4<TVector<T, 3>>> C2Contribution;
+		FSolverReal Zeta;
+		mutable TArray<TVec4<TVector<FSolverReal, 3>>> C1Contribution;
+		mutable TArray<TVec4<TVector<FSolverReal, 3>>> C2Contribution;
 	};
 
+	template <typename T, typename ParticleType>
+	using FBlendedXPBDCorotatedConstraints UE_DEPRECATED(5.8, "Deprecated. this class is to be deleted, use FBlendedXPBDCorotatedTetrahedralConstraints instead") = FBlendedXPBDCorotatedTetrahedralConstraints;
+
 }  // End namespace Chaos::Softs
 
diff --git a/Engine/Source/Runtime/Experimental/Chaos/Public/Chaos/Deformable/ChaosDeformableConstraintsProxy.h b/Engine/Source/Runtime/Experimental/Chaos/Public/Chaos/Deformable/ChaosDeformableConstraintsProxy.h
--- a/Engine/Source/Runtime/Experimental/Chaos/Public/Chaos/Deformable/ChaosDeformableConstraintsProxy.h
+++ b/Engine/Source/Runtime/Experimental/Chaos/Public/Chaos/Deformable/ChaosDeformableConstraintsProxy.h
@@ -67,7 +67,7 @@ namespace Chaos::Softs
 		int32 SourceParticleIndex = INDEX_NONE;
 		int32 TargetParticleIndex = INDEX_NONE;
 
-		TArray<const Chaos::Softs::FGaussSeidelWeakConstraints<Softs::FSolverReal, Softs::FSolverParticles>::FGaussSeidelConstraintHandle*> Handles;
+		TArray<const Chaos::Softs::FGaussSeidelSpringConstraints::FGaussSeidelConstraintHandle*> Handles;
 	};
 
 
diff --git a/Engine/Source/Runtime/Experimental/Chaos/Public/Chaos/Deformable/ChaosDeformableSolver.h b/Engine/Source/Runtime/Experimental/Chaos/Public/Chaos/Deformable/ChaosDeformableSolver.h
--- a/Engine/Source/Runtime/Experimental/Chaos/Public/Chaos/Deformable/ChaosDeformableSolver.h
+++ b/Engine/Source/Runtime/Experimental/Chaos/Public/Chaos/Deformable/ChaosDeformableSolver.h
@@ -203,19 +203,19 @@ namespace Chaos::Softs
 
 		// Simulation Variables
 		TUniquePtr<Softs::FPBDEvolution> Evolution;
-		TArray<TUniquePtr<Softs::FXPBDCorotatedConstraints<Softs::FSolverReal, Softs::FSolverParticles>>> CorotatedConstraints;
-		TUniquePtr<Softs::FGaussSeidelCorotatedConstraints<Softs::FSolverReal, Softs::FSolverParticles>> GSCorotatedConstraints;
-		TUniquePtr<Softs::FGaussSeidelNeohookeanConstraints<Softs::FSolverReal, Softs::FSolverParticles>> GSNeohookeanConstraints;
-		TUniquePtr<Softs::FGaussSeidelCorotatedCodimensionalConstraints<Softs::FSolverReal, Softs::FSolverParticles>> GSCorotatedCodConstraints;
-		TUniquePtr<Softs::FGaussSeidelLinearCodimensionalConstraints<Softs::FSolverReal, Softs::FSolverParticles>> GSLinearCodConstraints;
-		TUniquePtr<Softs::FGaussSeidelWeakConstraints<Softs::FSolverReal, Softs::FSolverParticles>> GSWeakConstraints;
-		TUniquePtr<Softs::FGaussSeidelDynamicWeakConstraints<Softs::FSolverReal, Softs::FSolverParticles>> GSDynamicWeakConstraints;
-		TUniquePtr<Softs::FGaussSeidelSphereRepulsionConstraints<Softs::FSolverReal, Softs::FSolverParticles>> GSSphereRepulsionConstraints;
-		TUniquePtr<Softs::FGaussSeidelUnilateralTetConstraints<Softs::FSolverReal, Softs::FSolverParticles>> GSVolumeConstraints;
-		TArray<TUniquePtr<Softs::FXPBDWeakConstraints<Softs::FSolverReal, Softs::FSolverParticles>>> WeakConstraints;
-		TArray<TUniquePtr<Softs::FBlendedXPBDCorotatedConstraints<Softs::FSolverReal, Softs::FSolverParticles>>> BlendedCorotatedConstraints;
-		TUniquePtr<Softs::FXPBDGridBasedCorotatedConstraints<Softs::FSolverReal, Softs::FSolverParticles>> GridBasedCorotatedConstraint;
-		TUniquePtr<Softs::FGaussSeidelMainConstraint<Softs::FSolverReal, Softs::FSolverParticles>> GSMainConstraint;
+		TArray<TUniquePtr<Softs::FXPBDCorotatedTetrahedralConstraints>> CorotatedConstraints;
+		TUniquePtr<Softs::FGaussSeidelCorotatedTetrahedralConstraints> GSCorotatedConstraints;
+		TUniquePtr<Softs::FGaussSeidelNeohookeanTetrahedralConstraints> GSNeohookeanConstraints;
+		TUniquePtr<Softs::FGaussSeidelCorotatedTriangleConstraints> GSCorotatedCodConstraints;
+		TUniquePtr<Softs::FGaussSeidelLinearTriangleConstraints> GSLinearCodConstraints;
+		TUniquePtr<Softs::FGaussSeidelSpringConstraints> GSWeakConstraints;
+		TUniquePtr<Softs::FGaussSeidelDynamicSpringConstraints> GSDynamicWeakConstraints;
+		TUniquePtr<Softs::FGaussSeidelParticleSphereRepulsionConstraints> GSSphereRepulsionConstraints;
+		TUniquePtr<Softs::FGaussSeidelUnilateralVolumeConstraints> GSVolumeConstraints;
+		TArray<TUniquePtr<Softs::FXPBDWeightedSpringConstraints>> WeakConstraints;
+		TArray<TUniquePtr<Softs::FBlendedXPBDCorotatedTetrahedralConstraints>> BlendedCorotatedConstraints;
+		TUniquePtr<Softs::FXPBDGridBasedCorotatedTetrahedralConstraints> GridBasedCorotatedConstraint;
+		TUniquePtr<Softs::FGaussSeidelMainConstraints> GSMainConstraint;
 		TUniquePtr<Softs::FPBDCollisionSpringConstraints> CollisionSpringConstraint;
 		TUniquePtr<Softs::FPBDTriangleMeshCollisions> TriangleMeshCollisions;
 		TArrayCollectionArray<const UObject*> MObjects;
@@ -240,7 +240,7 @@ namespace Chaos::Softs
 		TUniquePtr <TArray<int32>> ParticleComponentIndex;
 		TMap<int32, TSet<int32>> ParticleTriangleExclusionMap;
 		//Muscle Activation Variables
-		TUniquePtr<Softs::FMuscleActivationConstraints<Softs::FSolverReal, Softs::FSolverParticles>> MuscleActivationConstraints;		
+		TUniquePtr<Softs::FMuscleActivationTetrahedralConstraints> MuscleActivationConstraints;
 		TMap<FThreadingProxy::FKey, int32> MuscleIndexOffset;
 
 		typedef TMap<int32, TTuple<float, float, FVector3f>> TransientConstraintBufferMap;
diff --git a/Engine/Source/Runtime/Experimental/Chaos/Public/Chaos/Deformable/GaussSeidelCorotatedCodimensionalConstraints.h b/Engine/Source/Runtime/Experimental/Chaos/Public/Chaos/Deformable/GaussSeidelCorotatedCodimensionalConstraints.h
--- a/Engine/Source/Runtime/Experimental/Chaos/Public/Chaos/Deformable/GaussSeidelCorotatedCodimensionalConstraints.h
+++ b/Engine/Source/Runtime/Experimental/Chaos/Public/Chaos/Deformable/GaussSeidelCorotatedCodimensionalConstraints.h
@@ -12,23 +12,23 @@
 
 namespace Chaos::Softs
 {
-	template <typename T, typename ParticleType>
-	class FGaussSeidelCorotatedCodimensionalConstraints
+	// Triangle constraint using fixed corotated constitutive model [Stomakhin et al., 2012]
+	// 3D triangles are projected to the 2D space for deformation gradient (3*2 matrix) computation
+	class FGaussSeidelCorotatedTriangleConstraints
 	{
-
 	public:
-		FGaussSeidelCorotatedCodimensionalConstraints(
-			const ParticleType& InParticles,
+		FGaussSeidelCorotatedTriangleConstraints(
+			const FSolverParticles& InParticles,
 			const TArray<TVector<int32, 3>>& InMesh,
 			const bool bRecordMetricIn = true,
-			const T& EMesh = (T)10.0,
-			const T& NuMesh = (T).3
+			const FSolverReal& EMesh = (FSolverReal)10.0,
+			const FSolverReal& NuMesh = (FSolverReal).3
 		)
 			: bRecordMetric(bRecordMetricIn), MeshConstraints(InMesh)
 		{
-			Measure.Init((T)0., MeshConstraints.Num());
-			Lambda = EMesh * NuMesh / (((T)1. + NuMesh) * ((T)1. - (T)2. * NuMesh));
-			Mu = EMesh / ((T)2. * ((T)1. + NuMesh));
+			Measure.Init((FSolverReal)0., MeshConstraints.Num());
+			Lambda = EMesh * NuMesh / (((FSolverReal)1. + NuMesh) * ((FSolverReal)1. - (FSolverReal)2. * NuMesh));
+			Mu = EMesh / ((FSolverReal)2. * ((FSolverReal)1. + NuMesh));
 
 			MuElementArray.Init(Mu, MeshConstraints.Num());
 			LambdaElementArray.Init(Lambda, MeshConstraints.Num());
@@ -36,33 +36,33 @@ namespace Chaos::Softs
 			InitializeCodimensionData(InParticles);
 		}
 
-		FGaussSeidelCorotatedCodimensionalConstraints(
-			const ParticleType& InParticles,
+		FGaussSeidelCorotatedTriangleConstraints(
+			const FSolverParticles& InParticles,
 			const TArray<TVector<int32, 3>>& InMesh,
-			const TArray<T>& EMeshArray,
-			const T& NuMesh = (T).3
+			const TArray<FSolverReal>& EMeshArray,
+			const FSolverReal& NuMesh = (FSolverReal).3
 		)
 			: MeshConstraints(InMesh)
 		{
 			ensureMsgf(EMeshArray.Num() == InMesh.Num(), TEXT("Input Young Modulus Array Size is wrong"));
-			Measure.Init((T)0., MeshConstraints.Num());
+			Measure.Init((FSolverReal)0., MeshConstraints.Num());
 
 			InitializeCodimensionData(InParticles);
-			LambdaElementArray.Init((T)0., MeshConstraints.Num());
-			MuElementArray.Init((T)0., MeshConstraints.Num());
+			LambdaElementArray.Init((FSolverReal)0., MeshConstraints.Num());
+			MuElementArray.Init((FSolverReal)0., MeshConstraints.Num());
 			
 			for (int32 e = 0; e < InMesh.Num(); e++)
 			{
-				LambdaElementArray[e] = EMeshArray[e] * NuMesh / (((T)1. + NuMesh) * ((T)1. - (T)2. * NuMesh));
-				MuElementArray[e] = EMeshArray[e] / ((T)2. * ((T)1. + NuMesh));
+				LambdaElementArray[e] = EMeshArray[e] * NuMesh / (((FSolverReal)1. + NuMesh) * ((FSolverReal)1. - (FSolverReal)2. * NuMesh));
+				MuElementArray[e] = EMeshArray[e] / ((FSolverReal)2. * ((FSolverReal)1. + NuMesh));
 			}
 		}
 
-		virtual ~FGaussSeidelCorotatedCodimensionalConstraints() {}
+		virtual ~FGaussSeidelCorotatedTriangleConstraints() {}
 
-		PMatrix<T, 3, 3> DsInit(const int32 E, const ParticleType& InParticles)  
+		PMatrix<FSolverReal, 3, 3> DsInit(const int32 E, const FSolverParticles& InParticles)  
 		{
-			PMatrix<T, 3, 3> Result((T)0.);
+			PMatrix<FSolverReal, 3, 3> Result((FSolverReal)0.);
 			for (int32 i = 0; i < 3; i++) 
 			{
 				for (int32 c = 0; c < 3; c++) 
@@ -73,30 +73,30 @@ namespace Chaos::Softs
 			return Result;
 		}
 
-		PMatrix<T, 3, 2> Ds(const int32 E, const ParticleType& InParticles) const 
+		PMatrix<FSolverReal, 3, 2> Ds(const int32 E, const FSolverParticles& InParticles) const 
 		{
 			if (INDEX_NONE < E && E< MeshConstraints.Num()
 				&& INDEX_NONE < MeshConstraints[E][0] && MeshConstraints[E][0] < (int32)InParticles.Size()
 				&& INDEX_NONE < MeshConstraints[E][1] && MeshConstraints[E][1] < (int32)InParticles.Size()
 				&& INDEX_NONE < MeshConstraints[E][2] && MeshConstraints[E][2] < (int32)InParticles.Size())
 			{
-				const TVec3<T> P1P0 = InParticles.P(MeshConstraints[E][1]) - InParticles.P(MeshConstraints[E][0]);
-				const TVec3<T> P2P0 = InParticles.P(MeshConstraints[E][2]) - InParticles.P(MeshConstraints[E][0]);
+				const TVec3<FSolverReal> P1P0 = InParticles.P(MeshConstraints[E][1]) - InParticles.P(MeshConstraints[E][0]);
+				const TVec3<FSolverReal> P2P0 = InParticles.P(MeshConstraints[E][2]) - InParticles.P(MeshConstraints[E][0]);
 
-				return PMatrix<T, 3, 2>(
+				return PMatrix<FSolverReal, 3, 2>(
 					P1P0[0], P1P0[1], P1P0[2],
 					P2P0[0], P2P0[1], P2P0[2]);
 			}
 			else
 			{
-				return PMatrix<T, 3, 2>(
-					(T)0., (T)0, (T)0,
-					(T)0, (T)0, (T)0);
+				return PMatrix<FSolverReal, 3, 2>(
+					(FSolverReal)0., (FSolverReal)0, (FSolverReal)0,
+					(FSolverReal)0, (FSolverReal)0, (FSolverReal)0);
 			}
 		
 		}
 
-		PMatrix<T, 3, 2> F(const int32 E, const ParticleType& InParticles) const 
+		PMatrix<FSolverReal, 3, 2> F(const int32 E, const FSolverParticles& InParticles) const 
 		{
 			if (INDEX_NONE < E && E < DmInverse.Num())
 			{
@@ -104,9 +104,9 @@ namespace Chaos::Softs
 			}
 			else
 			{
-				return PMatrix<T, 3, 2>(
-					(T)0., (T)0, (T)0,
-					(T)0, (T)0, (T)0);
+				return PMatrix<FSolverReal, 3, 2>(
+					(FSolverReal)0., (FSolverReal)0, (FSolverReal)0,
+					(FSolverReal)0, (FSolverReal)0, (FSolverReal)0);
 			}
 		}
 
@@ -125,15 +125,15 @@ namespace Chaos::Softs
 			return Constraints;
 		}
 
-		void ComputeNodalMass(const T InDensity, const int32 NumParticles, TArray<T>& NodalMass) const
+		void ComputeNodalMass(const FSolverReal InDensity, const int32 NumParticles, TArray<FSolverReal>& NodalMass) const
 		{
-			NodalMass.Init((T)0., NumParticles);
+			NodalMass.Init((FSolverReal)0., NumParticles);
 			for (int32 e = 0; e < Measure.Num(); e++)
 			{
-				const T TotalMass = InDensity * Measure[e];
+				const FSolverReal TotalMass = InDensity * Measure[e];
 				for (int32 j = 0; j < 3; j++)
 				{
-					NodalMass[MeshConstraints[e][j]] += TotalMass / (T)3.;
+					NodalMass[MeshConstraints[e][j]] += TotalMass / (FSolverReal)3.;
 				}
 			}
 		}
@@ -142,9 +142,9 @@ namespace Chaos::Softs
 
 	protected:
 
-		void InitializeCodimensionData(const ParticleType& Particles)
+		void InitializeCodimensionData(const FSolverParticles& Particles)
 		{
-			Measure.Init((T)0., MeshConstraints.Num());
+			Measure.Init((FSolverReal)0., MeshConstraints.Num());
 			DmInverse.Init(PMatrix<FSolverReal, 2, 2>(0.f, 0.f, 0.f), MeshConstraints.Num());
 			for (int32 e = 0; e < MeshConstraints.Num(); e++)
 			{
@@ -155,31 +155,31 @@ namespace Chaos::Softs
 					&& MeshConstraints[e][1] < (int32)Particles.Size() && MeshConstraints[e][1] > INDEX_NONE
 					&& MeshConstraints[e][2] < (int32)Particles.Size() && MeshConstraints[e][2] > INDEX_NONE)
 				{
-					const TVec3<T> X1X0 = Particles.GetX(MeshConstraints[e][1]) - Particles.GetX(MeshConstraints[e][0]);
-					const TVec3<T> X2X0 = Particles.GetX(MeshConstraints[e][2]) - Particles.GetX(MeshConstraints[e][0]);
-					PMatrix<T, 2, 2> Dm((T)0., (T)0., (T)0.);
+					const TVec3<FSolverReal> X1X0 = Particles.GetX(MeshConstraints[e][1]) - Particles.GetX(MeshConstraints[e][0]);
+					const TVec3<FSolverReal> X2X0 = Particles.GetX(MeshConstraints[e][2]) - Particles.GetX(MeshConstraints[e][0]);
+					PMatrix<FSolverReal, 2, 2> Dm((FSolverReal)0., (FSolverReal)0., (FSolverReal)0.);
 					Dm.M[0] = X1X0.Size();
 					Dm.M[2] = X1X0.Dot(X2X0) / Dm.M[0];
-					Dm.M[3] = Chaos::TVector<T, 3>::CrossProduct(X1X0, X2X0).Size() / Dm.M[0];
-					Measure[e] = Chaos::TVector<T, 3>::CrossProduct(X1X0, X2X0).Size() / (T)2.;
-					ensureMsgf(Measure[e] > (T)0., TEXT("Degenerate triangle detected"));
+					Dm.M[3] = Chaos::TVector<FSolverReal, 3>::CrossProduct(X1X0, X2X0).Size() / Dm.M[0];
+					Measure[e] = Chaos::TVector<FSolverReal, 3>::CrossProduct(X1X0, X2X0).Size() / (FSolverReal)2.;
+					ensureMsgf(Measure[e] > (FSolverReal)0., TEXT("Degenerate triangle detected"));
 
-					PMatrix<T, 2, 2> DmInv = Dm.Inverse();
+					PMatrix<FSolverReal, 2, 2> DmInv = Dm.Inverse();
 					DmInverse[e] = DmInv;
 				}
 			}
 		}
 
-		static T SafeRecip(const T Len, const T Fallback)
+		static FSolverReal SafeRecip(const FSolverReal Len, const FSolverReal Fallback)
 		{
-			if (Len > (T)UE_SMALL_NUMBER)
+			if (Len > (FSolverReal)UE_SMALL_NUMBER)
 			{
-				return (T)1. / Len;
+				return (FSolverReal)1. / Len;
 			}
 			return Fallback;
 		}
 
-		static inline void SymSchur2D(const T Aqq, const T App, const T Apq, T& c, T& s)
+		static inline void SymSchur2D(const FSolverReal Aqq, const FSolverReal App, const FSolverReal Apq, FSolverReal& c, FSolverReal& s)
 		{
 			// A is an n x n matrix
 			// (p, q) is an index pair satisfying 1 <= p < q <= n
@@ -193,26 +193,26 @@ namespace Chaos::Softs
 
 			if (Apq != 0) 
 			{
-				T Tau = T(0.5) * (Aqq - App) / Apq;
-				T t;
+				FSolverReal Tau = FSolverReal(0.5) * (Aqq - App) / Apq;
+				FSolverReal t;
 				if (Tau >= 0)
-					t = T(1) / (Tau + FMath::Sqrt((T)1. + Tau * Tau));
+					t = FSolverReal(1) / (Tau + FMath::Sqrt((FSolverReal)1. + Tau * Tau));
 				else
-					t = T(-1) / (-Tau + FMath::Sqrt((T)1. + Tau * Tau));
-				c = T(1) / FMath::Sqrt((T)1. + t * t);
+					t = FSolverReal(-1) / (-Tau + FMath::Sqrt((FSolverReal)1. + Tau * Tau));
+				c = FSolverReal(1) / FMath::Sqrt((FSolverReal)1. + t * t);
 				s = t * c;
 			}
 			else 
 			{
-				c = (T)1.;
-				s = (T)0.;
+				c = (FSolverReal)1.;
+				s = (FSolverReal)0.;
 			}
 		}
 
 
-		static inline void Jacobi(const PMatrix<T, 2, 2>& B, PMatrix<T, 2, 2>& D, PMatrix<T, 2, 2>& V)
+		static inline void Jacobi(const PMatrix<FSolverReal, 2, 2>& B, PMatrix<FSolverReal, 2, 2>& D, PMatrix<FSolverReal, 2, 2>& V)
 		{
-			T c, s;
+			FSolverReal c, s;
 			SymSchur2D(B.M[3], B.M[0], B.M[1], c, s);
 
 			V.M[0] = c;
@@ -223,76 +223,76 @@ namespace Chaos::Softs
 			D = V.GetTransposed() * B * V;
 		}
 
-		static PMatrix<T, 3, 2> ComputeR(const PMatrix<T, 3, 2>& Fe) 
+		static PMatrix<FSolverReal, 3, 2> ComputeR(const PMatrix<FSolverReal, 3, 2>& Fe) 
 		{
-			PMatrix<T, 2, 2> FTF = ComputeFTF(Fe), D((T)0., (T)0., (T)0.), U((T)0., (T)0., (T)0.);
+			PMatrix<FSolverReal, 2, 2> FTF = ComputeFTF(Fe), D((FSolverReal)0., (FSolverReal)0., (FSolverReal)0.), U((FSolverReal)0., (FSolverReal)0., (FSolverReal)0.);
 			Jacobi(FTF, D, U);
-			PMatrix<T, 2, 2> SDiag(FMath::Sqrt(D.M[0]), (T)0., (T)0., FMath::Sqrt(D.M[3]));
-			PMatrix<T, 2, 2> S = U * SDiag * U.GetTransposed(), SInv = S.Inverse();;
+			PMatrix<FSolverReal, 2, 2> SDiag(FMath::Sqrt(D.M[0]), (FSolverReal)0., (FSolverReal)0., FMath::Sqrt(D.M[3]));
+			PMatrix<FSolverReal, 2, 2> S = U * SDiag * U.GetTransposed(), SInv = S.Inverse();;
 			return Fe * S.Inverse();
 		}
 
-		static PMatrix<T, 3, 2> ComputeRSimple(const PMatrix<T, 3, 2>& InputMatrix) 
+		static PMatrix<FSolverReal, 3, 2> ComputeRSimple(const PMatrix<FSolverReal, 3, 2>& InputMatrix) 
 		{
-			TVec3<T> Col1(InputMatrix.M[0], InputMatrix.M[1], InputMatrix.M[2]), Col2(InputMatrix.M[3], InputMatrix.M[4], InputMatrix.M[5]), a(T(0)), b(T(0));
-			if (FMath::Abs(Col1[0]) < (T)UE_SMALL_NUMBER && FMath::Abs(Col1[1]) < (T)UE_SMALL_NUMBER)
+			TVec3<FSolverReal> Col1(InputMatrix.M[0], InputMatrix.M[1], InputMatrix.M[2]), Col2(InputMatrix.M[3], InputMatrix.M[4], InputMatrix.M[5]), a(FSolverReal(0)), b(FSolverReal(0));
+			if (FMath::Abs(Col1[0]) < (FSolverReal)UE_SMALL_NUMBER && FMath::Abs(Col1[1]) < (FSolverReal)UE_SMALL_NUMBER)
 			{
-				a[0] = T(1);
+				a[0] = FSolverReal(1);
 			}
 			else
 			{
 				a[0] = -Col1[1];
 				a[1] = Col1[0];
-				const T OneOveraNorm = SafeRecip(a.Length(), 0.f);
+				const FSolverReal OneOveraNorm = SafeRecip(a.Length(), 0.f);
 				a *= OneOveraNorm;
 			}
-			b = TVec3<T>::CrossProduct(a, Col2);
-			const T OneOverbNorm = SafeRecip(b.Length(), 0.f);
+			b = TVec3<FSolverReal>::CrossProduct(a, Col2);
+			const FSolverReal OneOverbNorm = SafeRecip(b.Length(), 0.f);
 			b *= OneOverbNorm;
-			return Chaos::PMatrix<T, 3, 2>(b[0], b[1], b[2], a[0], a[1], a[2]);
+			return Chaos::PMatrix<FSolverReal, 3, 2>(b[0], b[1], b[2], a[0], a[1], a[2]);
 		}
 
 		//Does Gram Schmidt on a 3*2 matrix and returns an orthogonal one
-		static Chaos::PMatrix<T, 3, 2> GramSchmidt(const Chaos::PMatrix<T, 3, 2>& InputMatrix)
+		static Chaos::PMatrix<FSolverReal, 3, 2> GramSchmidt(const Chaos::PMatrix<FSolverReal, 3, 2>& InputMatrix)
 		{
-			TVec3<T> Col1(InputMatrix.M[0], InputMatrix.M[1], InputMatrix.M[2]), Col2(InputMatrix.M[3], InputMatrix.M[4], InputMatrix.M[5]);
-			const T Col1Norm = Col1.Length();
-			const T OneOverCol1Norm = SafeRecip(Col1Norm, 0.f);
-			TVec3<T> Col1Normalized = Col1 * OneOverCol1Norm;
-			TVec3<T> Col2Orthogonal = Col2 - TVec3<T>::DotProduct(Col1Normalized, Col2) * Col1Normalized;
-			const T Col2OrthogonalNorm = Col2Orthogonal.Length();
-			const T OneOverCol2OrthogonalNorm = SafeRecip(Col2OrthogonalNorm, 0.f);
+			TVec3<FSolverReal> Col1(InputMatrix.M[0], InputMatrix.M[1], InputMatrix.M[2]), Col2(InputMatrix.M[3], InputMatrix.M[4], InputMatrix.M[5]);
+			const FSolverReal Col1Norm = Col1.Length();
+			const FSolverReal OneOverCol1Norm = SafeRecip(Col1Norm, 0.f);
+			TVec3<FSolverReal> Col1Normalized = Col1 * OneOverCol1Norm;
+			TVec3<FSolverReal> Col2Orthogonal = Col2 - TVec3<FSolverReal>::DotProduct(Col1Normalized, Col2) * Col1Normalized;
+			const FSolverReal Col2OrthogonalNorm = Col2Orthogonal.Length();
+			const FSolverReal OneOverCol2OrthogonalNorm = SafeRecip(Col2OrthogonalNorm, 0.f);
 			Col2Orthogonal *= OneOverCol2OrthogonalNorm;
-			return Chaos::PMatrix<T, 3, 2>(Col1Normalized[0], Col1Normalized[1], Col1Normalized[2], Col2Orthogonal[0], Col2Orthogonal[1], Col2Orthogonal[2]);
+			return Chaos::PMatrix<FSolverReal, 3, 2>(Col1Normalized[0], Col1Normalized[1], Col1Normalized[2], Col2Orthogonal[0], Col2Orthogonal[1], Col2Orthogonal[2]);
 		}
 
-		static PMatrix<T, 2, 2> ComputeFTF(const PMatrix<T, 3, 2>& InputMatrix)
+		static PMatrix<FSolverReal, 2, 2> ComputeFTF(const PMatrix<FSolverReal, 3, 2>& InputMatrix)
 		{
-			return PMatrix<T, 2, 2>(InputMatrix.M[0] * InputMatrix.M[0] + InputMatrix.M[1] * InputMatrix.M[1] + InputMatrix.M[2] * InputMatrix.M[2],
+			return PMatrix<FSolverReal, 2, 2>(InputMatrix.M[0] * InputMatrix.M[0] + InputMatrix.M[1] * InputMatrix.M[1] + InputMatrix.M[2] * InputMatrix.M[2],
 				InputMatrix.M[0] * InputMatrix.M[3] + InputMatrix.M[1] * InputMatrix.M[4] + InputMatrix.M[2] * InputMatrix.M[5],
 				InputMatrix.M[3] * InputMatrix.M[3] + InputMatrix.M[4] * InputMatrix.M[4] + InputMatrix.M[5] * InputMatrix.M[5]);
 		}
 
 		//Computes det(FTF)^{1/2} * F * (FTF)^-1
-		static void dJdF32(const PMatrix<T, 3, 2>& F, PMatrix<T, 3, 2>& dJ)
+		static void dJdF32(const PMatrix<FSolverReal, 3, 2>& F, PMatrix<FSolverReal, 3, 2>& dJ)
 		{
-			const PMatrix<T, 2, 2> FTF = ComputeFTF(F);
-			const T J2 = FTF.Determinant();
+			const PMatrix<FSolverReal, 2, 2> FTF = ComputeFTF(F);
+			const FSolverReal J2 = FTF.Determinant();
 
-			const PMatrix<T, 2, 2> J2invT(FTF.M[3], -FTF.M[2], -FTF.M[1], FTF.M[0]);
+			const PMatrix<FSolverReal, 2, 2> J2invT(FTF.M[3], -FTF.M[2], -FTF.M[1], FTF.M[0]);
 
-			if (J2 > T(0))
+			if (J2 > FSolverReal(0))
 			{
-				dJ = (T(1) / FMath::Sqrt(J2)) *  (F * J2invT);
+				dJ = (FSolverReal(1) / FMath::Sqrt(J2)) *  (F * J2invT);
 			}
 			else
 			{
-				dJ = PMatrix<T, 3, 2>(TVec3<T>((T)0.), TVec3<T>((T)0.));
+				dJ = PMatrix<FSolverReal, 3, 2>(TVec3<FSolverReal>((FSolverReal)0.), TVec3<FSolverReal>((FSolverReal)0.));
 			}
 		}
 
 	public:
-		void AddHyperelasticResidualAndHessian(const ParticleType& Particles, const int32 ElementIndex, const int32 ElementIndexLocal, const T Dt, TVec3<T>& ParticleResidual, Chaos::PMatrix<T, 3, 3>& ParticleHessian)
+		void AddHyperelasticResidualAndHessian(const FSolverParticles& Particles, const int32 ElementIndex, const int32 ElementIndexLocal, const FSolverReal Dt, TVec3<FSolverReal>& ParticleResidual, Chaos::PMatrix<FSolverReal, 3, 3>& ParticleHessian)
 		{
 			check(ElementIndex < DmInverse.Num() && ElementIndex > INDEX_NONE);
 			check(ElementIndex < MuElementArray.Num());
@@ -307,19 +307,19 @@ namespace Chaos::Softs
 				&& ElementIndex < DmInverse.Num() 
 				&& ElementIndex > INDEX_NONE)
 			{
-				const Chaos::PMatrix<T, 2, 2> DmInvT = DmInverse[ElementIndex].GetTransposed(), DmInv = DmInverse[ElementIndex]; 
-				const Chaos::PMatrix<T, 3, 2> Fe = F(ElementIndex, Particles);
+				const Chaos::PMatrix<FSolverReal, 2, 2> DmInvT = DmInverse[ElementIndex].GetTransposed(), DmInv = DmInverse[ElementIndex]; 
+				const Chaos::PMatrix<FSolverReal, 3, 2> Fe = F(ElementIndex, Particles);
 
-				const PMatrix<T, 3, 2> Re = ComputeR(Fe);
-				const PMatrix<T, 2, 2> FTF = ComputeFTF(Fe);
-				const T J = FMath::Sqrt(FTF.Determinant());
+				const PMatrix<FSolverReal, 3, 2> Re = ComputeR(Fe);
+				const PMatrix<FSolverReal, 2, 2> FTF = ComputeFTF(Fe);
+				const FSolverReal J = FMath::Sqrt(FTF.Determinant());
 
-				PMatrix<T, 3, 2> Pe(TVec3<T>((T)0.), TVec3<T>((T)0.)), ForceTerm(TVec3<T>((T)0.), TVec3<T>((T)0.));
-				Pe = T(2) * MuElementArray[ElementIndex] * (Fe - Re) + LambdaElementArray[ElementIndex] * (J - T(1)) * J * Fe * ComputeFTF(Fe).Inverse(); 
+				PMatrix<FSolverReal, 3, 2> Pe(TVec3<FSolverReal>((FSolverReal)0.), TVec3<FSolverReal>((FSolverReal)0.)), ForceTerm(TVec3<FSolverReal>((FSolverReal)0.), TVec3<FSolverReal>((FSolverReal)0.));
+				Pe = FSolverReal(2) * MuElementArray[ElementIndex] * (Fe - Re) + LambdaElementArray[ElementIndex] * (J - FSolverReal(1)) * J * Fe * ComputeFTF(Fe).Inverse(); 
 
 				ForceTerm = -Measure[ElementIndex] * Pe * DmInverse[ElementIndex].GetTransposed();
 
-				Chaos::TVector<T, 3> Dx((T)0.);
+				Chaos::TVector<FSolverReal, 3> Dx((FSolverReal)0.);
 				if (ElementIndexLocal > 0)
 				{
 					for (int32 c = 0; c < 3; c++)
@@ -342,18 +342,18 @@ namespace Chaos::Softs
 
 				ParticleResidual -= Dx;
 
-				PMatrix<T, 3, 2> dJ(TVec3<T>((T)0.), TVec3<T>((T)0.));
+				PMatrix<FSolverReal, 3, 2> dJ(TVec3<FSolverReal>((FSolverReal)0.), TVec3<FSolverReal>((FSolverReal)0.));
 
 				dJdF32(Fe, dJ);
 
-				T Coeff = Dt * Dt * Measure[ElementIndex];
+				FSolverReal Coeff = Dt * Dt * Measure[ElementIndex];
 
 				if (ElementIndexLocal == 0) 
 				{
-					T DmInvsum = T(0);
+					FSolverReal DmInvsum = FSolverReal(0);
 					for (int32 nu = 0; nu < 2; nu++) 
 					{
-						T localDmsum = T(0);
+						FSolverReal localDmsum = FSolverReal(0);
 						for (int32 k = 0; k < 2; k++) 
 						{
 							localDmsum += DmInv.M[nu * 2 + k];
@@ -362,11 +362,11 @@ namespace Chaos::Softs
 					}
 					for (int32 alpha = 0; alpha < 3; alpha++) 
 					{
-						ParticleHessian.SetAt(alpha, alpha, ParticleHessian.GetAt(alpha, alpha) + Coeff * T(2) * MuElementArray[ElementIndex] * DmInvsum);
+						ParticleHessian.SetAt(alpha, alpha, ParticleHessian.GetAt(alpha, alpha) + Coeff * FSolverReal(2) * MuElementArray[ElementIndex] * DmInvsum);
 					}
 
-					const PMatrix<T, 3, 2> dJDmInvT = dJ * DmInvT; 
-					Chaos::TVector<T, 3> L((T)0.);
+					const PMatrix<FSolverReal, 3, 2> dJDmInvT = dJ * DmInvT; 
+					Chaos::TVector<FSolverReal, 3> L((FSolverReal)0.);
 					for (int32 alpha = 0; alpha < 3; alpha++) 
 					{
 						for (int32 k = 0; k < 2; k++) 
@@ -381,19 +381,19 @@ namespace Chaos::Softs
 				}
 				else 
 				{
-					T DmInvsum = T(0);
+					FSolverReal DmInvsum = FSolverReal(0);
 					for (int32 nu = 0; nu < 2; nu++)
 					{
 						DmInvsum += DmInv.GetAt(ElementIndexLocal - 1, nu) * DmInv.GetAt(ElementIndexLocal - 1, nu);
 					}
 					for (int32 alpha = 0; alpha < 3; alpha++)
 					{
-						ParticleHessian.SetAt(alpha, alpha, ParticleHessian.GetAt(alpha, alpha) + Dt * Dt * Measure[ElementIndex] * T(2) * MuElementArray[ElementIndex] * DmInvsum);
+						ParticleHessian.SetAt(alpha, alpha, ParticleHessian.GetAt(alpha, alpha) + Dt * Dt * Measure[ElementIndex] * FSolverReal(2) * MuElementArray[ElementIndex] * DmInvsum);
 					}
 
-					const PMatrix<T, 3, 2> dJDmInvT = dJ * DmInvT; 
+					const PMatrix<FSolverReal, 3, 2> dJDmInvT = dJ * DmInvT; 
 
-					Chaos::TVector<T, 3> L((T)0.);
+					Chaos::TVector<FSolverReal, 3> L((FSolverReal)0.);
 					for (int32 alpha = 0; alpha < 3; alpha++) 
 					{
 						const int32 IndexVisited = ElementIndexLocal* 3 - 3 + alpha;
@@ -416,16 +416,19 @@ namespace Chaos::Softs
 		mutable TArray<FSolverMatrix22> DmInverse;
 
 		//material constants calculated from E:
-		T Mu;
-		T Lambda;
-		TArray<T> MuElementArray;
-		TArray<T> LambdaElementArray;
-		TArray<T> AlphaJArray;
+		FSolverReal Mu;
+		FSolverReal Lambda;
+		TArray<FSolverReal> MuElementArray;
+		TArray<FSolverReal> LambdaElementArray;
+		TArray<FSolverReal> AlphaJArray;
 		bool bRecordMetric;
 
 		TArray<TVector<int32, 3>> MeshConstraints;
-		mutable TArray<T> Measure;
+		mutable TArray<FSolverReal> Measure;
 	};
 
+	template <typename T, typename ParticleType>
+	using FGaussSeidelCorotatedCodimensionalConstraints UE_DEPRECATED(5.8, "Deprecated. this class is to be deleted, use FGaussSeidelCorotatedTriangleConstraints instead") = FGaussSeidelCorotatedTriangleConstraints;
+
 }  // End namespace Chaos::Softs
 
diff --git a/Engine/Source/Runtime/Experimental/Chaos/Public/Chaos/Deformable/GaussSeidelCorotatedConstraints.h b/Engine/Source/Runtime/Experimental/Chaos/Public/Chaos/Deformable/GaussSeidelCorotatedConstraints.h
--- a/Engine/Source/Runtime/Experimental/Chaos/Public/Chaos/Deformable/GaussSeidelCorotatedConstraints.h
+++ b/Engine/Source/Runtime/Experimental/Chaos/Public/Chaos/Deformable/GaussSeidelCorotatedConstraints.h
@@ -13,11 +13,11 @@
 
 namespace Chaos::Softs
 {
-	template <typename T, typename ParticleType>
-	class FGaussSeidelCorotatedConstraints : public FXPBDCorotatedConstraints <T, ParticleType>
+	// Tetrahedral constraint using fixed corotated constitutive model [Stomakhin et al., 2012]
+	// for position-based nonlinear Gauss Seidel solver [Chen et al., 2024]
+	class FGaussSeidelCorotatedTetrahedralConstraints : public FXPBDCorotatedTetrahedralConstraints
 	{
-
-		typedef FXPBDCorotatedConstraints<T, ParticleType> Base;
+		typedef FXPBDCorotatedTetrahedralConstraints Base;
 		using Base::MeshConstraints;
 		using Base::LambdaArray;
 		using Base::Measure;
@@ -28,21 +28,21 @@ namespace Chaos::Softs
 		using Base::F;
 
 	public:
-		FGaussSeidelCorotatedConstraints(
-			const ParticleType& InParticles,
+		FGaussSeidelCorotatedTetrahedralConstraints(
+			const FSolverParticles& InParticles,
 			const TArray<TVector<int32, 4>>& InMesh,
-			const TArray<T>& EMeshArray,
-			const TArray<T>& NuMeshArray,
-			TArray<T>&& AlphaJMeshArray,
+			const TArray<FSolverReal>& EMeshArray,
+			const TArray<FSolverReal>& NuMeshArray,
+			TArray<FSolverReal>&& AlphaJMeshArray,
 			TArray<TArray<int32>>&& IncidentElementsIn,
 			TArray<TArray<int32>>&& IncidentElementsLocalIn,
 			const int32 ParticleStartIndexIn,
 			const int32 ParticleEndIndexIn,
 			const bool bDoQuasistaticsIn = false,
 			const bool bDoSORIn = true,
-			const T InOmegaSOR = (T)1.6,
+			const FSolverReal InOmegaSOR = (FSolverReal)1.6,
 			const FDeformableXPBDCorotatedParams& InParams = FDeformableXPBDCorotatedParams(),
-			const T& NuMesh = (T).3,
+			const FSolverReal& NuMesh = (FSolverReal).3,
 			const bool bRecordMetricIn = false
 		)
 			: Base(InParticles, InMesh, EMeshArray, NuMeshArray, MoveTemp(AlphaJMeshArray), InParams, NuMesh, bRecordMetricIn, false), IncidentElements(IncidentElementsIn), IncidentElementsLocal(IncidentElementsLocalIn),
@@ -64,21 +64,21 @@ namespace Chaos::Softs
 			}
 			if (!bDoQuasistatics)
 			{
-				xtilde.Init(TVector<T, 3>(T(0.)), ParticleEndIndex - ParticleStartIndex);
+				xtilde.Init(TVector<FSolverReal, 3>(FSolverReal(0.)), ParticleEndIndex - ParticleStartIndex);
 			}
 			if (bDoSOR)
 			{
-				X_k_1.Init(TVector<T, 3>(T(0.)), ParticleEndIndex - ParticleStartIndex);
-				X_k.Init(TVector<T, 3>(T(0.)), ParticleEndIndex - ParticleStartIndex);
+				X_k_1.Init(TVector<FSolverReal, 3>(FSolverReal(0.)), ParticleEndIndex - ParticleStartIndex);
+				X_k.Init(TVector<FSolverReal, 3>(FSolverReal(0.)), ParticleEndIndex - ParticleStartIndex);
 			}
 			InitColor(InParticles);
 			InitializeCorotatedLambdas();
 		}
 
-		virtual ~FGaussSeidelCorotatedConstraints() {}
+		virtual ~FGaussSeidelCorotatedTetrahedralConstraints() {}
 
 
-		void Apply(ParticleType& Particles, const T Dt) const 
+		void Apply(FSolverParticlesRange& Particles, const FSolverReal Dt) const 
 		{
 			TRACE_CPUPROFILER_EVENT_SCOPE(STAT_ChaosGaussSeidelApply);
 			for (int32 i = 0; i < ParticlesPerColor.Num(); i++)
@@ -96,7 +96,7 @@ namespace Chaos::Softs
 							if (TaskIndex < ParticlesPerColor[i].Num())
 							{
 								int32 ParticleIndex = ParticlesPerColor[i][TaskIndex];
-								Chaos::TVector<T, 3> Dx = ComputeDeltax(ParticleIndex, Particle2Incident[ParticleIndex - ParticleStartIndex], Particles,
+								Chaos::TVector<FSolverReal, 3> Dx = ComputeDeltax(ParticleIndex, Particle2Incident[ParticleIndex - ParticleStartIndex], Particles,
 									Dt);
 								Particles.P(ParticleIndex) += Dx;
 							}
@@ -108,7 +108,13 @@ namespace Chaos::Softs
 			CurrentIt += 1;
 		}
 
-		void Init(const T Dt, const ParticleType& Particles) const 
+		void Apply(FSolverParticles& Particles, const FSolverReal Dt) const
+		{
+			FSolverParticlesRange Range(&Particles, /*InOffset = */ 0, /*InRange = */ Particles.Size());
+			Apply(Range, Dt);
+		}
+
+		void Init(const FSolverReal Dt, const FSolverParticles& Particles) const 
 		{
 			if (!bDoQuasistatics)
 			{
@@ -143,14 +149,14 @@ namespace Chaos::Softs
 		}
 
 		// Apply Successive Over Relaxation (SOR) to accelerate convergence rate
-		void ApplySOR(ParticleType& Particles, const T Dt) const
+		void ApplySOR(FSolverParticlesRange& Particles, const FSolverReal Dt) const
 		{
 			if (bDoSOR)
 			{
 				PhysicsParallelFor(ParticleEndIndex - ParticleStartIndex, [&](const int32 i)
 					{
 						int32 ParticleIndex = i + ParticleStartIndex;
-						if (Particles.InvM(ParticleIndex) != T(0) && CurrentIt > 3)
+						if (Particles.InvM(ParticleIndex) != FSolverReal(0) && CurrentIt > 3)
 						{
 							Particles.P(ParticleIndex) = OmegaSOR * (Particles.P(ParticleIndex) - X_k_1[i]) + X_k_1[i];
 						}
@@ -160,11 +166,11 @@ namespace Chaos::Softs
 			}
 		}
 
-		Chaos::TVector<T, 3> ComputePerParticleResidual(const int32 p, const int32 IncidentIndex, const ParticleType& Particles,
-			const T Dt, const bool AddMass = true) const
+		Chaos::TVector<FSolverReal, 3> ComputePerParticleResidual(const int32 p, const int32 IncidentIndex, const FSolverParticlesRange& Particles,
+			const FSolverReal Dt, const bool AddMass = true) const
 		{
 
-			Chaos::TVector<T, 3> res(T(0));
+			Chaos::TVector<FSolverReal, 3> res(FSolverReal(0));
 			if (AddMass) {
 				for (int32 alpha = 0; alpha < 3; alpha++) {
 					res[alpha] = Particles.M(p) * (Particles.P(p)[alpha] - xtilde[p - ParticleStartIndex][alpha]);
@@ -173,14 +179,14 @@ namespace Chaos::Softs
 			for (int32 j = 0; j < IncidentElements[IncidentIndex].Num(); j++) {
 				int32 local_index = IncidentElementsLocal[IncidentIndex][j];
 				int32 e = IncidentElements[IncidentIndex][j];
-				Chaos::PMatrix<T, 3, 3> DmInvT = Base::ElementDmInv(e).GetTransposed();
-				Chaos::PMatrix<T, 3, 3> Fe = Base::F(e, Particles);
-				Chaos::PMatrix<T, 3, 3> P((T)0.);
+				Chaos::PMatrix<FSolverReal, 3, 3> DmInvT = Base::ElementDmInv(e).GetTransposed();
+				Chaos::PMatrix<FSolverReal, 3, 3> Fe = Base::F(e, Particles);
+				Chaos::PMatrix<FSolverReal, 3, 3> P((FSolverReal)0.);
 
 				ComputeStress(Fe, MuElementArray[e], LambdaElementArray[e], P);
 
-				Chaos::PMatrix<T, 3, 3> force_term = -Measure[e] * DmInvT * P;
-				Chaos::TVector<T, 3> dx((T)0.);
+				Chaos::PMatrix<FSolverReal, 3, 3> force_term = -Measure[e] * DmInvT * P;
+				Chaos::TVector<FSolverReal, 3> dx((FSolverReal)0.);
 				if (local_index > 0)
 				{
 					for (int32 c = 0; c < 3; c++)
@@ -211,10 +217,10 @@ namespace Chaos::Softs
 		}
 
 
-		Chaos::PMatrix<T, 3, 3> ComputePerParticleCorotatedHessianSimple(const int32 p, const int32 IncidentIndex, const ParticleType& Particles,
-			const T Dt, const bool AddMass = true) const
+		Chaos::PMatrix<FSolverReal, 3, 3> ComputePerParticleCorotatedHessianSimple(const int32 p, const int32 IncidentIndex, const FSolverParticlesRange& Particles,
+			const FSolverReal Dt, const bool AddMass = true) const
 		{
-			Chaos::PMatrix<T, 3, 3> final_hessian = Chaos::PMatrix<T, 3, 3>::Zero;
+			Chaos::PMatrix<FSolverReal, 3, 3> final_hessian = Chaos::PMatrix<FSolverReal, 3, 3>::Zero;
 			if (AddMass) {
 				for (int32 alpha = 0; alpha < 3; alpha++) {
 					final_hessian.SetAt(alpha, alpha, Particles.M(p));
@@ -224,10 +230,10 @@ namespace Chaos::Softs
 			for (int32 j = 0; j < IncidentElements[IncidentIndex].Num(); j++) {
 				int32 local_index = IncidentElementsLocal[IncidentIndex][j];
 				int32 ElementIndex = IncidentElements[IncidentIndex][j];
-				Chaos::PMatrix<T, 3, 3> DmInv = Base::ElementDmInv(ElementIndex);
-				Chaos::PMatrix<T, 3, 3> Fe = Base::F(ElementIndex, Particles);
+				Chaos::PMatrix<FSolverReal, 3, 3> DmInv = Base::ElementDmInv(ElementIndex);
+				Chaos::PMatrix<FSolverReal, 3, 3> Fe = Base::F(ElementIndex, Particles);
 
-				T Coeff = Dt * Dt * Measure[ElementIndex];
+				FSolverReal Coeff = Dt * Dt * Measure[ElementIndex];
 
 				ComputeHessianHelper(Fe, DmInv, MuElementArray[ElementIndex], LambdaElementArray[ElementIndex], local_index, Coeff, final_hessian);
 
@@ -236,16 +242,16 @@ namespace Chaos::Softs
 		}
 
 
-		void AddHyperelasticResidualAndHessian (const ParticleType& Particles, const int32 ElementIndex, const int32 ElementIndexLocal, const T Dt, TVec3<T>& ParticleResidual, Chaos::PMatrix<T, 3, 3>& ParticleHessian)
+		void AddHyperelasticResidualAndHessian (const FSolverParticlesRange& Particles, const int32 ElementIndex, const int32 ElementIndexLocal, const FSolverReal Dt, TVec3<FSolverReal>& ParticleResidual, Chaos::PMatrix<FSolverReal, 3, 3>& ParticleHessian)
 		{
-			Chaos::PMatrix<T, 3, 3> DmInvT = Base::ElementDmInv(ElementIndex).GetTransposed();
-			Chaos::PMatrix<T, 3, 3> Fe = F(ElementIndex, Particles);
-			Chaos::PMatrix<T, 3, 3> P((T)0.);
+			Chaos::PMatrix<FSolverReal, 3, 3> DmInvT = Base::ElementDmInv(ElementIndex).GetTransposed();
+			Chaos::PMatrix<FSolverReal, 3, 3> Fe = F(ElementIndex, Particles);
+			Chaos::PMatrix<FSolverReal, 3, 3> P((FSolverReal)0.);
 
 			this->ComputeStress(Fe, MuElementArray[ElementIndex], LambdaElementArray[ElementIndex], P);
 
-			Chaos::PMatrix<T, 3, 3> force_term = -Measure[ElementIndex] * DmInvT * P;
-			Chaos::TVector<T, 3> dx((T)0.);
+			Chaos::PMatrix<FSolverReal, 3, 3> force_term = -Measure[ElementIndex] * DmInvT * P;
+			Chaos::TVector<FSolverReal, 3> dx((FSolverReal)0.);
 			if (ElementIndexLocal > 0)
 			{
 				for (int32 c = 0; c < 3; c++)
@@ -276,21 +282,27 @@ namespace Chaos::Softs
 
 		}
 
+		void AddHyperelasticResidualAndHessian(const FSolverParticles& Particles, const int32 ElementIndex, const int32 ElementIndexLocal, const FSolverReal Dt, TVec3<FSolverReal>& ParticleResidual, Chaos::PMatrix<FSolverReal, 3, 3>& ParticleHessian)
+		{
+			AddHyperelasticResidualAndHessian(FSolverParticlesRange(const_cast<FSolverParticles*>(&Particles), /*InOffset = */ 0, /*InRange = */ Particles.Size()), 
+				ElementIndex, ElementIndexLocal, Dt, ParticleResidual, ParticleHessian);
+		}
+
 	protected:
-		 void InitColor(const ParticleType& Particles)
+		 void InitColor(const FSolverParticles& Particles)
 		{
 			ParticlesPerColor = ComputeNodalColoring(MeshConstraints, Particles, ParticleStartIndex, ParticleEndIndex, IncidentElements, IncidentElementsLocal);
 		}
 
 		 void InitializeCorotatedLambdas()
 		 {
-			 ComputeStress = [](const Chaos::PMatrix<T, 3, 3>& Fe, const T mu, const T lambda, Chaos::PMatrix<T, 3, 3>& P)
+			 ComputeStress = [](const Chaos::PMatrix<FSolverReal, 3, 3>& Fe, const FSolverReal mu, const FSolverReal lambda, Chaos::PMatrix<FSolverReal, 3, 3>& P)
 			 {
 				 PCorotated(Fe, mu, lambda, P);
 			 };
-			 ComputeHessianHelper = [](const Chaos::PMatrix<T, 3, 3>& Fe, const Chaos::PMatrix<T, 3, 3>& DmInv, const T mu, const T lambda, const int32 local_index, const T Coeff, Chaos::PMatrix<T, 3, 3>& final_hessian) 
+			 ComputeHessianHelper = [](const Chaos::PMatrix<FSolverReal, 3, 3>& Fe, const Chaos::PMatrix<FSolverReal, 3, 3>& DmInv, const FSolverReal mu, const FSolverReal lambda, const int32 local_index, const FSolverReal Coeff, Chaos::PMatrix<FSolverReal, 3, 3>& final_hessian) 
 			 {
-				 Chaos::PMatrix<T, 3, 3> JFinvT((T)0.);
+				 Chaos::PMatrix<FSolverReal, 3, 3> JFinvT((FSolverReal)0.);
 				 JFinvT.SetAt(0, 0, Fe.GetAt(1, 1) * Fe.GetAt(2, 2) - Fe.GetAt(2, 1) * Fe.GetAt(1, 2));
 				 JFinvT.SetAt(0, 1, Fe.GetAt(2, 0) * Fe.GetAt(1, 2) - Fe.GetAt(1, 0) * Fe.GetAt(2, 2));
 				 JFinvT.SetAt(0, 2, Fe.GetAt(1, 0) * Fe.GetAt(2, 1) - Fe.GetAt(2, 0) * Fe.GetAt(1, 1));
@@ -301,23 +313,23 @@ namespace Chaos::Softs
 				 JFinvT.SetAt(2, 1, Fe.GetAt(1, 0) * Fe.GetAt(0, 2) - Fe.GetAt(0, 0) * Fe.GetAt(1, 2));
 				 JFinvT.SetAt(2, 2, Fe.GetAt(0, 0) * Fe.GetAt(1, 1) - Fe.GetAt(1, 0) * Fe.GetAt(0, 1));
 
-				 Chaos::PMatrix<T, 3, 3> JFinv = JFinvT.GetTransposed();
+				 Chaos::PMatrix<FSolverReal, 3, 3> JFinv = JFinvT.GetTransposed();
 
 				 if (local_index == 0) {
-					 T DmInvsum = T(0);
+					 FSolverReal DmInvsum = FSolverReal(0);
 					 for (int32 nu = 0; nu < 3; nu++) {
-						 T localDmsum = T(0);
+						 FSolverReal localDmsum = FSolverReal(0);
 						 for (int32 k = 0; k < 3; k++) {
 							 localDmsum += DmInv.GetAt(k, nu);
 						 }
 						 DmInvsum += localDmsum * localDmsum;
 					 }
 					 for (int32 alpha = 0; alpha < 3; alpha++) {
-						 final_hessian.SetAt(alpha, alpha, final_hessian.GetAt(alpha, alpha) + Coeff * T(2) * mu * DmInvsum);
+						 final_hessian.SetAt(alpha, alpha, final_hessian.GetAt(alpha, alpha) + Coeff * FSolverReal(2) * mu * DmInvsum);
 					 }
 
-					 Chaos::PMatrix<T, 3, 3> DmInvJFinv = JFinv * DmInv;
-					 Chaos::TVector<T, 3> l((T)0.);
+					 Chaos::PMatrix<FSolverReal, 3, 3> DmInvJFinv = JFinv * DmInv;
+					 Chaos::TVector<FSolverReal, 3> l((FSolverReal)0.);
 					 for (int32 alpha = 0; alpha < 3; alpha++) {
 						 for (int32 k = 0; k < 3; k++) {
 							 l[alpha] += DmInvJFinv.GetAt(k, alpha);
@@ -330,18 +342,18 @@ namespace Chaos::Softs
 
 				 }
 				 else {
-					 T DmInvsum = T(0);
+					 FSolverReal DmInvsum = FSolverReal(0);
 					 for (int32 nu = 0; nu < 3; nu++)
 					 {
 						 DmInvsum += DmInv.GetAt(local_index - 1, nu) * DmInv.GetAt(local_index - 1, nu);
 					 }
 					 for (int32 alpha = 0; alpha < 3; alpha++)
 					 {
-						 final_hessian.SetAt(alpha, alpha, final_hessian.GetAt(alpha, alpha) + Coeff * T(2) * mu * DmInvsum);
+						 final_hessian.SetAt(alpha, alpha, final_hessian.GetAt(alpha, alpha) + Coeff * FSolverReal(2) * mu * DmInvsum);
 					 }
 
-					 Chaos::PMatrix<T, 3, 3> DmInvJFinv = JFinv * DmInv;
-					 Chaos::TVector<T, 3> l((T)0.);
+					 Chaos::PMatrix<FSolverReal, 3, 3> DmInvJFinv = JFinv * DmInv;
+					 Chaos::TVector<FSolverReal, 3> l((FSolverReal)0.);
 					 for (int32 alpha = 0; alpha < 3; alpha++) {
 						 l[alpha] = DmInvJFinv.GetAt(local_index - 1, alpha);
 					 }
@@ -352,33 +364,33 @@ namespace Chaos::Softs
 				 }
 			 };
 
-			 AddAdditionalRes = [](const ParticleType& InParticles, const int32 p, const T Dt, TVec3<T>& res) {};
-			 AddAdditionalHessian = [](const ParticleType& InParticles, const int32 p, const T Dt, Chaos::PMatrix<T, 3, 3>& hessian) {};
+			 AddAdditionalRes = [](const FSolverParticlesRange& InParticles, const int32 p, const FSolverReal Dt, TVec3<FSolverReal>& res) {};
+			 AddAdditionalHessian = [](const FSolverParticlesRange& InParticles, const int32 p, const FSolverReal Dt, Chaos::PMatrix<FSolverReal, 3, 3>& hessian) {};
 		 }
 
 	private:
 
-		TVector<T, 3> ComputeDeltax(const int32 p, const int32 IncidentIndex, const ParticleType& Particles,
-			const T Dt) const
+		TVector<FSolverReal, 3> ComputeDeltax(const int32 p, const int32 IncidentIndex, const FSolverParticlesRange& Particles,
+			const FSolverReal Dt) const
 		{
-			Chaos::TVector<T, 3> ParticleResidual = ComputePerParticleResidual(p, Particle2Incident[p], Particles, Dt, !bDoQuasistatics);
+			Chaos::TVector<FSolverReal, 3> ParticleResidual = ComputePerParticleResidual(p, Particle2Incident[p], Particles, Dt, !bDoQuasistatics);
 
 			AddAdditionalRes(Particles, p, Dt, ParticleResidual);
 
 			if (ParticleResidual.Size() > LocalNewtonTol)
 			{
-				Chaos::PMatrix<T, 3, 3> SimplifiedHessian = ComputePerParticleCorotatedHessianSimple(p, Particle2Incident[p], Particles, Dt, !bDoQuasistatics);
+				Chaos::PMatrix<FSolverReal, 3, 3> SimplifiedHessian = ComputePerParticleCorotatedHessianSimple(p, Particle2Incident[p], Particles, Dt, !bDoQuasistatics);
 				AddAdditionalHessian(Particles, p, Dt, SimplifiedHessian);
 				
-				T HessianDet = SimplifiedHessian.Determinant();
+				FSolverReal HessianDet = SimplifiedHessian.Determinant();
 				if (HessianDet > UE_SMALL_NUMBER)
 				{
-					Chaos::PMatrix<T, 3, 3> HessianInv = SimplifiedHessian.SymmetricCofactorMatrix();
-					HessianInv *= T(1) / HessianDet;
+					Chaos::PMatrix<FSolverReal, 3, 3> HessianInv = SimplifiedHessian.SymmetricCofactorMatrix();
+					HessianInv *= FSolverReal(1) / HessianDet;
 					return HessianInv.GetTransposed() * (-ParticleResidual);
 				}
 			}
-			return TVector<T, 3>((T)0.);
+			return TVector<FSolverReal, 3>((FSolverReal)0.);
 		}
 
 		
@@ -388,26 +400,28 @@ namespace Chaos::Softs
 		TArray<int32> Particle2Incident;
 		TArray<TArray<int32>> IncidentElements;
 		TArray<TArray<int32>> IncidentElementsLocal;
-		T LocalNewtonTol = T(1e-5);
+		FSolverReal LocalNewtonTol = FSolverReal(1e-5);
 		TArray<TArray<int32>> ParticlesPerColor;
 		int32 ParticleStartIndex;
 		int32 ParticleEndIndex;
-		mutable TArray<Chaos::TVector<T, 3>> xtilde;
-		mutable TArray<Chaos::TVector<T, 3>> X_k_1;
-		mutable TArray<Chaos::TVector<T, 3>> X_k;
+		mutable TArray<Chaos::TVector<FSolverReal, 3>> xtilde;
+		mutable TArray<Chaos::TVector<FSolverReal, 3>> X_k_1;
+		mutable TArray<Chaos::TVector<FSolverReal, 3>> X_k;
 		bool bDoQuasistatics = false;
 		bool bDoSOR = true;
 		//Parameter to control the behavior of SOR. 1-1.7 is the usual range. Larger paramater leads to potentially more convergence but less stable behavior. 
-		T OmegaSOR = T(1.6);
+		FSolverReal OmegaSOR = FSolverReal(1.6);
 		mutable int32 CurrentIt = 0;
-		TFunction<void(const Chaos::PMatrix<T, 3, 3>&, const T, const T, Chaos::PMatrix<T, 3, 3>& )> ComputeStress;
-		TFunction<void(const Chaos::PMatrix<T, 3, 3>&, const Chaos::PMatrix<T, 3, 3>&, const T, const T, const int32, const T, Chaos::PMatrix<T, 3, 3>&)> ComputeHessianHelper;
+		TFunction<void(const Chaos::PMatrix<FSolverReal, 3, 3>&, const FSolverReal, const FSolverReal, Chaos::PMatrix<FSolverReal, 3, 3>& )> ComputeStress;
+		TFunction<void(const Chaos::PMatrix<FSolverReal, 3, 3>&, const Chaos::PMatrix<FSolverReal, 3, 3>&, const FSolverReal, const FSolverReal, const int32, const FSolverReal, Chaos::PMatrix<FSolverReal, 3, 3>&)> ComputeHessianHelper;
 		
 	public:
-		TFunction<void(const ParticleType&, const int32, const T, TVec3<T>&)> AddAdditionalRes;
-		TFunction<void(const ParticleType&, const int32, const T, Chaos::PMatrix<T, 3, 3>&)> AddAdditionalHessian;
+		TFunction<void(const FSolverParticlesRange&, const int32, const FSolverReal, TVec3<FSolverReal>&)> AddAdditionalRes;
+		TFunction<void(const FSolverParticlesRange&, const int32, const FSolverReal, Chaos::PMatrix<FSolverReal, 3, 3>&)> AddAdditionalHessian;
 		TUniquePtr<TArray<int32>> ParticleColors;
 	};
 
+	template <typename T, typename ParticleType>
+	using FGaussSeidelCorotatedConstraints UE_DEPRECATED(5.8, "Deprecated. this class is to be deleted, use FGaussSeidelCorotatedTetrahedralConstraints instead") = FGaussSeidelCorotatedTetrahedralConstraints;
 
 }// End namespace Chaos::Softs
diff --git a/Engine/Source/Runtime/Experimental/Chaos/Public/Chaos/Deformable/GaussSeidelDynamicWeakConstraints.h b/Engine/Source/Runtime/Experimental/Chaos/Public/Chaos/Deformable/GaussSeidelDynamicWeakConstraints.h
--- a/Engine/Source/Runtime/Experimental/Chaos/Public/Chaos/Deformable/GaussSeidelDynamicWeakConstraints.h
+++ b/Engine/Source/Runtime/Experimental/Chaos/Public/Chaos/Deformable/GaussSeidelDynamicWeakConstraints.h
@@ -20,24 +20,23 @@
 namespace Chaos::Softs
 {
 	using Chaos::TVec3;
-	template <typename T, typename ParticleType>
-	struct FGaussSeidelDynamicWeakConstraints: public FGaussSeidelWeakConstraints <T, ParticleType>
+	struct FGaussSeidelDynamicSpringConstraints: public FGaussSeidelSpringConstraints
 	{
 
-		typedef FGaussSeidelWeakConstraints<T, ParticleType> Base;
+		typedef FGaussSeidelSpringConstraints Base;
 		using Base::ConstraintsData;
 		using Base::NodalWeights;
 		using Base::WCIncidentElements;
 		using Base::WCIncidentElementsLocal;
-		typedef typename FGaussSeidelWeakConstraints<T, ParticleType>::FGaussSeidelConstraintHandle GSConstraintHandle;
+		typedef typename FGaussSeidelSpringConstraints::FGaussSeidelConstraintHandle GSConstraintHandle;
 
-		FGaussSeidelDynamicWeakConstraints(const FDeformableXPBDWeakConstraintParams& InParams)
+		FGaussSeidelDynamicSpringConstraints(const FDeformableXPBDWeightedSpringConstraintParams& InParams)
 		: Base({}, {}, {}, {}, {}, InParams)
 		{
 			ConstraintsData.AddArray(&Handles);
 		}
 
-		virtual ~FGaussSeidelDynamicWeakConstraints()
+		virtual ~FGaussSeidelDynamicSpringConstraints()
 		{
 			for (GSConstraintHandle* Handle:Handles)
 			{
@@ -45,7 +44,7 @@ namespace Chaos::Softs
 			}
 		}
 		
-		const GSConstraintHandle* AddSingleParticleTetrahedraConstraint(const GeometryCollection::Facades::FTetrahedralFacade& InTargetGeom, const ParticleType& AllParticles, const GeometryCollection::Facades::TetrahedralParticleEmbedding& InIntersection, const int32 InConstraintIndex, const Chaos::FRange& SourceRange, const Chaos::FRange& TargetRange, const T PositionTargetStiffness)
+		const GSConstraintHandle* AddSingleParticleTetrahedraConstraint(const GeometryCollection::Facades::FTetrahedralFacade& InTargetGeom, const FSolverParticles& AllParticles, const GeometryCollection::Facades::TetrahedralParticleEmbedding& InIntersection, const int32 InConstraintIndex, const Chaos::FRange& SourceRange, const Chaos::FRange& TargetRange, const FSolverReal PositionTargetStiffness)
 		{
 			GSConstraintHandle* NewConstraintPtr = new GSConstraintHandle();
 			NewConstraintPtr -> ConstraintIndex = InConstraintIndex;
@@ -53,20 +52,20 @@ namespace Chaos::Softs
 
 			const FIntVector4 TargetTet = InTargetGeom.Tetrahedron[InTargetGeom.TetrahedronStart[InIntersection.GeometryIndex] + InIntersection.TetrahedronIndex] + FIntVector4(TargetRange.Start);
 
-			FGaussSeidelWeakConstraintSingleData<T> SingleConstraintData;
+			FGaussSeidelSpringConstraintSingleData SingleConstraintData;
 			SingleConstraintData.SingleIndices = { InIntersection.ParticleIndex + SourceRange.Start };
 			SingleConstraintData.SingleSecondIndices =  {TargetTet[0], TargetTet[1], TargetTet[2], TargetTet[3]};
-			SingleConstraintData.SingleWeights = {(T)1.};
+			SingleConstraintData.SingleWeights = {(FSolverReal)1.};
 			SingleConstraintData.SingleSecondWeights = {InIntersection.BarycentricWeights[0], InIntersection.BarycentricWeights[1], InIntersection.BarycentricWeights[2], InIntersection.BarycentricWeights[3]};
 			SingleConstraintData.bIsAnisotropic = false;
-			SingleConstraintData.SingleNormal = TVec3<T>((T)0.);
+			SingleConstraintData.SingleNormal = TVec3<FSolverReal>((FSolverReal)0.);
 			float ConstraintStiffness = 0.f;
 			for (int32 k = 0; k < 4; k++)
 			{
 				ConstraintStiffness += SingleConstraintData.SingleSecondWeights[k] * PositionTargetStiffness * AllParticles.M(TargetTet[k]);
 			}
 			ConstraintStiffness += PositionTargetStiffness * AllParticles.M(SingleConstraintData.SingleIndices[0]);
-			SingleConstraintData.SingleStiffness = (T)ConstraintStiffness;
+			SingleConstraintData.SingleStiffness = (FSolverReal)ConstraintStiffness;
 
 			ConstraintsData.SetSingleConstraint(SingleConstraintData, InConstraintIndex);
 
@@ -78,7 +77,7 @@ namespace Chaos::Softs
 		{
 			const int32 OldConstraintIndex = DeletedSingleConstraint->ConstraintIndex;
 
-			const FGaussSeidelWeakConstraintSingleData<T>& SingleData = ConstraintsData.GetSingleConstraintData(OldConstraintIndex);
+			const FGaussSeidelSpringConstraintSingleData& SingleData = ConstraintsData.GetSingleConstraintData(OldConstraintIndex);
 
 			delete DeletedSingleConstraint;
 
@@ -118,7 +117,7 @@ namespace Chaos::Softs
 			
 		}
 
-		TArray<const GSConstraintHandle*> AddParticleTetrahedraConstraints(const GeometryCollection::Facades::FTetrahedralFacade& InTargetGeom, const ParticleType& AllParticles, const TArray<GeometryCollection::Facades::TetrahedralParticleEmbedding>& InIntersections, const Chaos::FRange& SourceRange, const Chaos::FRange& TargetRange, const T ConstraintStiffness)
+		TArray<const GSConstraintHandle*> AddParticleTetrahedraConstraints(const GeometryCollection::Facades::FTetrahedralFacade& InTargetGeom, const FSolverParticles& AllParticles, const TArray<GeometryCollection::Facades::TetrahedralParticleEmbedding>& InIntersections, const Chaos::FRange& SourceRange, const Chaos::FRange& TargetRange, const FSolverReal ConstraintStiffness)
 		{
 			TArray<const GSConstraintHandle*> ConstraintHandles;
 			ConstraintHandles.SetNum(InIntersections.Num());
@@ -147,13 +146,13 @@ namespace Chaos::Softs
 			for (const int32 Vert : DirtyVerts)
 			{
 				check(Vert < NodalWeights.Num());
-				NodalWeights[Vert].Init(T(0), 6);
+				NodalWeights[Vert].Init(FSolverReal(0), 6);
 				for (int32 j = 0; j < WCIncidentElements[Vert].Num(); j++)
 				{
 					const int32 ConstraintIndex = WCIncidentElements[Vert][j];
 					const int32 LocalIndex = WCIncidentElementsLocal[Vert][j];
 
-					T Weight = T(0);
+					FSolverReal Weight = FSolverReal(0);
 					if (LocalIndex >= ConstraintsData.GetIndices(ConstraintIndex).Num())
 					{
 						Weight = ConstraintsData.GetSecondWeights(ConstraintIndex)[LocalIndex - ConstraintsData.GetIndices(ConstraintIndex).Num()];
@@ -185,16 +184,16 @@ namespace Chaos::Softs
 			}
 		}
 
-		void AdjustStiffness(const TArray<const GSConstraintHandle*> ConstraintHandles, const TArray<T>& StiffnessPerConstraint, const ParticleType& Particles)
+		void AdjustStiffness(const TArray<const GSConstraintHandle*> ConstraintHandles, const TArray<FSolverReal>& StiffnessPerConstraint, const FSolverParticles& Particles)
 		{
 			TSet<int32> DirtyVertices;
 			for (int32 i = 0; i < ConstraintHandles.Num(); i++)
 			{
 				const GSConstraintHandle* SingleHandle = ConstraintHandles[i];
-				const T NewStiffness = StiffnessPerConstraint[i];
+				const FSolverReal NewStiffness = StiffnessPerConstraint[i];
 				const int32 ConstraintIndex = SingleHandle->ConstraintIndex;
-				const FGaussSeidelWeakConstraintSingleData<T>& SingleConstraintData = ConstraintsData.GetSingleConstraintData(ConstraintIndex);
-				T NewConstraintStiffness = (T)0.;
+				const FGaussSeidelSpringConstraintSingleData& SingleConstraintData = ConstraintsData.GetSingleConstraintData(ConstraintIndex);
+				FSolverReal NewConstraintStiffness = (FSolverReal)0.;
 				for (int32 k = 0; k < SingleConstraintData.SingleWeights.Num(); k++)
 				{
 					NewConstraintStiffness += SingleConstraintData.SingleWeights[k] * NewStiffness * Particles.M(SingleConstraintData.SingleIndices[k]);
@@ -214,4 +213,7 @@ namespace Chaos::Softs
 
 	};
 
-}
+	template <typename T, typename ParticleType>
+	using FGaussSeidelDynamicWeakConstraints UE_DEPRECATED(5.8, "Deprecated. this struct is to be deleted, use FGaussSeidelDynamicSpringConstraints instead") = FGaussSeidelDynamicSpringConstraints;
+
+} //End namespace Chaos::Softs
diff --git a/Engine/Source/Runtime/Experimental/Chaos/Public/Chaos/Deformable/GaussSeidelLinearCodimensionalConstraints.h b/Engine/Source/Runtime/Experimental/Chaos/Public/Chaos/Deformable/GaussSeidelLinearCodimensionalConstraints.h
--- a/Engine/Source/Runtime/Experimental/Chaos/Public/Chaos/Deformable/GaussSeidelLinearCodimensionalConstraints.h
+++ b/Engine/Source/Runtime/Experimental/Chaos/Public/Chaos/Deformable/GaussSeidelLinearCodimensionalConstraints.h
@@ -11,23 +11,24 @@
 
 
 namespace Chaos::Softs
-{
-	template <typename T, typename ParticleType>
-	class FGaussSeidelLinearCodimensionalConstraints
+{	
+	// Triangle constraint using linear constitutive model
+	// 3D triangles are projected to the 2D space for deformation gradient (3*2 matrix) computation
+	class FGaussSeidelLinearTriangleConstraints
 	{
 
 	public:
-		FGaussSeidelLinearCodimensionalConstraints(
-			const ParticleType& InParticles,
+		FGaussSeidelLinearTriangleConstraints(
+			const FSolverParticles& InParticles,
 			const TArray<TVector<int32, 3>>& InMesh,
-			const T& EMesh = (T)10.0,
-			const T& NuMesh = (T).3
+			const FSolverReal& EMesh = (FSolverReal)10.0,
+			const FSolverReal& NuMesh = (FSolverReal).3
 		)
 			: MeshConstraints(InMesh)
 		{
-			Measure.Init((T)0., MeshConstraints.Num());
-			Lambda = EMesh * NuMesh / (((T)1. + NuMesh) * ((T)1. - (T)2. * NuMesh));
-			Mu = EMesh / ((T)2. * ((T)1. + NuMesh));
+			Measure.Init((FSolverReal)0., MeshConstraints.Num());
+			Lambda = EMesh * NuMesh / (((FSolverReal)1. + NuMesh) * ((FSolverReal)1. - (FSolverReal)2. * NuMesh));
+			Mu = EMesh / ((FSolverReal)2. * ((FSolverReal)1. + NuMesh));
 
 			MuElementArray.Init(Mu, MeshConstraints.Num());
 			LambdaElementArray.Init(Lambda, MeshConstraints.Num());
@@ -35,32 +36,32 @@ namespace Chaos::Softs
 			InitializeCodimensionData(InParticles);
 		}
 
-		FGaussSeidelLinearCodimensionalConstraints(
-			const ParticleType& InParticles,
+		FGaussSeidelLinearTriangleConstraints(
+			const FSolverParticles& InParticles,
 			const TArray<TVector<int32, 3>>& InMesh,
-			const TArray<T>& EMeshArray,
-			const T& NuMesh = (T).3
+			const TArray<FSolverReal>& EMeshArray,
+			const FSolverReal& NuMesh = (FSolverReal).3
 		)
 			: MeshConstraints(InMesh)
 		{
-			Measure.Init((T)0., MeshConstraints.Num());
+			Measure.Init((FSolverReal)0., MeshConstraints.Num());
 
 			InitializeCodimensionData(InParticles);
-			LambdaElementArray.Init((T)0., MeshConstraints.Num());
-			MuElementArray.Init((T)0., MeshConstraints.Num());
+			LambdaElementArray.Init((FSolverReal)0., MeshConstraints.Num());
+			MuElementArray.Init((FSolverReal)0., MeshConstraints.Num());
 			
 			for (int32 e = 0; e < InMesh.Num(); e++)
 			{
-				LambdaElementArray[e] = EMeshArray[e] * NuMesh / (((T)1. + NuMesh) * ((T)1. - (T)2. * NuMesh));
-				MuElementArray[e] = EMeshArray[e] / ((T)2. * ((T)1. + NuMesh));
+				LambdaElementArray[e] = EMeshArray[e] * NuMesh / (((FSolverReal)1. + NuMesh) * ((FSolverReal)1. - (FSolverReal)2. * NuMesh));
+				MuElementArray[e] = EMeshArray[e] / ((FSolverReal)2. * ((FSolverReal)1. + NuMesh));
 			}
 		}
 
-		virtual ~FGaussSeidelLinearCodimensionalConstraints() {}
+		virtual ~FGaussSeidelLinearTriangleConstraints() {}
 
-		PMatrix<T, 3, 3> DsInit(const int32 E, const ParticleType& InParticles)  
+		PMatrix<FSolverReal, 3, 3> DsInit(const int32 E, const FSolverParticles& InParticles)  
 		{
-			PMatrix<T, 3, 3> Result((T)0.);
+			PMatrix<FSolverReal, 3, 3> Result((FSolverReal)0.);
 			for (int32 i = 0; i < 3; i++) 
 			{
 				for (int32 c = 0; c < 3; c++) 
@@ -71,29 +72,29 @@ namespace Chaos::Softs
 			return Result;
 		}
 
-		PMatrix<T, 3, 2> Ds(const int32 E, const ParticleType& InParticles) const 
+		PMatrix<FSolverReal, 3, 2> Ds(const int32 E, const FSolverParticles& InParticles) const 
 		{
 			if (INDEX_NONE < E && E< MeshConstraints.Num()
 				&& INDEX_NONE < MeshConstraints[E][0] && MeshConstraints[E][0] < (int32)InParticles.Size()
 				&& INDEX_NONE < MeshConstraints[E][1] && MeshConstraints[E][1] < (int32)InParticles.Size()
 				&& INDEX_NONE < MeshConstraints[E][2] && MeshConstraints[E][2] < (int32)InParticles.Size())
 			{
-				const TVec3<T> P1P0 = InParticles.P(MeshConstraints[E][1]) - InParticles.P(MeshConstraints[E][0]);
-				const TVec3<T> P2P0 = InParticles.P(MeshConstraints[E][2]) - InParticles.P(MeshConstraints[E][0]);
+				const TVec3<FSolverReal> P1P0 = InParticles.P(MeshConstraints[E][1]) - InParticles.P(MeshConstraints[E][0]);
+				const TVec3<FSolverReal> P2P0 = InParticles.P(MeshConstraints[E][2]) - InParticles.P(MeshConstraints[E][0]);
 
-				return PMatrix<T, 3, 2>(
+				return PMatrix<FSolverReal, 3, 2>(
 					P1P0[0], P1P0[1], P1P0[2],
 					P2P0[0], P2P0[1], P2P0[2]);
 			}
 			else
 			{
-				return PMatrix<T, 3, 2>(
-					(T)0., (T)0, (T)0,
-					(T)0, (T)0, (T)0);
+				return PMatrix<FSolverReal, 3, 2>(
+					(FSolverReal)0., (FSolverReal)0, (FSolverReal)0,
+					(FSolverReal)0, (FSolverReal)0, (FSolverReal)0);
 			}
 		}
 
-		PMatrix<T, 3, 2> F(const int32 E, const ParticleType& InParticles) const 
+		PMatrix<FSolverReal, 3, 2> F(const int32 E, const FSolverParticles& InParticles) const 
 		{
 			if (INDEX_NONE < E && E < DmInverse.Num())
 			{
@@ -101,9 +102,9 @@ namespace Chaos::Softs
 			}
 			else
 			{
-				return PMatrix<T, 3, 2>(
-					(T)0., (T)0, (T)0,
-					(T)0, (T)0, (T)0);
+				return PMatrix<FSolverReal, 3, 2>(
+					(FSolverReal)0., (FSolverReal)0, (FSolverReal)0,
+					(FSolverReal)0, (FSolverReal)0, (FSolverReal)0);
 			}
 		}
 
@@ -124,9 +125,9 @@ namespace Chaos::Softs
 
 	protected:
 
-		void InitializeCodimensionData(const ParticleType& Particles)
+		void InitializeCodimensionData(const FSolverParticles& Particles)
 		{
-			Measure.Init((T)0., MeshConstraints.Num());
+			Measure.Init((FSolverReal)0., MeshConstraints.Num());
 			DmInverse.Init(PMatrix<FSolverReal, 2, 2>(0.f, 0.f, 0.f), MeshConstraints.Num());
 			for (int32 e = 0; e < MeshConstraints.Num(); e++)
 			{
@@ -137,32 +138,32 @@ namespace Chaos::Softs
 					&& MeshConstraints[e][1] < (int32)Particles.Size() && MeshConstraints[e][1] > INDEX_NONE
 					&& MeshConstraints[e][2] < (int32)Particles.Size() && MeshConstraints[e][2] > INDEX_NONE)
 				{
-					const TVec3<T> X1X0 = Particles.GetX(MeshConstraints[e][1]) - Particles.GetX(MeshConstraints[e][0]);
-					const TVec3<T> X2X0 = Particles.GetX(MeshConstraints[e][2]) - Particles.GetX(MeshConstraints[e][0]);
-					PMatrix<T, 2, 2> Dm((T)0., (T)0., (T)0.);
+					const TVec3<FSolverReal> X1X0 = Particles.GetX(MeshConstraints[e][1]) - Particles.GetX(MeshConstraints[e][0]);
+					const TVec3<FSolverReal> X2X0 = Particles.GetX(MeshConstraints[e][2]) - Particles.GetX(MeshConstraints[e][0]);
+					PMatrix<FSolverReal, 2, 2> Dm((FSolverReal)0., (FSolverReal)0., (FSolverReal)0.);
 					Dm.M[0] = X1X0.Size();
 					Dm.M[2] = X1X0.Dot(X2X0) / Dm.M[0];
-					Dm.M[3] = Chaos::TVector<T, 3>::CrossProduct(X1X0, X2X0).Size() / Dm.M[0];
-					Measure[e] = Chaos::TVector<T, 3>::CrossProduct(X1X0, X2X0).Size() / (T)2.;
-					ensureMsgf(Measure[e] > (T)0., TEXT("Degenerate triangle detected"));
+					Dm.M[3] = Chaos::TVector<FSolverReal, 3>::CrossProduct(X1X0, X2X0).Size() / Dm.M[0];
+					Measure[e] = Chaos::TVector<FSolverReal, 3>::CrossProduct(X1X0, X2X0).Size() / (FSolverReal)2.;
+					ensureMsgf(Measure[e] > (FSolverReal)0., TEXT("Degenerate triangle detected"));
 
-					PMatrix<T, 2, 2> DmInv = Dm.Inverse();
+					PMatrix<FSolverReal, 2, 2> DmInv = Dm.Inverse();
 					DmInverse[e] = DmInv;
 				}
 			}
 		}
 
-		static T SafeRecip(const T Len, const T Fallback)
+		static FSolverReal SafeRecip(const FSolverReal Len, const FSolverReal Fallback)
 		{
-			if (Len > (T)UE_SMALL_NUMBER)
+			if (Len > (FSolverReal)UE_SMALL_NUMBER)
 			{
-				return (T)1. / Len;
+				return (FSolverReal)1. / Len;
 			}
 			return Fallback;
 		}
 
 	public:
-		void AddHyperelasticResidualAndHessian(const ParticleType& Particles, const int32 ElementIndex, const int32 ElementIndexLocal, const T Dt, TVec3<T>& ParticleResidual, Chaos::PMatrix<T, 3, 3>& ParticleHessian)
+		void AddHyperelasticResidualAndHessian(const FSolverParticles& Particles, const int32 ElementIndex, const int32 ElementIndexLocal, const FSolverReal Dt, TVec3<FSolverReal>& ParticleResidual, Chaos::PMatrix<FSolverReal, 3, 3>& ParticleHessian)
 		{
 			check(ElementIndex < DmInverse.Num() && ElementIndex > INDEX_NONE);
 			check(ElementIndex < MuElementArray.Num());
@@ -177,15 +178,15 @@ namespace Chaos::Softs
 				&& ElementIndex < DmInverse.Num() 
 				&& ElementIndex > INDEX_NONE)
 			{
-				const Chaos::PMatrix<T, 2, 2> DmInvT = DmInverse[ElementIndex].GetTransposed(), DmInv = DmInverse[ElementIndex]; 
-				const Chaos::PMatrix<T, 3, 2> Fe = F(ElementIndex, Particles);
+				const Chaos::PMatrix<FSolverReal, 2, 2> DmInvFSolverReal = DmInverse[ElementIndex].GetTransposed(), DmInv = DmInverse[ElementIndex]; 
+				const Chaos::PMatrix<FSolverReal, 3, 2> Fe = F(ElementIndex, Particles);
 
-				PMatrix<T, 3, 2> Pe(TVec3<T>((T)0.), TVec3<T>((T)0.)), ForceTerm(TVec3<T>((T)0.), TVec3<T>((T)0.));
-				Pe = T(2) * MuElementArray[ElementIndex] * Fe; 
+				PMatrix<FSolverReal, 3, 2> Pe(TVec3<FSolverReal>((FSolverReal)0.), TVec3<FSolverReal>((FSolverReal)0.)), ForceTerm(TVec3<FSolverReal>((FSolverReal)0.), TVec3<FSolverReal>((FSolverReal)0.));
+				Pe = FSolverReal(2) * MuElementArray[ElementIndex] * Fe; 
 
 				ForceTerm = -Measure[ElementIndex] * Pe * DmInverse[ElementIndex].GetTransposed();
 
-				Chaos::TVector<T, 3> Dx((T)0.);
+				Chaos::TVector<FSolverReal, 3> Dx((FSolverReal)0.);
 				if (ElementIndexLocal > 0)
 				{
 					for (int32 c = 0; c < 3; c++)
@@ -206,14 +207,14 @@ namespace Chaos::Softs
 
 				Dx *= Dt * Dt;
 				ParticleResidual -= Dx;
-				T Coeff = Dt * Dt * Measure[ElementIndex];
+				FSolverReal Coeff = Dt * Dt * Measure[ElementIndex];
 
 				if (ElementIndexLocal == 0) 
 				{
-					T DmInvsum = T(0);
+					FSolverReal DmInvsum = FSolverReal(0);
 					for (int32 nu = 0; nu < 2; nu++) 
 					{
-						T localDmsum = T(0);
+						FSolverReal localDmsum = FSolverReal(0);
 						for (int32 k = 0; k < 2; k++) 
 						{
 							localDmsum += DmInv.M[nu * 2 + k];
@@ -222,19 +223,19 @@ namespace Chaos::Softs
 					}
 					for (int32 alpha = 0; alpha < 3; alpha++) 
 					{
-						ParticleHessian.SetAt(alpha, alpha, ParticleHessian.GetAt(alpha, alpha) + Coeff * T(2) * MuElementArray[ElementIndex] * DmInvsum);
+						ParticleHessian.SetAt(alpha, alpha, ParticleHessian.GetAt(alpha, alpha) + Coeff * FSolverReal(2) * MuElementArray[ElementIndex] * DmInvsum);
 					}
 				}
 				else 
 				{
-					T DmInvsum = T(0);
+					FSolverReal DmInvsum = FSolverReal(0);
 					for (int32 nu = 0; nu < 2; nu++)
 					{
 						DmInvsum += DmInv.GetAt(ElementIndexLocal - 1, nu) * DmInv.GetAt(ElementIndexLocal - 1, nu);
 					}
 					for (int32 alpha = 0; alpha < 3; alpha++)
 					{
-						ParticleHessian.SetAt(alpha, alpha, ParticleHessian.GetAt(alpha, alpha) + Dt * Dt * Measure[ElementIndex] * T(2) * MuElementArray[ElementIndex] * DmInvsum);
+						ParticleHessian.SetAt(alpha, alpha, ParticleHessian.GetAt(alpha, alpha) + Dt * Dt * Measure[ElementIndex] * FSolverReal(2) * MuElementArray[ElementIndex] * DmInvsum);
 					}
 				}
 			}
@@ -244,16 +245,19 @@ namespace Chaos::Softs
 		mutable TArray<FSolverMatrix22> DmInverse;
 
 		//material constants calculated from E:
-		T Mu;
-		T Lambda;
-		TArray<T> MuElementArray;
-		TArray<T> LambdaElementArray;
-		TArray<T> AlphaJArray;
+		FSolverReal Mu;
+		FSolverReal Lambda;
+		TArray<FSolverReal> MuElementArray;
+		TArray<FSolverReal> LambdaElementArray;
+		TArray<FSolverReal> AlphaJArray;
 
 		TArray<TVector<int32, 3>> MeshConstraints;
-		mutable TArray<T> Measure;
+		mutable TArray<FSolverReal> Measure;
 	};
 
+	template <typename T, typename ParticleType>
+	using FGaussSeidelLinearCodimensionalConstraints UE_DEPRECATED(5.8, "Deprecated. this class is to be deleted, use FGaussSeidelLinearTriangleConstraints instead") = FGaussSeidelLinearTriangleConstraints;
+
 }  // End namespace Chaos::Softs
 
 
diff --git a/Engine/Source/Runtime/Experimental/Chaos/Public/Chaos/Deformable/GaussSeidelMainConstraint.h b/Engine/Source/Runtime/Experimental/Chaos/Public/Chaos/Deformable/GaussSeidelMainConstraint.h
--- a/Engine/Source/Runtime/Experimental/Chaos/Public/Chaos/Deformable/GaussSeidelMainConstraint.h
+++ b/Engine/Source/Runtime/Experimental/Chaos/Public/Chaos/Deformable/GaussSeidelMainConstraint.h
@@ -23,18 +23,17 @@ DEFINE_LOG_CATEGORY_STATIC(LogDeformableGaussSeidelMainConstraint, Log, All);
 
 namespace Chaos::Softs
 {
-	template <typename T, typename ParticleType>
-	class FGaussSeidelMainConstraint
+	class FGaussSeidelMainConstraints
 	{
 
 	public:
-		FGaussSeidelMainConstraint(
-			const ParticleType& InParticles, 
+		FGaussSeidelMainConstraints(
+			const FSolverParticles& InParticles, 
 			const bool bDoQuasistaticsIn = false,
 			const bool bDoSORIn = true,
-			const T InOmegaSOR = (T)1.6,
+			const FSolverReal InOmegaSOR = (FSolverReal)1.6,
 			const int32 ParallelMaxIn = 1000, 
-			const T MaxDxRatioIn = T(1),
+			const FSolverReal MaxDxRatioIn = FSolverReal(1),
 			const FDeformableXPBDCorotatedParams& InParams = FDeformableXPBDCorotatedParams())
 			: bDoQuasistatics(bDoQuasistaticsIn)
 			, bDoAcceleration(bDoSORIn)
@@ -46,7 +45,7 @@ namespace Chaos::Softs
 
 			InitializeLambdas();
 			
-			TVec3<T> MaxCoord((T)100.), MinCoord((T)-100.);
+			TVec3<FSolverReal> MaxCoord((FSolverReal)100.), MinCoord((FSolverReal)-100.);
 			for (int32 i = 0; i < (int32)InParticles.Size(); i++)
 			{
 				for (int32 j = 0; j < 3; j++) 
@@ -64,7 +63,7 @@ namespace Chaos::Softs
 			MaxDxSize = (MaxCoord - MinCoord).Size() * MaxDxRatioIn;
 		}
 
-		virtual ~FGaussSeidelMainConstraint() {}
+		virtual ~FGaussSeidelMainConstraints() {}
 
 		void Resize(const int32 NewSize)
 		{
@@ -75,20 +74,20 @@ namespace Chaos::Softs
 			DynamicIncidentElementsLocal.SetNum(NewSize);
 			TransientIncidentElements.SetNum(NewSize);
 			TransientIncidentElementsLocal.SetNum(NewSize);
-			X_k_1.Init(TVector<T, 3>(T(0.)), NewSize);
-			X_k.Init(TVector<T, 3>(T(0.)), NewSize);
+			X_k_1.Init(TVector<FSolverReal, 3>(FSolverReal(0.)), NewSize);
+			X_k.Init(TVector<FSolverReal, 3>(FSolverReal(0.)), NewSize);
 		}
 
-		const TArray<TFunction<void(const ParticleType&, const int32, const int32, const T, TVec3<T>&, Chaos::PMatrix<T, 3, 3>&)>>& StaticConstraintResidualAndHessian() const { return AddStaticConstraintResidualAndHessian; }
-		TArray<TFunction<void(const ParticleType&, const int32, const int32, const T, TVec3<T>&, Chaos::PMatrix<T, 3, 3>&)>>& StaticConstraintResidualAndHessian() { return AddStaticConstraintResidualAndHessian; }
-		const TArray<TFunction<void(const ParticleType&, const int32, const int32, const T, TVec3<T>&, Chaos::PMatrix<T, 3, 3>&)>>& TransientConstraintResidualAndHessian() const { return AddTransientConstraintResidualAndHessian; }
-		TArray<TFunction<void(const ParticleType&, const int32, const int32, const T, TVec3<T>&, Chaos::PMatrix<T, 3, 3>&)>>& TransientConstraintResidualAndHessian() { return AddTransientConstraintResidualAndHessian; }
-		const TArray<TFunction<void(const ParticleType&, const int32, const int32, const T, TVec3<T>&, Chaos::PMatrix<T, 3, 3>&)>>& DynamicConstraintResidualAndHessian() const { return AddDynamicConstraintResidualAndHessian; }
-		TArray<TFunction<void(const ParticleType&, const int32, const int32, const T, TVec3<T>&, Chaos::PMatrix<T, 3, 3>&)>>& DynamicConstraintResidualAndHessian() { return AddDynamicConstraintResidualAndHessian; }
-		const TArray<TFunction<void(const int32, const T, Chaos::PMatrix<T, 3, 3>&)>>& PerNodeHessian() const { return AddPerNodeHessian; }
-		TArray<TFunction<void(const int32, const T, Chaos::PMatrix<T, 3, 3>&)>>& PerNodeHessian() { return AddPerNodeHessian; }
-		const TArray<TFunction<void(const ParticleType&, const TArray<TVec3<T>>&, TArray<TVec3<T>>&)>>& InternalForceDifferentials() const { return AddInternalForceDifferentials; }
-		TArray<TFunction<void(const ParticleType&, const TArray<TVec3<T>>&, TArray<TVec3<T>>&)>>& InternalForceDifferentials() { return AddInternalForceDifferentials; }
+		const TArray<TFunction<void(const FSolverParticles&, const int32, const int32, const FSolverReal, TVec3<FSolverReal>&, Chaos::PMatrix<FSolverReal, 3, 3>&)>>& StaticConstraintResidualAndHessian() const { return AddStaticConstraintResidualAndHessian; }
+		TArray<TFunction<void(const FSolverParticles&, const int32, const int32, const FSolverReal, TVec3<FSolverReal>&, Chaos::PMatrix<FSolverReal, 3, 3>&)>>& StaticConstraintResidualAndHessian() { return AddStaticConstraintResidualAndHessian; }
+		const TArray<TFunction<void(const FSolverParticles&, const int32, const int32, const FSolverReal, TVec3<FSolverReal>&, Chaos::PMatrix<FSolverReal, 3, 3>&)>>& TransientConstraintResidualAndHessian() const { return AddTransientConstraintResidualAndHessian; }
+		TArray<TFunction<void(const FSolverParticles&, const int32, const int32, const FSolverReal, TVec3<FSolverReal>&, Chaos::PMatrix<FSolverReal, 3, 3>&)>>& TransientConstraintResidualAndHessian() { return AddTransientConstraintResidualAndHessian; }
+		const TArray<TFunction<void(const FSolverParticles&, const int32, const int32, const FSolverReal, TVec3<FSolverReal>&, Chaos::PMatrix<FSolverReal, 3, 3>&)>>& DynamicConstraintResidualAndHessian() const { return AddDynamicConstraintResidualAndHessian; }
+		TArray<TFunction<void(const FSolverParticles&, const int32, const int32, const FSolverReal, TVec3<FSolverReal>&, Chaos::PMatrix<FSolverReal, 3, 3>&)>>& DynamicConstraintResidualAndHessian() { return AddDynamicConstraintResidualAndHessian; }
+		const TArray<TFunction<void(const int32, const FSolverReal, Chaos::PMatrix<FSolverReal, 3, 3>&)>>& PerNodeHessian() const { return AddPerNodeHessian; }
+		TArray<TFunction<void(const int32, const FSolverReal, Chaos::PMatrix<FSolverReal, 3, 3>&)>>& PerNodeHessian() { return AddPerNodeHessian; }
+		const TArray<TFunction<void(const FSolverParticles&, const TArray<TVec3<FSolverReal>>&, TArray<TVec3<FSolverReal>>&)>>& InternalForceDifferentials() const { return AddInternalForceDifferentials; }
+		TArray<TFunction<void(const FSolverParticles&, const TArray<TVec3<FSolverReal>>&, TArray<TVec3<FSolverReal>>&)>>& InternalForceDifferentials() { return AddInternalForceDifferentials; }
 
 		int32 AddStaticConstraintResidualAndHessianRange(int32 NumConstraints)
 		{
@@ -143,7 +142,7 @@ namespace Chaos::Softs
 		}
 
 
-		void Apply(ParticleType& Particles, const T Dt, const int32 MaxWriteIters = 10, const bool Write2File = false, const TPBDActiveView<FSolverParticles>* InParticleActiveView = nullptr)
+		void Apply(FSolverParticles& Particles, const FSolverReal Dt, const int32 MaxWriteIters = 10, const bool Write2File = false, const TPBDActiveView<FSolverParticles>* InParticleActiveView = nullptr)
 		{
  			PERF_SCOPE(STAT_ChaosGSMainConstraint_Apply);
 			
@@ -168,7 +167,7 @@ namespace Chaos::Softs
 							int32 TaskIndex = CorotatedParams.XPBDCorotatedBatchSize * BatchIndex + BatchSubIndex;
 							if (TaskIndex < ParticlesPerColor[i].Num())
 							{
-								if (Particles.InvM(ParticlesPerColor[i][TaskIndex]) != T(0))
+								if (Particles.InvM(ParticlesPerColor[i][TaskIndex]) != FSolverReal(0))
 								{
 									int32 ParticleIndex = ParticlesPerColor[i][TaskIndex];
 									if (!ApplySingleParticle(ParticleIndex, Dt, Particles))
@@ -207,14 +206,14 @@ namespace Chaos::Softs
 			}
 		}
 
-		void InitStaticColor(const ParticleType& Particles, const TPBDActiveView<FSolverParticles>* InParticleActiveView = nullptr)
+		void InitStaticColor(const FSolverParticles& Particles, const TPBDActiveView<FSolverParticles>* InParticleActiveView = nullptr)
 		{
 			StaticParticlesPerColor = ComputeNodalColoring(StaticConstraints, Particles, 0, Particles.Size(), StaticIncidentElements, StaticIncidentElementsLocal, InParticleActiveView,  &StaticParticleColors);
 			ParticleColors = StaticParticleColors;
 			ParticlesPerColor = StaticParticlesPerColor;
 		}
 
-		void InitTransientColor(const ParticleType& Particles)
+		void InitTransientColor(const FSolverParticles& Particles)
 		{
 			PERF_SCOPE(STAT_ChaosGSMainConstraint_InitTransientColor);
 			ParticleColors = StaticParticleColors;
@@ -222,7 +221,7 @@ namespace Chaos::Softs
 			Chaos::ComputeExtraNodalColoring(StaticConstraints, DynamicConstraints, TransientConstraints, Particles, StaticIncidentElements, DynamicIncidentElements, TransientIncidentElements, ParticleColors, ParticlesPerColor);
 		}
 
-		void InitDynamicColor(const ParticleType& Particles)
+		void InitDynamicColor(const FSolverParticles& Particles)
 		{
 			PERF_SCOPE(STAT_ChaosGSMainConstraint_InitDynamicColor);
 			ParticleColors = StaticParticleColors;
@@ -230,7 +229,7 @@ namespace Chaos::Softs
 			Chaos::ComputeExtraNodalColoring(StaticConstraints, DynamicConstraints, Particles, StaticIncidentElements, TransientIncidentElements, ParticleColors, ParticlesPerColor);
 		}
 
-		void Init(const T Dt, const ParticleType& Particles)
+		void Init(const FSolverReal Dt, const FSolverParticles& Particles)
 		{
 			Resize((int32)Particles.Size());
 
@@ -281,22 +280,22 @@ namespace Chaos::Softs
 			return false;
 		}
 
-		CHAOS_API TArray<TVec3<T>> ComputeNewtonResiduals(const ParticleType& Particles, const T Dt, const bool Write2File = false, TArray<PMatrix<T, 3, 3>>* AllParticleHessian = nullptr);
+		CHAOS_API TArray<TVec3<FSolverReal>> ComputeNewtonResiduals(const FSolverParticles& Particles, const FSolverReal Dt, const bool Write2File = false, TArray<PMatrix<FSolverReal, 3, 3>>* AllParticleHessian = nullptr);
 		
-		CHAOS_API void ApplyCG(ParticleType& Particles, const T Dt);
+		CHAOS_API void ApplyCG(FSolverParticles& Particles, const FSolverReal Dt);
 
 		/* Adds external acceleration, e.g. gravity (0,0,-980) cm/s^2 */
-		void AddExternalAcceleration(const TVec3<T>& Acceleration) { ExternalAcceleration += Acceleration; };
+		void AddExternalAcceleration(const TVec3<FSolverReal>& Acceleration) { ExternalAcceleration += Acceleration; };
 
-		void ResetExternalAcceleration() { ExternalAcceleration = TVec3<T>((T)0.); };
+		void ResetExternalAcceleration() { ExternalAcceleration = TVec3<FSolverReal>((FSolverReal)0.); };
 
 	private:
 
-		bool ApplySingleParticle(const int32 p, const T Dt, ParticleType& Particles)
+		bool ApplySingleParticle(const int32 p, const FSolverReal Dt, FSolverParticles& Particles)
 		{
 			int32 ConstraintIndex = 0;
-			Chaos::TVector<T, 3> ParticleResidual((T)0.);
-			Chaos::PMatrix<T, 3, 3> ParticleHessian((T)0., (T)0., (T)0.);
+			Chaos::TVector<FSolverReal, 3> ParticleResidual((FSolverReal)0.);
+			Chaos::PMatrix<FSolverReal, 3, 3> ParticleHessian((FSolverReal)0., (FSolverReal)0., (FSolverReal)0.);
 
 			ComputeInitialResidualAndHessian(Particles, p, Dt, ParticleResidual, ParticleHessian);
 
@@ -339,17 +338,17 @@ namespace Chaos::Softs
 				AddPerNodeHessian[i](p, Dt, ParticleHessian);
 			}
 
-			T HessianScale = (T)1.;
-			T HessianDet = ParticleHessian.Determinant();
-			auto IsIllConditioned = [](T HessianDet)
+			FSolverReal HessianScale = (FSolverReal)1.;
+			FSolverReal HessianDet = ParticleHessian.Determinant();
+			auto IsIllConditioned = [](FSolverReal HessianDet)
 			{
-				return FMath::Abs(HessianDet) < TMathUtilConstants<T>::Epsilon ||
-					FMath::Abs(HessianDet) > TMathUtilConstants<T>::MaxReal;
+				return FMath::Abs(HessianDet) < TMathUtilConstants<FSolverReal>::Epsilon ||
+					FMath::Abs(HessianDet) > TMathUtilConstants<FSolverReal>::MaxReal;
 			};
 			if (IsIllConditioned(HessianDet))
 			{
 				// scale the hessian so that the determinant (~HessianScale^3) falls into the normal range.
-				HessianScale = (T)0.;
+				HessianScale = (FSolverReal)0.;
 				for (int32 RowIdx = 0; RowIdx < 3; ++RowIdx)
 				{
 					for (int32 ColIdx = 0; ColIdx < 3; ++ColIdx)
@@ -372,9 +371,9 @@ namespace Chaos::Softs
 				}
 			}
 
-			Chaos::PMatrix<T, 3, 3> HessianInv = ParticleHessian.SymmetricCofactorMatrix();
-			HessianInv *= T(1) / HessianDet;
-			Chaos::TVector<T, 3> Dx = HessianInv.GetTransposed() * (-ParticleResidual) / HessianScale; // add back HessianScale
+			Chaos::PMatrix<FSolverReal, 3, 3> HessianInv = ParticleHessian.SymmetricCofactorMatrix();
+			HessianInv *= FSolverReal(1) / HessianDet;
+			Chaos::TVector<FSolverReal, 3> Dx = HessianInv.GetTransposed() * (-ParticleResidual) / HessianScale; // add back HessianScale
 
 			if (Dx.Size() < MaxDxSize)
 			{
@@ -389,7 +388,7 @@ namespace Chaos::Softs
 		
 		void InitializeLambdas()
 		{
-			ComputeInitialResidualAndHessian = [this](const ParticleType& Particles, const int32 p, const T Dt, TVec3<T>& ParticleResidual, Chaos::PMatrix<T, 3, 3>& ParticleHessian)
+			ComputeInitialResidualAndHessian = [this](const FSolverParticles& Particles, const int32 p, const FSolverReal Dt, TVec3<FSolverReal>& ParticleResidual, Chaos::PMatrix<FSolverReal, 3, 3>& ParticleHessian)
 			{
 				if (!this->bDoQuasistatics) {
 					for (int32 alpha = 0; alpha < 3; alpha++) {
@@ -407,13 +406,13 @@ namespace Chaos::Softs
 				}
 			};
 
-			AccelerationTechniquePerParticle = [this](ParticleType& Particles, int32 ParticleIndex)
+			AccelerationTechniquePerParticle = [this](FSolverParticles& Particles, int32 ParticleIndex)
 			{
-				if (Particles.InvM(ParticleIndex) != T(0) && CurrentIt > SORStart)
+				if (Particles.InvM(ParticleIndex) != FSolverReal(0) && CurrentIt > SORStart)
 				{
 					Particles.P(ParticleIndex) = OmegaSOR * (Particles.P(ParticleIndex) - this->X_k_1[ParticleIndex]) + this->X_k_1[ParticleIndex];
 				}
-				if (Particles.InvM(ParticleIndex) != T(0))
+				if (Particles.InvM(ParticleIndex) != FSolverReal(0))
 				{
 					this->X_k_1[ParticleIndex] = this->X_k[ParticleIndex];
 					this->X_k[ParticleIndex] = Particles.P(ParticleIndex);
@@ -434,11 +433,11 @@ namespace Chaos::Softs
 		TArray<TArray<int32>> DynamicIncidentElementsLocal;
 
 		//Lambds for spcifying residual/hessian computations:
-		TFunction<void(const ParticleType&, const int32, const T, TVec3<T>&, Chaos::PMatrix<T, 3, 3>&)> ComputeInitialResidualAndHessian;
-		TArray<TFunction<void(const ParticleType&, const int32, const int32, const T, TVec3<T>&, Chaos::PMatrix<T, 3, 3>&)>> AddStaticConstraintResidualAndHessian;
-		TArray<TFunction<void(const ParticleType&, const int32, const int32, const T, TVec3<T>&, Chaos::PMatrix<T, 3, 3>&)>> AddDynamicConstraintResidualAndHessian;
-		TArray<TFunction<void(const ParticleType&, const int32, const int32, const T, TVec3<T>&, Chaos::PMatrix<T, 3, 3>&)>> AddTransientConstraintResidualAndHessian;
-		TArray<TFunction<void(const int32, const T, Chaos::PMatrix<T, 3, 3>&)>> AddPerNodeHessian;
+		TFunction<void(const FSolverParticles&, const int32, const FSolverReal, TVec3<FSolverReal>&, Chaos::PMatrix<FSolverReal, 3, 3>&)> ComputeInitialResidualAndHessian;
+		TArray<TFunction<void(const FSolverParticles&, const int32, const int32, const FSolverReal, TVec3<FSolverReal>&, Chaos::PMatrix<FSolverReal, 3, 3>&)>> AddStaticConstraintResidualAndHessian;
+		TArray<TFunction<void(const FSolverParticles&, const int32, const int32, const FSolverReal, TVec3<FSolverReal>&, Chaos::PMatrix<FSolverReal, 3, 3>&)>> AddDynamicConstraintResidualAndHessian;
+		TArray<TFunction<void(const FSolverParticles&, const int32, const int32, const FSolverReal, TVec3<FSolverReal>&, Chaos::PMatrix<FSolverReal, 3, 3>&)>> AddTransientConstraintResidualAndHessian;
+		TArray<TFunction<void(const int32, const FSolverReal, Chaos::PMatrix<FSolverReal, 3, 3>&)>> AddPerNodeHessian;
 
 		//Coloring information:
 		TArray<int32> StaticParticleColors;
@@ -452,15 +451,15 @@ namespace Chaos::Softs
 		TArray<int32> DynamicIncidentElementsOffsets;
 
 		bool bDoQuasistatics = false;
-		mutable TArray<Chaos::TVector<T, 3>> xtilde;
+		mutable TArray<Chaos::TVector<FSolverReal, 3>> xtilde;
 
 		//SOR variables:
-		TFunction<void(ParticleType&, int32)> AccelerationTechniquePerParticle;
-		mutable TArray<Chaos::TVector<T, 3>> X_k_1;
-		mutable TArray<Chaos::TVector<T, 3>> X_k;
+		TFunction<void(FSolverParticles&, int32)> AccelerationTechniquePerParticle;
+		mutable TArray<Chaos::TVector<FSolverReal, 3>> X_k_1;
+		mutable TArray<Chaos::TVector<FSolverReal, 3>> X_k;
 		int32 CurrentIt = 0;
 		bool bDoAcceleration = true;
-		T OmegaSOR = T(1.6);
+		FSolverReal OmegaSOR = FSolverReal(1.6);
 		int32 SORStart = 1; //1 is smallest
 
 		int32 ParallelMax = 1000;
@@ -468,23 +467,25 @@ namespace Chaos::Softs
 		FDeformableXPBDCorotatedParams CorotatedParams;
 
 		//Newton solver variables:
-		TArray<TFunction<void(const ParticleType&, const TArray<TVec3<T>>&, TArray<TVec3<T>>&)>> AddInternalForceDifferentials;
+		TArray<TFunction<void(const FSolverParticles&, const TArray<TVec3<FSolverReal>>&, TArray<TVec3<FSolverReal>>&)>> AddInternalForceDifferentials;
 		TUniquePtr<TArray<int32>> use_list;
 
 		int32 NumTotalParticles;
-		TArray<TVec3<T>> ReorderedPs;
+		TArray<TVec3<FSolverReal>> ReorderedPs;
 
-		T MaxDxSize = T(0);
+		FSolverReal MaxDxSize = FSolverReal(0);
 
 		public:
 		bool DebugResidual = false;
 		bool IsFirstFrame = true;
 		int32 PassedIters = 0;
 
-		TVec3<T> ExternalAcceleration = TVec3<T>((T)0.);
+		TVec3<FSolverReal> ExternalAcceleration = TVec3<FSolverReal>((FSolverReal)0.);
 
 	};
 
-	//template class FGaussSeidelMainConstraint<FSolverReal, FSolverParticles>;
+	template <typename T, typename ParticleType>
+	using FGaussSeidelMainConstraint UE_DEPRECATED(5.8, "Deprecated. this class is to be deleted, use FGaussSeidelMainConstraints instead") = FGaussSeidelMainConstraints;
+
 }
 
diff --git a/Engine/Source/Runtime/Experimental/Chaos/Public/Chaos/Deformable/GaussSeidelNeohookeanConstraints.h b/Engine/Source/Runtime/Experimental/Chaos/Public/Chaos/Deformable/GaussSeidelNeohookeanConstraints.h
--- a/Engine/Source/Runtime/Experimental/Chaos/Public/Chaos/Deformable/GaussSeidelNeohookeanConstraints.h
+++ b/Engine/Source/Runtime/Experimental/Chaos/Public/Chaos/Deformable/GaussSeidelNeohookeanConstraints.h
@@ -13,10 +13,11 @@
 
 namespace Chaos::Softs
 {
-	template <typename T, typename ParticleType>
-	class FGaussSeidelNeohookeanConstraints : public FGaussSeidelCorotatedConstraints <T, ParticleType>
+	// Tetrahedral constraint using Neohookean constitutive model
+	// for position-based nonlinear Gauss Seidel solver [Chen et al., 2024]
+	class FGaussSeidelNeohookeanTetrahedralConstraints : public FGaussSeidelCorotatedTetrahedralConstraints
 	{
-		typedef FGaussSeidelCorotatedConstraints<T, ParticleType> GSBase;
+		typedef FGaussSeidelCorotatedTetrahedralConstraints GSBase;
 		using GSBase::ParticlesPerColor;
 		using GSBase::IncidentElements;
 		using GSBase::IncidentElementsLocal;
@@ -26,7 +27,7 @@ namespace Chaos::Softs
 		using GSBase::bDoQuasistatics;
 		using GSBase::ComputeStress;
 		using GSBase::ComputeHessianHelper;
-		typedef FXPBDCorotatedConstraints<T, ParticleType> Base;
+		typedef FXPBDCorotatedTetrahedralConstraints Base;
 		using Base::MeshConstraints;
 		using Base::LambdaArray;
 		using Base::Measure;
@@ -36,21 +37,21 @@ namespace Chaos::Softs
 		using Base::CorotatedParams;
 
 	public:
-		FGaussSeidelNeohookeanConstraints(
-			const ParticleType& InParticles,
+		FGaussSeidelNeohookeanTetrahedralConstraints(
+			const FSolverParticles& InParticles,
 			const TArray<TVector<int32, 4>>& InMesh,
-			const TArray<T>& EMeshArray,
-			const TArray<T>& NuMeshArray,
-			TArray<T>&& AlphaJMeshArray,
+			const TArray<FSolverReal>& EMeshArray,
+			const TArray<FSolverReal>& NuMeshArray,
+			TArray<FSolverReal>&& AlphaJMeshArray,
 			TArray<TArray<int32>>&& IncidentElementsIn,
 			TArray<TArray<int32>>&& IncidentElementsLocalIn,
 			const int32 ParticleStartIndexIn,
 			const int32 ParticleEndIndexIn,
 			const bool bDoQuasistaticsIn = false,
 			const bool bDoSORIn = true,
-			const T InOmegaSOR = (T)1.6,
+			const FSolverReal InOmegaSOR = (FSolverReal)1.6,
 			const FDeformableXPBDCorotatedParams& InParams = FDeformableXPBDCorotatedParams(),
-			const T& NuMesh = (T).3,
+			const FSolverReal& NuMesh = (FSolverReal).3,
 			const bool bRecordMetricIn = false
 		)
 			: GSBase(InParticles, InMesh, EMeshArray, NuMeshArray, MoveTemp(AlphaJMeshArray), MoveTemp(IncidentElementsIn), MoveTemp(IncidentElementsLocalIn), ParticleStartIndexIn, ParticleEndIndexIn, bDoQuasistaticsIn, bDoSORIn, InOmegaSOR, InParams, NuMesh, bRecordMetricIn)
@@ -58,20 +59,20 @@ namespace Chaos::Softs
 			InitializeNeohookeanLambdas();
 		}
 
-		virtual ~FGaussSeidelNeohookeanConstraints() {}
+		virtual ~FGaussSeidelNeohookeanTetrahedralConstraints() {}
 
 	protected:
 
 		void InitializeNeohookeanLambdas() 
 		{
-			ComputeStress = [](const Chaos::PMatrix<T, 3, 3>& Fe, const T mu, const T lambda, Chaos::PMatrix<T, 3, 3>& P)
+			ComputeStress = [](const Chaos::PMatrix<FSolverReal, 3, 3>& Fe, const FSolverReal mu, const FSolverReal lambda, Chaos::PMatrix<FSolverReal, 3, 3>& P)
 			{
 				PNeohookeanMM(Fe, mu, lambda, P);
 			};
-			ComputeHessianHelper = [](const Chaos::PMatrix<T, 3, 3>& Fe, const Chaos::PMatrix<T, 3, 3>& DmInv, const T mu, const T lambda, const int32 local_index, const T Coeff, Chaos::PMatrix<T, 3, 3>& final_hessian)
+			ComputeHessianHelper = [](const Chaos::PMatrix<FSolverReal, 3, 3>& Fe, const Chaos::PMatrix<FSolverReal, 3, 3>& DmInv, const FSolverReal mu, const FSolverReal lambda, const int32 local_index, const FSolverReal Coeff, Chaos::PMatrix<FSolverReal, 3, 3>& final_hessian)
 			{
-				T LambdaHat = mu + lambda;
-				Chaos::PMatrix<T, 3, 3> JFinvT((T)0.);
+				FSolverReal LambdaHat = mu + lambda;
+				Chaos::PMatrix<FSolverReal, 3, 3> JFinvT((FSolverReal)0.);
 				JFinvT.SetAt(0, 0, Fe.GetAt(1, 1) * Fe.GetAt(2, 2) - Fe.GetAt(2, 1) * Fe.GetAt(1, 2));
 				JFinvT.SetAt(0, 1, Fe.GetAt(2, 0) * Fe.GetAt(1, 2) - Fe.GetAt(1, 0) * Fe.GetAt(2, 2));
 				JFinvT.SetAt(0, 2, Fe.GetAt(1, 0) * Fe.GetAt(2, 1) - Fe.GetAt(2, 0) * Fe.GetAt(1, 1));
@@ -82,12 +83,12 @@ namespace Chaos::Softs
 				JFinvT.SetAt(2, 1, Fe.GetAt(1, 0) * Fe.GetAt(0, 2) - Fe.GetAt(0, 0) * Fe.GetAt(1, 2));
 				JFinvT.SetAt(2, 2, Fe.GetAt(0, 0) * Fe.GetAt(1, 1) - Fe.GetAt(1, 0) * Fe.GetAt(0, 1));
 
-				Chaos::PMatrix<T, 3, 3> JFinv = JFinvT.GetTransposed();
+				Chaos::PMatrix<FSolverReal, 3, 3> JFinv = JFinvT.GetTransposed();
 
 				if (local_index == 0) {
-					T DmInvsum = T(0);
+					FSolverReal DmInvsum = FSolverReal(0);
 					for (int32 nu = 0; nu < 3; nu++) {
-						T localDmsum = T(0);
+						FSolverReal localDmsum = FSolverReal(0);
 						for (int32 k = 0; k < 3; k++) {
 							localDmsum += DmInv.GetAt(k, nu);
 						}
@@ -97,8 +98,8 @@ namespace Chaos::Softs
 						final_hessian.SetAt(alpha, alpha, final_hessian.GetAt(alpha, alpha) + Coeff * mu * DmInvsum);
 					}
 
-					Chaos::PMatrix<T, 3, 3> DmInvJFinv = JFinv * DmInv;
-					Chaos::TVector<T, 3> l((T)0.);
+					Chaos::PMatrix<FSolverReal, 3, 3> DmInvJFinv = JFinv * DmInv;
+					Chaos::TVector<FSolverReal, 3> l((FSolverReal)0.);
 					for (int32 alpha = 0; alpha < 3; alpha++) {
 						for (int32 k = 0; k < 3; k++) {
 							l[alpha] += DmInvJFinv.GetAt(k, alpha);
@@ -111,7 +112,7 @@ namespace Chaos::Softs
 
 				}
 				else {
-					T DmInvsum = T(0);
+					FSolverReal DmInvsum = FSolverReal(0);
 					for (int32 nu = 0; nu < 3; nu++)
 					{
 						DmInvsum += DmInv.GetAt(local_index - 1, nu) * DmInv.GetAt(local_index - 1, nu);
@@ -121,8 +122,8 @@ namespace Chaos::Softs
 						final_hessian.SetAt(alpha, alpha, final_hessian.GetAt(alpha, alpha) + Coeff * mu * DmInvsum);
 					}
 
-					Chaos::PMatrix<T, 3, 3> DmInvJFinv = JFinv * DmInv;
-					Chaos::TVector<T, 3> l((T)0.);
+					Chaos::PMatrix<FSolverReal, 3, 3> DmInvJFinv = JFinv * DmInv;
+					Chaos::TVector<FSolverReal, 3> l((FSolverReal)0.);
 					for (int32 alpha = 0; alpha < 3; alpha++) {
 						l[alpha] = DmInvJFinv.GetAt(local_index - 1, alpha);
 					}
@@ -136,4 +137,7 @@ namespace Chaos::Softs
 
 	};
 
+	template <typename T, typename ParticleType>
+	using FGaussSeidelNeohookeanConstraints UE_DEPRECATED(5.8, "Deprecated. this class is to be deleted, use FGaussSeidelNeohookeanTetrahedralConstraints instead") = FGaussSeidelNeohookeanTetrahedralConstraints;
+
 }// End namespace Chaos::Softs
diff --git a/Engine/Source/Runtime/Experimental/Chaos/Public/Chaos/Deformable/GaussSeidelSphereRepulsionConstraints.h b/Engine/Source/Runtime/Experimental/Chaos/Public/Chaos/Deformable/GaussSeidelSphereRepulsionConstraints.h
--- a/Engine/Source/Runtime/Experimental/Chaos/Public/Chaos/Deformable/GaussSeidelSphereRepulsionConstraints.h
+++ b/Engine/Source/Runtime/Experimental/Chaos/Public/Chaos/Deformable/GaussSeidelSphereRepulsionConstraints.h
@@ -22,12 +22,12 @@
 namespace Chaos::Softs
 {
 	using Chaos::TVec3;
-	template <typename T, typename ParticleType>
-	class FGaussSeidelSphereRepulsionConstraints
+
+	class FGaussSeidelParticleSphereRepulsionConstraints
 	{
 	public:
 		//TODO(Yushan): Add unit tests for Gauss Seidel Sphere Repulsion Constraints
-		FGaussSeidelSphereRepulsionConstraints(FSolverReal InRadius, FSolverReal InStiffness, const ParticleType& InParticles, const FDeformableXPBDWeakConstraintParams& InParams): Radius(InRadius), Stiffness(InStiffness), DebugDrawParams(InParams)
+		FGaussSeidelParticleSphereRepulsionConstraints(FSolverReal InRadius, FSolverReal InStiffness, const FSolverParticles& InParticles, const FDeformableXPBDWeightedSpringConstraintParams& InParams): Radius(InRadius), Stiffness(InStiffness), DebugDrawParams(InParams)
 		{
 			const TArrayCollectionArray<FPAndInvM>& PAndInvM = InParticles.GetPAndInvM();
 			ReferencePositions.SetNum(PAndInvM.Num());
@@ -37,43 +37,43 @@ namespace Chaos::Softs
 			}
 		}
 
-		virtual ~FGaussSeidelSphereRepulsionConstraints(){}
+		virtual ~FGaussSeidelParticleSphereRepulsionConstraints(){}
 
 		//Energy = k/2*(2r-d)^2
 		//Residual = de/dx = -force = -k*(2r-d)*dd/dx
 		//Hessian = de2/dx2 = k*dd/dx*dd/dx-k*(2r-d)*(-dd/dx*dd/dx^T+I)/d
-		void AddSphereRepulsionResidual(const ParticleType& InParticles, const int32 p, const T Dt, TVec3<T>& res)
+		void AddSphereRepulsionResidual(const FSolverParticles& InParticles, const int32 p, const FSolverReal Dt, TVec3<FSolverReal>& res)
 		{
 
 		}
 
-		void AddSphereRepulsionHessian(const int32 p, const T Dt, Chaos::PMatrix<T, 3, 3>& ParticleHessian)
+		void AddSphereRepulsionHessian(const int32 p, const FSolverReal Dt, Chaos::PMatrix<FSolverReal, 3, 3>& ParticleHessian)
 		{
 
 		}
 
-		void AddSphereRepulsionResidualAndHessian(const ParticleType& InParticles, const int32 ConstraintIndex, const int32 LocalIndex, const T Dt, TVec3<T>& ParticleResidual, Chaos::PMatrix<T, 3, 3>& ParticleHessian)
+		void AddSphereRepulsionResidualAndHessian(const FSolverParticles& InParticles, const int32 ConstraintIndex, const int32 LocalIndex, const FSolverReal Dt, TVec3<FSolverReal>& ParticleResidual, Chaos::PMatrix<FSolverReal, 3, 3>& ParticleHessian)
 		{
-			Chaos::TVec3<T> x0 = InParticles.P(Constraints[ConstraintIndex][0]);
-			Chaos::TVec3<T> x1 = InParticles.P(Constraints[ConstraintIndex][1]);
-			Chaos::TVec3<T> normal = (x1 - x0).GetSafeNormal();
+			Chaos::TVec3<FSolverReal> x0 = InParticles.P(Constraints[ConstraintIndex][0]);
+			Chaos::TVec3<FSolverReal> x1 = InParticles.P(Constraints[ConstraintIndex][1]);
+			Chaos::TVec3<FSolverReal> normal = (x1 - x0).GetSafeNormal();
 			normal = (LocalIndex ? normal : -normal); //dd/dx
-			T dist = (x1 - x0).Size();
+			FSolverReal dist = (x1 - x0).Size();
 			FSolverReal penetration = 2 * Radius - dist; //2r-d
 
-			if (penetration > T(0)) {
-				T dist_inv = T(1) / (dist+T(1e-12));
-				Chaos::PMatrix<T, 3, 3> outer = Chaos::PMatrix<T, 3, 3>::OuterProduct(normal, normal);
-				Chaos::PMatrix<T, 3, 3> A = (Chaos::PMatrix<T, 3, 3>::Identity - outer) * dist_inv; //(-dd/dx*dd/dx^T+I)/d
+			if (penetration > FSolverReal(0)) {
+				FSolverReal dist_inv = FSolverReal(1) / (dist+FSolverReal(1e-12));
+				Chaos::PMatrix<FSolverReal, 3, 3> outer = Chaos::PMatrix<FSolverReal, 3, 3>::OuterProduct(normal, normal);
+				Chaos::PMatrix<FSolverReal, 3, 3> A = (Chaos::PMatrix<FSolverReal, 3, 3>::Identity - outer) * dist_inv; //(-dd/dx*dd/dx^T+I)/d
 				ParticleHessian += Dt * Dt * ConstraintStiffness[ConstraintIndex] * (outer - penetration * A);
 				ParticleResidual += -Dt * Dt * penetration * ConstraintStiffness[ConstraintIndex] * normal;
 			}
 		}
 
-		void VisualizeAllBindings(const FSolverParticles& InParticles, const T Dt) const
+		void VisualizeAllBindings(const FSolverParticles& InParticles, const FSolverReal Dt) const
 		{
 #if WITH_EDITOR
-			auto DoubleVert = [](Chaos::TVec3<T> V) { return FVector3d(V.X, V.Y, V.Z); };
+			auto DoubleVert = [](Chaos::TVec3<FSolverReal> V) { return FVector3d(V.X, V.Y, V.Z); };
 			for (int32 i = 0; i < Constraints.Num(); i++)
 			{
 				FVector3d SourcePos = DoubleVert(InParticles.P(Constraints[i][0]));
@@ -89,7 +89,7 @@ namespace Chaos::Softs
 #endif
 		}
 
-		void Init(const FSolverParticles& InParticles, const T Dt) const
+		void Init(const FSolverParticles& InParticles, const FSolverReal Dt) const
 		{
 			if (DebugDrawParams.bVisualizeBindings)
 			{
@@ -149,7 +149,7 @@ namespace Chaos::Softs
 
 			for (const TVec2<int32>& CollisionPair : Constraints)
 			{
-				T TotalMass = T(0);
+				FSolverReal TotalMass = FSolverReal(0);
 				if (Particles.InvM(CollisionPair[0]) > (FSolverReal)1e-7)
 				{
 					TotalMass += Particles.M(CollisionPair[0]);
@@ -158,7 +158,7 @@ namespace Chaos::Softs
 				{
 					TotalMass += Particles.M(CollisionPair[1]);
 				}
-				ConstraintStiffness.Add(Stiffness*TotalMass/T(2));
+				ConstraintStiffness.Add(Stiffness*TotalMass/FSolverReal(2));
 			}
 		}
 
@@ -178,7 +178,7 @@ namespace Chaos::Softs
 			TArray<TVec2<int32>> Constraints;
 			FSolverReal Radius = FSolverReal(0);
 			FSolverReal Stiffness = FSolverReal(0);
-			TArray<T> ConstraintStiffness;
+			TArray<FSolverReal> ConstraintStiffness;
 		private:
 			struct FSphereSpatialEntry
 			{
@@ -197,8 +197,10 @@ namespace Chaos::Softs
 				}
 			};
 			TArray<FSolverVec3> ReferencePositions;
-			FDeformableXPBDWeakConstraintParams DebugDrawParams;
+			FDeformableXPBDWeightedSpringConstraintParams DebugDrawParams;
 	};
 
+	template <typename T, typename ParticleType>
+	using FGaussSeidelSphereRepulsionConstraints UE_DEPRECATED(5.8, "Deprecated. this class is to be deleted, use FGaussSeidelParticleSphereRepulsionConstraints instead") = FGaussSeidelParticleSphereRepulsionConstraints;
 
 }// End namespace Chaos::Softs
diff --git a/Engine/Source/Runtime/Experimental/Chaos/Public/Chaos/Deformable/GaussSeidelUnilateralTetConstraints.h b/Engine/Source/Runtime/Experimental/Chaos/Public/Chaos/Deformable/GaussSeidelUnilateralTetConstraints.h
--- a/Engine/Source/Runtime/Experimental/Chaos/Public/Chaos/Deformable/GaussSeidelUnilateralTetConstraints.h
+++ b/Engine/Source/Runtime/Experimental/Chaos/Public/Chaos/Deformable/GaussSeidelUnilateralTetConstraints.h
@@ -2,17 +2,19 @@
 #pragma once
 
 #include "Chaos/Utilities.h"
+#include "Chaos/PBDSoftsEvolutionFwd.h"
+#include "Chaos/PBDSoftsSolverParticles.h"
 
 namespace Chaos::Softs
 {
-	template <typename T, typename ParticleType>
-	class FGaussSeidelUnilateralTetConstraints
+	// Tetrahdral volume constraint for position-based nonlinear Gauss Seidel solver [Chen et al., 2024]
+	class FGaussSeidelUnilateralVolumeConstraints
 	{
 	public: 
-		FGaussSeidelUnilateralTetConstraints(
-			const ParticleType& Particles,
+		FGaussSeidelUnilateralVolumeConstraints(
+			const FSolverParticles& Particles,
 			TArray<TVector<int32, 4>>&& InConstraints,
-			TArray<T> InStiffnessArray)
+			TArray<FSolverReal> InStiffnessArray)
 			: Constraints(InConstraints), 
 			StiffnessArray(InStiffnessArray)
 		{
@@ -31,26 +33,26 @@ namespace Chaos::Softs
 			}
 		};
 
-		static T ComputeVolume(const TVec3<T>& P1, const TVec3<T>& P2, const TVec3<T>& P3, const TVec3<T>& P4)
+		static FSolverReal ComputeVolume(const TVec3<FSolverReal>& P1, const TVec3<FSolverReal>& P2, const TVec3<FSolverReal>& P3, const TVec3<FSolverReal>& P4)
 		{
-			const TVec3<T> P2P1 = P2 - P1;
-			const TVec3<T> P4P1 = P4 - P1;
-			const TVec3<T> P3P1 = P3 - P1;
-			return TVec3<T>::DotProduct(TVec3<T>::CrossProduct(P2P1, P3P1), P4P1) / (T)6.;
+			const TVec3<FSolverReal> P2P1 = P2 - P1;
+			const TVec3<FSolverReal> P4P1 = P4 - P1;
+			const TVec3<FSolverReal> P3P1 = P3 - P1;
+			return TVec3<FSolverReal>::DotProduct(TVec3<FSolverReal>::CrossProduct(P2P1, P3P1), P4P1) / (FSolverReal)6.;
 		}
 
-		void AddEnergy(const ParticleType& InParticles, const int32 ConstraintIndex, const T Dt, T& Energy) const
+		void AddEnergy(const FSolverParticles& InParticles, const int32 ConstraintIndex, const FSolverReal Dt, FSolverReal& Energy) const
 		{
 			const TVec4<int32>& Constraint = Constraints[ConstraintIndex];
-			T CurrentVol = ComputeVolume(
+			FSolverReal CurrentVol = ComputeVolume(
 				InParticles.P(Constraint[0]),
 				InParticles.P(Constraint[1]),
 				InParticles.P(Constraint[2]),
 				InParticles.P(Constraint[3]));
-			Energy += Dt * Dt * StiffnessArray[ConstraintIndex] / (T)2. * (CurrentVol - Volumes[ConstraintIndex]) * (CurrentVol - Volumes[ConstraintIndex]);
+			Energy += Dt * Dt * StiffnessArray[ConstraintIndex] / (FSolverReal)2. * (CurrentVol - Volumes[ConstraintIndex]) * (CurrentVol - Volumes[ConstraintIndex]);
 		}
 
-		void AddResidualAndHessian(const ParticleType& InParticles, const int32 ConstraintIndex, const int32 LocalIndex, const T Dt, TVec3<T>& ParticleResidual, Chaos::PMatrix<T, 3, 3>& ParticleHessian) const
+		void AddResidualAndHessian(const FSolverParticles& InParticles, const int32 ConstraintIndex, const int32 LocalIndex, const FSolverReal Dt, TVec3<FSolverReal>& ParticleResidual, Chaos::PMatrix<FSolverReal, 3, 3>& ParticleHessian) const
 		{
 			const TVec4<int32>& Constraint = Constraints[ConstraintIndex];
 			const int32 Index1 = Constraint[0];
@@ -58,25 +60,25 @@ namespace Chaos::Softs
 			const int32 Index3 = Constraint[2];
 			const int32 Index4 = Constraint[3];
 			
-			const TVec3<T>& P1 = InParticles.P(Index1);
-			const TVec3<T>& P2 = InParticles.P(Index2);
-			const TVec3<T>& P3 = InParticles.P(Index3);
-			const TVec3<T>& P4 = InParticles.P(Index4);
+			const TVec3<FSolverReal>& P1 = InParticles.P(Index1);
+			const TVec3<FSolverReal>& P2 = InParticles.P(Index2);
+			const TVec3<FSolverReal>& P3 = InParticles.P(Index3);
+			const TVec3<FSolverReal>& P4 = InParticles.P(Index4);
 
-			TVec4<TVec3<T>> Grads;
-			const TVec3<T> P2P1 = P2 - P1;
-			const TVec3<T> P4P1 = P4 - P1;
-			const TVec3<T> P3P1 = P3 - P1;
-			Grads[1] = TVec3<T>::CrossProduct(P3P1, P4P1) / (T)6.;
-			Grads[2] = TVec3<T>::CrossProduct(P4P1, P2P1) / (T)6.;
-			Grads[3] = TVec3<T>::CrossProduct(P2P1, P3P1) / (T)6.;
+			TVec4<TVec3<FSolverReal>> Grads;
+			const TVec3<FSolverReal> P2P1 = P2 - P1;
+			const TVec3<FSolverReal> P4P1 = P4 - P1;
+			const TVec3<FSolverReal> P3P1 = P3 - P1;
+			Grads[1] = TVec3<FSolverReal>::CrossProduct(P3P1, P4P1) / (FSolverReal)6.;
+			Grads[2] = TVec3<FSolverReal>::CrossProduct(P4P1, P2P1) / (FSolverReal)6.;
+			Grads[3] = TVec3<FSolverReal>::CrossProduct(P2P1, P3P1) / (FSolverReal)6.;
 			Grads[0] = -(Grads[1] + Grads[2] + Grads[3]);
 
-			const T Volume = ComputeVolume(P1, P2, P3, P4);
-			const T C_Hessian = StiffnessArray[ConstraintIndex] * Dt * Dt;
-			const T C_Residual = (Volume - Volumes[ConstraintIndex]) * C_Hessian;
+			const FSolverReal Volume = ComputeVolume(P1, P2, P3, P4);
+			const FSolverReal C_Hessian = StiffnessArray[ConstraintIndex] * Dt * Dt;
+			const FSolverReal C_Residual = (Volume - Volumes[ConstraintIndex]) * C_Hessian;
 			ParticleResidual += C_Residual * Grads[LocalIndex];
-			ParticleHessian += C_Hessian * Chaos::PMatrix<T, 3, 3>::OuterProduct(Grads[LocalIndex], Grads[LocalIndex]);
+			ParticleHessian += C_Hessian * Chaos::PMatrix<FSolverReal, 3, 3>::OuterProduct(Grads[LocalIndex], Grads[LocalIndex]);
 		}
 
 		TArray<TArray<int32>> GetStaticConstraintArrays(TArray<TArray<int32>>& InIncidentElements, TArray<TArray<int32>>& InIncidentElementsLocal) const
@@ -99,11 +101,13 @@ namespace Chaos::Softs
 
 	private:
 		TArray<TVector<int32, 4>> Constraints;
-		TArray<T> Volumes;
+		TArray<FSolverReal> Volumes;
 		TArray<TArray<int32>> IncidentElements;
 		TArray<TArray<int32>> IncidentElementsLocal;
-		TArray<T> StiffnessArray;
+		TArray<FSolverReal> StiffnessArray;
 	};
 
+	template <typename T, typename ParticleType>
+	using FGaussSeidelUnilateralTetConstraints UE_DEPRECATED(5.8, "Deprecated. this class is to be deleted, use FGaussSeidelUnilateralVolumeConstraints instead") = FGaussSeidelUnilateralVolumeConstraints;
 
 }// End namespace Chaos::Softs
diff --git a/Engine/Source/Runtime/Experimental/Chaos/Public/Chaos/Deformable/GaussSeidelWeakConstraints.h b/Engine/Source/Runtime/Experimental/Chaos/Public/Chaos/Deformable/GaussSeidelWeakConstraints.h
--- a/Engine/Source/Runtime/Experimental/Chaos/Public/Chaos/Deformable/GaussSeidelWeakConstraints.h
+++ b/Engine/Source/Runtime/Experimental/Chaos/Public/Chaos/Deformable/GaussSeidelWeakConstraints.h
@@ -20,25 +20,23 @@ namespace Chaos::Softs
 {
 	using Chaos::TVec3;
 
-	template <typename T>
-	struct FGaussSeidelWeakConstraintSingleData
+	struct FGaussSeidelSpringConstraintSingleData
 	{
 		TArray<int32> SingleIndices  = {};
 		TArray<int32> SingleSecondIndices = {};
-		T SingleStiffness = (T)0.;
-		TArray<T> SingleWeights = {};
-		TArray<T> SingleSecondWeights = {};
+		FSolverReal SingleStiffness = (FSolverReal)0.;
+		TArray<FSolverReal> SingleWeights = {};
+		TArray<FSolverReal> SingleSecondWeights = {};
 		bool bIsAnisotropic = false;
-		TVec3<T> SingleNormal = TVec3<T>((T)0.);
+		TVec3<FSolverReal> SingleNormal = TVec3<FSolverReal>((FSolverReal)0.);
 		bool bIsZeroRestLength = false;
-		T RestLength = (T)0.;
+		FSolverReal RestLength = (FSolverReal)0.;
 	};
 
-	template<class T>
-	class TGaussSeidelWeakConstraintData : public TArrayCollection
+	class FGaussSeidelSpringConstraintData : public TArrayCollection
 	{
 	public:
-		TGaussSeidelWeakConstraintData()
+		FGaussSeidelSpringConstraintData()
 		{
 			AddArray(&MIndices);
 			AddArray(&MSecondIndices);
@@ -50,8 +48,19 @@ namespace Chaos::Softs
 			AddArray(&MIsZeroRestLength);
 			AddArray(&MRestLength);
 		}
-		TGaussSeidelWeakConstraintData(const TGaussSeidelWeakConstraintData<T>& Other) = delete;
-		TGaussSeidelWeakConstraintData(TGaussSeidelWeakConstraintData<T>&& Other)
+
+		//Explicit copy constructor
+		FGaussSeidelSpringConstraintData(const FGaussSeidelSpringConstraintData& Other)
+			: FGaussSeidelSpringConstraintData()
+		{
+			AddConstraints(Other.Size());
+			for (int32 ConstraintIdx = 0; ConstraintIdx < Other.Size(); ++ConstraintIdx)
+			{
+				SetSingleConstraint(Other.GetSingleConstraintData(ConstraintIdx), ConstraintIdx);
+			}
+		}
+
+		FGaussSeidelSpringConstraintData(FGaussSeidelSpringConstraintData&& Other)
 		    : TArrayCollection(), MIndices(MoveTemp(Other.MIndices))
 			, MSecondIndices(MoveTemp(Other.MSecondIndices))
 			, MWeights(MoveTemp(Other.MWeights))
@@ -59,7 +68,6 @@ namespace Chaos::Softs
 			, MStiffness(MoveTemp(Other.MStiffness))
 			, MIsAnisotropic(MoveTemp(Other.MIsAnisotropic))
 		{
-			AddParticles(Other.Size());
 			AddArray(&MIndices);
 			AddArray(&MSecondIndices);
 			AddArray(&MWeights);
@@ -72,7 +80,7 @@ namespace Chaos::Softs
 			Other.MSize = 0;
 		}
 
-		virtual ~TGaussSeidelWeakConstraintData()
+		virtual ~FGaussSeidelSpringConstraintData()
 		{}
 
 		void AddConstraints(const int32 Num)
@@ -85,7 +93,7 @@ namespace Chaos::Softs
 			RemoveAtSwapHelper(Idx);
 		}
 
-		void SetSingleConstraint(const FGaussSeidelWeakConstraintSingleData<T>& SingleData, const int32 ConstraintIndex)
+		void SetSingleConstraint(const FGaussSeidelSpringConstraintSingleData& SingleData, const int32 ConstraintIndex)
 		{
 			MIndices[ConstraintIndex] = SingleData.SingleIndices;
 			MSecondIndices[ConstraintIndex] = SingleData.SingleSecondIndices;
@@ -97,7 +105,7 @@ namespace Chaos::Softs
 			MIsZeroRestLength[ConstraintIndex] = SingleData.bIsZeroRestLength;
 		}
 
-		void AddSingleConstraint(const FGaussSeidelWeakConstraintSingleData<T>& SingleData)
+		void AddSingleConstraint(const FGaussSeidelSpringConstraintSingleData& SingleData)
 		{
 			AddConstraints(1);
 			SetSingleConstraint(SingleData, MSize - 1);
@@ -113,7 +121,7 @@ namespace Chaos::Softs
 			ResizeHelper(Num);
 		}
 
-		TGaussSeidelWeakConstraintData& operator=(TGaussSeidelWeakConstraintData<T>&& Other)
+		FGaussSeidelSpringConstraintData& operator=(FGaussSeidelSpringConstraintData&& Other)
 		{
 			MIndices = MoveTemp(Other.MIndices);
 			MSecondIndices = MoveTemp(Other.MSecondIndices);
@@ -158,32 +166,32 @@ namespace Chaos::Softs
 			MSecondIndices[Index] = InIndices;
 		}
 
-		inline const TArrayCollectionArray<TArray<T>>& Weights() const
+		inline const TArrayCollectionArray<TArray<FSolverReal>>& Weights() const
 		{
 			return MWeights;
 		}
 
-		const TArray<T>& GetWeights(const int32 Index) const
+		const TArray<FSolverReal>& GetWeights(const int32 Index) const
 		{
 			return MWeights[Index];
 		}
 
-		void SetWeights(const int32 Index, const TArray<int32>& InWeights)
+		void SetWeights(const int32 Index, const TArray<FSolverReal>& InWeights)
 		{
 			MWeights[Index] = InWeights;
 		}
 
-		inline const TArrayCollectionArray<TArray<T>>& SecondWeights() const
+		inline const TArrayCollectionArray<TArray<FSolverReal>>& SecondWeights() const
 		{
 			return MSecondWeights;
 		}
 
-		const TArray<T>& GetSecondWeights(const int32 Index) const
+		const TArray<FSolverReal>& GetSecondWeights(const int32 Index) const
 		{
 			return MSecondWeights[Index];
 		}
 
-		void SetSecondWeights(const int32 Index, const TArray<int32>& InWeights)
+		void SetSecondWeights(const int32 Index, const TArray<FSolverReal>& InWeights)
 		{
 			MSecondWeights[Index] = InWeights;
 		}
@@ -198,32 +206,32 @@ namespace Chaos::Softs
 			MIsAnisotropic[Index] = InIsAnisotropic;
 		}
 
-		inline const TArrayCollectionArray<TVec3<T>>& Normals() const
+		inline const TArrayCollectionArray<TVec3<FSolverReal>>& Normals() const
 		{
 			return MNormals;
 		}
 
-		const TVec3<T>& GetNormal(const int32 Index) const
+		const TVec3<FSolverReal>& GetNormal(const int32 Index) const
 		{
 			return MNormals[Index];
 		}
 
-		void SetNormal(const int32 Index, const TVec3<T>& InNormal)
+		void SetNormal(const int32 Index, const TVec3<FSolverReal>& InNormal)
 		{
 			MNormals[Index] = InNormal;
 		}
 
-		inline const TArrayCollectionArray<T>& Stiffness() const
+		inline const TArrayCollectionArray<FSolverReal>& Stiffness() const
 		{
 			return MStiffness;
 		}
 
-		T GetStiffness(const int32 Index) const
+		FSolverReal GetStiffness(const int32 Index) const
 		{
 			return MStiffness[Index];
 		}
 
-		void SetStiffness(const int32 Index, const T InStiffness)
+		void SetStiffness(const int32 Index, const FSolverReal InStiffness)
 		{
 			MStiffness[Index] = InStiffness;
 		}
@@ -238,14 +246,14 @@ namespace Chaos::Softs
 			MIsZeroRestLength[Index] = InIsZeroRestLength;
 		}
 
-		void SetRestLength(const int32 Index, const T InRestLength)
+		void SetRestLength(const int32 Index, const FSolverReal InRestLength)
 		{
 			MRestLength[Index] = InRestLength;
 		}
 
-		const FGaussSeidelWeakConstraintSingleData<T> GetSingleConstraintData(const int32 ConstraintIndex) const 
+		const FGaussSeidelSpringConstraintSingleData GetSingleConstraintData(const int32 ConstraintIndex) const 
 		{
-			FGaussSeidelWeakConstraintSingleData<T> SingleConstraintData;
+			FGaussSeidelSpringConstraintSingleData SingleConstraintData;
 			check(static_cast<uint32>(ConstraintIndex) < MSize);
 			if (ConstraintIndex > INDEX_NONE && static_cast<uint32>(ConstraintIndex) < MSize)
 			{
@@ -265,28 +273,26 @@ namespace Chaos::Softs
 	private:
 		TArrayCollectionArray<TArray<int32>> MIndices;
 		TArrayCollectionArray<TArray<int32>> MSecondIndices;
-		TArrayCollectionArray<TArray<T>> MWeights;
-		TArrayCollectionArray<TArray<T>> MSecondWeights;
-		TArrayCollectionArray<T> MStiffness;
+		TArrayCollectionArray<TArray<FSolverReal>> MWeights;
+		TArrayCollectionArray<TArray<FSolverReal>> MSecondWeights;
+		TArrayCollectionArray<FSolverReal> MStiffness;
 		TArrayCollectionArray<bool> MIsAnisotropic;
-		TArrayCollectionArray<TVector<T, 3>> MNormals;
+		TArrayCollectionArray<TVector<FSolverReal, 3>> MNormals;
 		TArrayCollectionArray<bool> MIsZeroRestLength;
-		TArrayCollectionArray<T> MRestLength;
+		TArrayCollectionArray<FSolverReal> MRestLength;
 	};
 
-
-
-	template <typename T, typename ParticleType>
-	struct FGaussSeidelWeakConstraints 
+	// Spring constraint for position-based nonlinear Gauss Seidel solver [Chen et al., 2024]
+	struct FGaussSeidelSpringConstraints 
 	{
 		//TODO(Yizhou): Add unittest for Gauss Seidel Weak Constraints
-		FGaussSeidelWeakConstraints(
+		FGaussSeidelSpringConstraints(
 			const TArray<TArray<int32>>& InIndices,
-			const TArray<TArray<T>>& InWeights,
-			const TArray<T>& InStiffness,
+			const TArray<TArray<FSolverReal>>& InWeights,
+			const TArray<FSolverReal>& InStiffness,
 			const TArray<TArray<int32>>& InSecondIndices,
-			const TArray<TArray<T>>& InSecondWeights,
-			const FDeformableXPBDWeakConstraintParams& InParams
+			const TArray<TArray<FSolverReal>>& InSecondWeights,
+			const FDeformableXPBDWeightedSpringConstraintParams& InParams
 		): DebugDrawParams(InParams)
 		{
 			ensureMsgf(InIndices.Num() == InSecondIndices.Num(), TEXT("Input Double Bindings have wrong size"));
@@ -297,7 +303,7 @@ namespace Chaos::Softs
 
 			for (int32 i = 0; i < InIndices.Num(); i++)
 			{
-				FGaussSeidelWeakConstraintSingleData<T> SingleConstraintData;
+				FGaussSeidelSpringConstraintSingleData SingleConstraintData;
 				SingleConstraintData.SingleIndices = InIndices[i];
 				SingleConstraintData.SingleSecondIndices = InSecondIndices[i];
 				SingleConstraintData.SingleWeights = InWeights[i];
@@ -322,9 +328,9 @@ namespace Chaos::Softs
 			int32 ConstraintIndex;
 		};
 
-		virtual ~FGaussSeidelWeakConstraints() {}
+		virtual ~FGaussSeidelSpringConstraints() {}
 
-		void ComputeInitialWCData(const ParticleType& InParticles)
+		void ComputeInitialWCData(const FSolverParticles& InParticles)
 		{
 			TArray<TArray<int32>> ExtraConstraints;
 			ExtraConstraints.Init(TArray<int32>(), ConstraintsData.Size());
@@ -363,13 +369,13 @@ namespace Chaos::Softs
 				{
 					const TArray<int32>& Indices = ConstraintsData.GetIndices(ConstraintIdx);
 					const TArray<int32>& SecondIndices = ConstraintsData.GetSecondIndices(ConstraintIdx);
-					const TArray<T>& Weights = ConstraintsData.GetWeights(ConstraintIdx);
-					const TArray<T>& SecondWeights = ConstraintsData.GetSecondWeights(ConstraintIdx);
-					const TVec3<T> RestSpringEdge = ComputeSpringEdge(InParticles, Indices, SecondIndices, Weights, SecondWeights, /*bUseParticleX =*/true);
+					const TArray<FSolverReal>& Weights = ConstraintsData.GetWeights(ConstraintIdx);
+					const TArray<FSolverReal>& SecondWeights = ConstraintsData.GetSecondWeights(ConstraintIdx);
+					const TVec3<FSolverReal> RestSpringEdge = ComputeSpringEdge(InParticles, Indices, SecondIndices, Weights, SecondWeights, /*bUseParticleX =*/true);
 					if (ConstraintsData.GetIsAnisotropic(ConstraintIdx))
 					{
 						//if the spring is anisotropic, rest length could be negative depending on the normal direction
-						ConstraintsData.SetRestLength(ConstraintIdx, TVec3<T>::DotProduct(ConstraintsData.GetNormal(ConstraintIdx), RestSpringEdge));
+						ConstraintsData.SetRestLength(ConstraintIdx, TVec3<FSolverReal>::DotProduct(ConstraintsData.GetNormal(ConstraintIdx), RestSpringEdge));
 					}
 					else
 					{
@@ -379,10 +385,10 @@ namespace Chaos::Softs
 			}
 		}
 
-		TVec3<T> ComputeSpringEdge(const ParticleType& InParticles, const TArray<int32>& LocalIndices,
-			const TArray<int32>& LocalSecondIndices, const TArray<T>& Weight, const TArray<T>& SecondWeight, bool bUseParticleX) const
+		TVec3<FSolverReal> ComputeSpringEdge(const FSolverParticles& InParticles, const TArray<int32>& LocalIndices,
+			const TArray<int32>& LocalSecondIndices, const TArray<FSolverReal>& Weight, const TArray<FSolverReal>& SecondWeight, bool bUseParticleX) const
 		{
-			TVec3<T> SpringEdge((T)0.);
+			TVec3<FSolverReal> SpringEdge((FSolverReal)0.);
 			if (ensure(LocalIndices.Num() == Weight.Num() && LocalSecondIndices.Num() == SecondWeight.Num()))
 			{
 				if (bUseParticleX)
@@ -423,7 +429,7 @@ namespace Chaos::Softs
 			return SpringEdge;
 		}
 
-		void AddWCHessian(const int32 p, const T Dt, Chaos::PMatrix<T, 3, 3>& ParticleHessian) const
+		void AddWCHessian(const int32 p, const FSolverReal Dt, Chaos::PMatrix<FSolverReal, 3, 3>& ParticleHessian) const
 		{
 			if (NodalWeights[p].Num() > 0)
 			{
@@ -448,10 +454,10 @@ namespace Chaos::Softs
 		}
 
 		void AddExtraConstraints(const TArray<TArray<int32>>& InIndices,
-								const TArray<TArray<T>>& InWeights,
-								const TArray<T>& InStiffness,
+								const TArray<TArray<FSolverReal>>& InWeights,
+								const TArray<FSolverReal>& InStiffness,
 								const TArray<TArray<int32>>& InSecondIndices,
-								const TArray<TArray<T>>& InSecondWeights,
+								const TArray<TArray<FSolverReal>>& InSecondWeights,
 								const TArray<bool>& InIsAnisotrpic,
 								const TArray<bool>& InIsZeroRestLength)
 		{
@@ -461,7 +467,7 @@ namespace Chaos::Softs
 
 			for (int32 i = 0; i < InIndices.Num(); i++)
 			{
-				FGaussSeidelWeakConstraintSingleData<T> SingleConstraintData;
+				FGaussSeidelSpringConstraintSingleData SingleConstraintData;
 				SingleConstraintData.SingleIndices = InIndices[i];
 				SingleConstraintData.SingleSecondIndices = InSecondIndices[i];
 				SingleConstraintData.SingleWeights = InWeights[i];
@@ -480,7 +486,7 @@ namespace Chaos::Softs
 
 		void UpdatePointTriangleCollisionWCData(const FSolverParticles& Particles)
 		{	
-			TGaussSeidelWeakConstraintData<T> OriginalConstraintsData = ConstraintsData;
+			FGaussSeidelSpringConstraintData OriginalConstraintsData = ConstraintsData;
 
 			ConstraintsData.Resize(InitialWCSize);
 
@@ -491,22 +497,22 @@ namespace Chaos::Softs
 				ensureMsgf(OriginalConstraintsData.GetIndices(i).Num() == 3, TEXT("Collision format is not point-triangle"));
 				ensureMsgf(OriginalConstraintsData.GetSecondIndices(i).Num() == 1, TEXT("Collision format is not point-triangle"));
 				Chaos::TVector<float, 3> TriPos0(Particles.P(IndicesTemp[0])), TriPos1(Particles.P(IndicesTemp[1])), TriPos2(Particles.P(IndicesTemp[2])), ParticlePos(Particles.P(SecondIndicesTemp[0]));
-				Chaos::TVector<T, 3> Normal = FVector3f::CrossProduct(TriPos2 - TriPos0, TriPos1 - TriPos0); //triangle normal convention (see FTriangleMesh::GetFaceNormals())
+				Chaos::TVector<FSolverReal, 3> Normal = FVector3f::CrossProduct(TriPos2 - TriPos0, TriPos1 - TriPos0); //triangle normal convention (see FTriangleMesh::GetFaceNormals())
 				if (FVector3f::DotProduct(ParticlePos - TriPos0, Normal) < 0.f) //not resolved, keep the spring
 				{
-					ConstraintsData.AddConstraints(OriginalConstraintsData.GetSingleConstraintData(i));
+					ConstraintsData.AddSingleConstraint(OriginalConstraintsData.GetSingleConstraintData(i));
 				}
 			}
 		}
 
-		void VisualizeAllBindings(const FSolverParticles& InParticles, const T Dt) const
+		void VisualizeAllBindings(const FSolverParticles& InParticles, const FSolverReal Dt) const
 		{
 #if WITH_EDITOR
-			auto DoubleVert = [](Chaos::TVec3<T> V) { return FVector3d(V.X, V.Y, V.Z); };
+			auto DoubleVert = [](Chaos::TVec3<FSolverReal> V) { return FVector3d(V.X, V.Y, V.Z); };
 			for (int32 i = 0; i < ConstraintsData.Size(); i++)
 			{
-				const FGaussSeidelWeakConstraintSingleData<T>& SingleConstraintData = ConstraintsData.GetSingleConstraintData(i);
-				Chaos::TVec3<T> SourcePos((T)0.), TargetPos((T)0.);
+				const FGaussSeidelSpringConstraintSingleData& SingleConstraintData = ConstraintsData.GetSingleConstraintData(i);
+				Chaos::TVec3<FSolverReal> SourcePos((FSolverReal)0.), TargetPos((FSolverReal)0.);
 				for (int32 j = 0; j < SingleConstraintData.SingleIndices.Num(); j++)
 				{
 					SourcePos += SingleConstraintData.SingleWeights[j] * InParticles.P(SingleConstraintData.SingleIndices[j]);
@@ -545,7 +551,7 @@ namespace Chaos::Softs
 #endif
 		}
 
-		void Init(const FSolverParticles& InParticles, const T Dt)
+		void Init(const FSolverParticles& InParticles, const FSolverReal Dt)
 		{
 			UpdateTriangleNormalAndNodalWeight(InParticles, /*bUseParticleX =*/ false);
 			if (DebugDrawParams.bVisualizeBindings)
@@ -564,7 +570,7 @@ namespace Chaos::Softs
 					const TArray<int32>& SecondIndicesTemp = ConstraintsData.GetSecondIndices(i);
 					ensureMsgf(ConstraintsData.GetIndices(i).Num() == 3, TEXT("Collision format is not point-triangle"));
 					ensureMsgf(ConstraintsData.GetSecondIndices(i).Num() == 1, TEXT("Collision format is not point-triangle"));
-					Chaos::TVector<T, 3> Normal;
+					Chaos::TVector<FSolverReal, 3> Normal;
 					if (bUseParticleX)
 					{
 						Chaos::TVector<float, 3> TriPos0(InParticles.X(IndicesTemp[0])), TriPos1(InParticles.X(IndicesTemp[1])), TriPos2(InParticles.X(IndicesTemp[2]));
@@ -584,13 +590,13 @@ namespace Chaos::Softs
 			{
 				if (WCIncidentElements[p].Num() > 0)
 				{
-					NodalWeights[p].Init(T(0), 6);
+					NodalWeights[p].Init(FSolverReal(0), 6);
 					for (int32 j = 0; j < WCIncidentElements[p].Num(); j++)
 					{
 						int32 ConstraintIndex = WCIncidentElements[p][j];
 						int32 LocalIndex = WCIncidentElementsLocal[p][j];
 
-						T Weight = T(0);
+						FSolverReal Weight = FSolverReal(0);
 						if (LocalIndex >= ConstraintsData.GetIndices(ConstraintIndex).Num())
 						{
 							Weight = ConstraintsData.GetSecondWeights(ConstraintIndex)[LocalIndex - ConstraintsData.GetIndices(ConstraintIndex).Num()];
@@ -722,7 +728,7 @@ namespace Chaos::Softs
 						Chaos::FRealSingle CurrentDistance = (Particles.P(SurfaceElementsMap[TriangleIntersections[j]]) - ClosestPoint).Size();
 						if (CurrentDistance < MinDis)
 						{
-							Chaos::TVector<T, 3> Normal = FVector3f::CrossProduct(TriPos2 - TriPos0, TriPos1 - TriPos0); //The normal needs to point outwards of the geometry
+							Chaos::TVector<FSolverReal, 3> Normal = FVector3f::CrossProduct(TriPos2 - TriPos0, TriPos1 - TriPos0); //The normal needs to point outwards of the geometry
 							if (FVector3f::DotProduct(ParticlePos - TriPos0, Normal) < 0.f)
 							{
 								Normal.SafeNormalize(1e-8f);
@@ -740,11 +746,11 @@ namespace Chaos::Softs
 					&& MinIndex != SurfaceElements[i][1]
 					&& MinIndex != SurfaceElements[i][2])
 				{
-					FGaussSeidelWeakConstraintSingleData<T> SingleConstraintData;
+					FGaussSeidelSpringConstraintSingleData SingleConstraintData;
 					SingleConstraintData.SingleIndices = {SurfaceElements[i][0], SurfaceElements[i][1], SurfaceElements[i][2]};
 					SingleConstraintData.SingleSecondIndices = {MinIndex};
-					SingleConstraintData.Weights = {ClosestBary[0], ClosestBary[1], ClosestBary[2]};
-					SingleConstraintData.SecondWeights = {T(1.f)};
+					SingleConstraintData.SingleWeights = {ClosestBary[0], ClosestBary[1], ClosestBary[2]};
+					SingleConstraintData.SingleSecondWeights = {FSolverReal(1.f)};
 					SingleConstraintData.bIsAnisotropic = UseAnisotropicSpring;
 					SingleConstraintData.SingleNormal = FaceNormal;
 					SingleConstraintData.bIsZeroRestLength = true; //Push-out type collision springs should be zero rest length
@@ -755,7 +761,7 @@ namespace Chaos::Softs
 						SpringStiffness += ClosestBary[k] * PositionTargetStiffness * Particles.M(SurfaceElements[i][k]);
 					}
 					SpringStiffness += PositionTargetStiffness * Particles.M(MinIndex);
-					SingleConstraintData.SingleStiffness = (T)SpringStiffness;
+					SingleConstraintData.SingleStiffness = (FSolverReal)SpringStiffness;
 					ConstraintsData.AddSingleConstraint(SingleConstraintData);
 				}
 			}
@@ -790,11 +796,11 @@ namespace Chaos::Softs
 								const TVector<int32, 3>& Elem = Elements[CollisionPoint.Indices[1]];
 								const int32 IndexToWrite = ConstraintIndex.fetch_add(1);
 
-								FGaussSeidelWeakConstraintSingleData<T> SingleConstraintData;
+								FGaussSeidelSpringConstraintSingleData SingleConstraintData;
 								SingleConstraintData.SingleIndices = { Elem[0], Elem[1] ,Elem[2] };
 								SingleConstraintData.SingleSecondIndices =  { Index };
 								SingleConstraintData.SingleWeights = { CollisionPoint.Bary[1], CollisionPoint.Bary[2], CollisionPoint.Bary[3] };
-								SingleConstraintData.SingleSecondWeights = {T(1.f)};
+								SingleConstraintData.SingleSecondWeights = {FSolverReal(1.f)};
 								SingleConstraintData.bIsAnisotropic = UseAnisotropicSpring;
 								SingleConstraintData.SingleNormal = CollisionPoint.Normal;
 								SingleConstraintData.bIsZeroRestLength = true; //Push-out type collision springs should be zero rest length
@@ -805,7 +811,7 @@ namespace Chaos::Softs
 									SpringStiffness += SingleConstraintData.SingleWeights[k] * PositionTargetStiffness * Particles.M(Elem[k]);
 								}
 								SpringStiffness += PositionTargetStiffness * Particles.M(Index);
-								SingleConstraintData.SingleStiffness = (T)SpringStiffness;
+								SingleConstraintData.SingleStiffness = (FSolverReal)SpringStiffness;
 								ConstraintsData.SetSingleConstraint(SingleConstraintData, IndexToWrite);
 							}
 						}
@@ -846,11 +852,11 @@ namespace Chaos::Softs
 								const TVector<int32, 3>& Elem = Elements[CollisionPoint.Indices[1]];
 								const int32 IndexToWrite = ConstraintIndex.fetch_add(1);
 
-								FGaussSeidelWeakConstraintSingleData<T> SingleConstraintData;
+								FGaussSeidelSpringConstraintSingleData SingleConstraintData;
 								SingleConstraintData.SingleIndices = { Elem[0], Elem[1] ,Elem[2] };
 								SingleConstraintData.SingleSecondIndices =  { Index };
 								SingleConstraintData.SingleWeights = { CollisionPoint.Bary[1], CollisionPoint.Bary[2], CollisionPoint.Bary[3] };
-								SingleConstraintData.SingleSecondWeights = {T(1.f)};
+								SingleConstraintData.SingleSecondWeights = {FSolverReal(1.f)};
 								SingleConstraintData.bIsAnisotropic = UseAnisotropicSpring;
 								SingleConstraintData.SingleNormal = CollisionPoint.Normal;
 								SingleConstraintData.bIsZeroRestLength = true; //Push-out type collision springs should be zero rest length
@@ -861,7 +867,7 @@ namespace Chaos::Softs
 									SpringStiffness += SingleConstraintData.SingleWeights[k] * PositionTargetStiffness * Particles.M(Elem[k]);
 								}
 								SpringStiffness += PositionTargetStiffness * Particles.M(Index);
-								SingleConstraintData.SingleStiffness = (T)SpringStiffness;
+								SingleConstraintData.SingleStiffness = (FSolverReal)SpringStiffness;
 								ConstraintsData.SetSingleConstraint(SingleConstraintData, IndexToWrite);
 							}
 						}
@@ -905,16 +911,16 @@ namespace Chaos::Softs
 					int32 p = ExtraConstraints[ExtraWCIncidentElements[i][0]][ExtraWCIncidentElementsLocal[i][0]];
 					if (NodalWeights[p].Num() == 0)
 					{
-						NodalWeights[p].Init(T(0), 6);
+						NodalWeights[p].Init(FSolverReal(0), 6);
 					}
 					for (int32 j = 0; j < ExtraWCIncidentElements[i].Num(); j++)
 					{
 						int32 LocalIndex = ExtraWCIncidentElementsLocal[i][j];
 						int32 ConstraintIndex = ExtraWCIncidentElements[i][j] + InitialWCSize;
 
-						const FGaussSeidelWeakConstraintSingleData<T>& SingleData = ConstraintsData.GetSingleConstraintData(ConstraintIndex);
+						const FGaussSeidelSpringConstraintSingleData& SingleData = ConstraintsData.GetSingleConstraintData(ConstraintIndex);
 
-						T weight = T(0);
+						FSolverReal weight = FSolverReal(0);
 						if (LocalIndex >= SingleData.SingleIndices.Num())
 						{
 							weight = SingleData.SingleWeights[LocalIndex - SingleData.SingleIndices.Num()];
@@ -979,13 +985,13 @@ namespace Chaos::Softs
 		}
 
 		// Deprecated, now replaced with AddWCResidual for more general cases
-		void AddZeroRestLengthWCResidualAndHessian(const ParticleType& InParticles, const int32 ConstraintIndex, const int32 LocalIndex, const T Dt, TVec3<T>& ParticleResidual, Chaos::PMatrix<T, 3, 3>& ParticleHessian) const
+		void AddZeroRestLengthWCResidualAndHessian(const FSolverParticles& InParticles, const int32 ConstraintIndex, const int32 LocalIndex, const FSolverReal Dt, TVec3<FSolverReal>& ParticleResidual, Chaos::PMatrix<FSolverReal, 3, 3>& ParticleHessian) const
 		{
-			const FGaussSeidelWeakConstraintSingleData<T>& SingleData = ConstraintsData.GetSingleConstraintData(ConstraintIndex);
+			const FGaussSeidelSpringConstraintSingleData& SingleData = ConstraintsData.GetSingleConstraintData(ConstraintIndex);
 
-			const TVec3<T> SpringEdge = ComputeSpringEdge(InParticles, SingleData.SingleIndices, SingleData.SingleSecondIndices,
+			const TVec3<FSolverReal> SpringEdge = ComputeSpringEdge(InParticles, SingleData.SingleIndices, SingleData.SingleSecondIndices,
 				SingleData.SingleWeights, SingleData.SingleSecondWeights, /*bUseParticleX =*/false);
-			T weight = T(0);
+			FSolverReal weight = FSolverReal(0);
 			if (LocalIndex >= SingleData.SingleIndices.Num())
 			{
 				weight = -SingleData.SingleSecondWeights[LocalIndex - SingleData.SingleIndices.Num()];
@@ -996,8 +1002,8 @@ namespace Chaos::Softs
 			}
 			if (SingleData.bIsAnisotropic)
 			{
-				T comp = TVec3<T>::DotProduct(SpringEdge, SingleData.SingleNormal);
-				TVec3<T> proj = SingleData.SingleNormal * comp;
+				FSolverReal comp = TVec3<FSolverReal>::DotProduct(SpringEdge, SingleData.SingleNormal);
+				TVec3<FSolverReal> proj = SingleData.SingleNormal * comp;
 				for (int32 alpha = 0; alpha < 3; alpha++)
 				{
 					ParticleResidual[alpha] += Dt * Dt * SingleData.SingleStiffness * proj[alpha] * weight;
@@ -1012,13 +1018,13 @@ namespace Chaos::Softs
 			}
 		}
 
-		void AddWCResidual(const ParticleType& InParticles, const int32 ConstraintIndex, const int32 LocalIndex, const T Dt, TVec3<T>& ParticleResidual, Chaos::PMatrix<T, 3, 3>& ParticleHessian) const
+		void AddWCResidual(const FSolverParticles& InParticles, const int32 ConstraintIndex, const int32 LocalIndex, const FSolverReal Dt, TVec3<FSolverReal>& ParticleResidual, Chaos::PMatrix<FSolverReal, 3, 3>& ParticleHessian) const
 		{
-			const FGaussSeidelWeakConstraintSingleData<T>& SingleData = ConstraintsData.GetSingleConstraintData(ConstraintIndex);
+			const FGaussSeidelSpringConstraintSingleData& SingleData = ConstraintsData.GetSingleConstraintData(ConstraintIndex);
 
-			const TVec3<T> SpringEdge = ComputeSpringEdge(InParticles, SingleData.SingleIndices, SingleData.SingleSecondIndices,
+			const TVec3<FSolverReal> SpringEdge = ComputeSpringEdge(InParticles, SingleData.SingleIndices, SingleData.SingleSecondIndices,
 				SingleData.SingleWeights, SingleData.SingleSecondWeights, /*bUseParticleX =*/false);
-			T weight = T(0);
+			FSolverReal weight = FSolverReal(0);
 			if (LocalIndex >= SingleData.SingleIndices.Num())
 			{
 				weight = -SingleData.SingleSecondWeights[LocalIndex - SingleData.SingleIndices.Num()];
@@ -1027,10 +1033,10 @@ namespace Chaos::Softs
 			{
 				weight = SingleData.SingleWeights[LocalIndex];
 			}
-			TVec3<T> Projection;
+			TVec3<FSolverReal> Projection;
 			if (SingleData.bIsAnisotropic)
 			{
-				T LengthDiff = TVec3<T>::DotProduct(SpringEdge, SingleData.SingleNormal) - SingleData.RestLength;
+				FSolverReal LengthDiff = TVec3<FSolverReal>::DotProduct(SpringEdge, SingleData.SingleNormal) - SingleData.RestLength;
 				Projection = SingleData.SingleNormal * LengthDiff;
 			}
 			else
@@ -1047,21 +1053,29 @@ namespace Chaos::Softs
 			}
 		}
 
-		TGaussSeidelWeakConstraintData<T> ConstraintsData;
+		FGaussSeidelSpringConstraintData ConstraintsData;
 
-		TArray<TArray<T>> NodalWeights;
+		TArray<TArray<FSolverReal>> NodalWeights;
 
 		TArray<TArray<int32>> WCIncidentElements;
 		TArray<TArray<int32>> WCIncidentElementsLocal;
 
-		FDeformableXPBDWeakConstraintParams DebugDrawParams;
+		FDeformableXPBDWeightedSpringConstraintParams DebugDrawParams;
 
 		int32 InitialWCSize;
-		TArray<TArray<T>> NoCollisionNodalWeights;
+		TArray<TArray<FSolverReal>> NoCollisionNodalWeights;
 		TArray<TArray<int32>> NoCollisionConstraints;
 		TArray<TArray<int32>> NoCollisionWCIncidentElements;
 		TArray<TArray<int32>> NoCollisionWCIncidentElementsLocal;
 	};
 
+	template <typename T>
+	using FGaussSeidelWeakConstraintSingleData UE_DEPRECATED(5.8, "Deprecated. this struct is to be deleted, use FGaussSeidelSpringConstraintSingleData instead") = FGaussSeidelSpringConstraintSingleData;
+
+	template <class T>
+	using FGaussSeidelWeakConstraintData UE_DEPRECATED(5.8, "Deprecated. this class is to be deleted, use FGaussSeidelSpringConstraintData instead") = FGaussSeidelSpringConstraintData;
+	
+	template <typename T, typename ParticleType>
+	using FGaussSeidelWeakConstraints UE_DEPRECATED(5.8, "Deprecated. this struct is to be deleted, use FGaussSeidelSpringConstraints instead") = FGaussSeidelSpringConstraints;
 
 }// End namespace Chaos::Softs
diff --git a/Engine/Source/Runtime/Experimental/Chaos/Public/Chaos/Deformable/MuscleActivationConstraints.h b/Engine/Source/Runtime/Experimental/Chaos/Public/Chaos/Deformable/MuscleActivationConstraints.h
--- a/Engine/Source/Runtime/Experimental/Chaos/Public/Chaos/Deformable/MuscleActivationConstraints.h
+++ b/Engine/Source/Runtime/Experimental/Chaos/Public/Chaos/Deformable/MuscleActivationConstraints.h
@@ -13,16 +13,16 @@ namespace Chaos::Softs
 	using Chaos::TVec3;
 	typedef GeometryCollection::Facades::FMuscleActivationFacade Facade;
 	typedef GeometryCollection::Facades::FMuscleActivationData Data;
-	template <typename T, typename ParticleType>
-	class FMuscleActivationConstraints
+
+	// Stores data for muscle activation and apply rest volume changes to tetrahedral constraints
+	class FMuscleActivationTetrahedralConstraints
 	{
 	public:
-		//Handles muscle activation data
-		FMuscleActivationConstraints(){}
+		FMuscleActivationTetrahedralConstraints(){}
 
-		virtual ~FMuscleActivationConstraints() {}
+		virtual ~FMuscleActivationTetrahedralConstraints() {}
 
-		void AddMuscles(const ParticleType& RestParticles, const Facade& FMuscleActivation, int32 VertexOffset = 0, int32 ElementOffset = 0)
+		void AddMuscles(const FSolverParticles& RestParticles, const Facade& FMuscleActivation, int32 VertexOffset = 0, int32 ElementOffset = 0)
 		{
 			for (int MuscleIdx = 0; MuscleIdx < FMuscleActivation.NumMuscles(); MuscleIdx++)
 			{
@@ -69,7 +69,7 @@ namespace Chaos::Softs
 						InflationVolumeScale.Add(MuscleActivationData.InflationVolumeScale);
 						OriginInsertionPair.Add(MuscleActivationData.OriginInsertionPair);
 						MuscleRestLength.Add((RestParticles.GetX(MuscleActivationData.OriginInsertionPair[0]) - RestParticles.GetX(MuscleActivationData.OriginInsertionPair[1])).Size());
-						MuscleActivation.Add(0.f);
+						MuscleActivation.Add(FSolverReal(0));
 						MuscleVertexOffset.Add(FMuscleActivation.MuscleVertexOffset(MuscleIdx));
 						MuscleVertexCount.Add(FMuscleActivation.NumMuscleVertices(MuscleIdx));
 						LengthActivationCurves.Add(FMuscleActivation.GetLengthActivationCurve(MuscleIdx));
@@ -83,7 +83,7 @@ namespace Chaos::Softs
 			}
 		}
 
-		void UpdateLengthBasedMuscleActivation(const ParticleType& InParticles)
+		void UpdateLengthBasedMuscleActivation(const FSolverParticles& InParticles)
 		{
 			for (int32 MuscleIdx = 0; MuscleIdx < MuscleActivationElement.Num(); MuscleIdx++)
 			{
@@ -92,15 +92,15 @@ namespace Chaos::Softs
 					TEXT("MuscleLengthRatioThresholdForMaxActivation %f of muscle indexed %d is out of range (0,1), please check your setup."), MuscleLengthRatioThresholdForMaxActivation[MuscleIdx], MuscleIdx))
 				{
 					// calculate origin/insertion length
-					const float MuscleLengthRatio = (InParticles.P(OriginInsertionPair[MuscleIdx][0]) - InParticles.P(OriginInsertionPair[MuscleIdx][1])).Size() / MuscleRestLength[MuscleIdx];
-					if (MuscleLengthRatio >= 1.f)
+					const FSolverReal MuscleLengthRatio = (InParticles.P(OriginInsertionPair[MuscleIdx][0]) - InParticles.P(OriginInsertionPair[MuscleIdx][1])).Size() / MuscleRestLength[MuscleIdx];
+					if (MuscleLengthRatio >= FSolverReal(1))
 					{
 						// not active
-						MuscleActivation[MuscleIdx] = 0.f;
+						MuscleActivation[MuscleIdx] = FSolverReal(0);
 					}
 					else
 					{
-						float NormalizedMuscleLengthLevel = FMath::Clamp((1.f - MuscleLengthRatio) / (1 - MuscleLengthRatioThresholdForMaxActivation[MuscleIdx]), 0.f, 1.f);
+						FSolverReal NormalizedMuscleLengthLevel = FMath::Clamp((FSolverReal(1) - MuscleLengthRatio) / (FSolverReal(1) - MuscleLengthRatioThresholdForMaxActivation[MuscleIdx]), FSolverReal(0), FSolverReal(1));
 						if (LengthActivationCurves[MuscleIdx].GetNumKeys())
 						{
 							MuscleActivation[MuscleIdx] = LengthActivationCurves[MuscleIdx].Eval(NormalizedMuscleLengthLevel);
@@ -115,7 +115,7 @@ namespace Chaos::Softs
 			}
 		}
 
-		void ApplyMuscleActivation(FXPBDCorotatedConstraints<T,ParticleType>& Constraints) const
+		void ApplyMuscleActivation(FXPBDCorotatedTetrahedralConstraints& Constraints) const
 		{
 			for (int32 MuscleIdx = 0; MuscleIdx < MuscleActivationElement.Num(); MuscleIdx++)
 			{
@@ -125,7 +125,7 @@ namespace Chaos::Softs
 						FiberLengthRatioAtMaxActivation[MuscleIdx] < 1,
 						TEXT("FiberLengthRatioAtMaxActivation %f of muscle indexed %d is out of range (0,1), please check your setup."), FiberLengthRatioAtMaxActivation[MuscleIdx], MuscleIdx))
 				{
-					const float FiberLengthRatio = 1.f - MuscleActivation[MuscleIdx] * (1.f - FiberLengthRatioAtMaxActivation[MuscleIdx]);
+					const FSolverReal FiberLengthRatio = FSolverReal(1) - MuscleActivation[MuscleIdx] * (FSolverReal(1) - FiberLengthRatioAtMaxActivation[MuscleIdx]);
 					for (int32 ElemIdx = 0; ElemIdx < MuscleActivationElement[MuscleIdx].Num(); ElemIdx++)
 					{
 						Constraints.ModifyDmInverseFromMuscleLength(MuscleActivationElement[MuscleIdx][ElemIdx], FiberLengthRatio, FiberDirectionMatrix[MuscleIdx][ElemIdx], ContractionVolumeScale[MuscleIdx][ElemIdx]);
@@ -134,7 +134,7 @@ namespace Chaos::Softs
 			}
 		}
 
-		void ApplyInflationVolumeScale(FXPBDCorotatedConstraints<T, ParticleType>& Constraints) const
+		void ApplyInflationVolumeScale(FXPBDCorotatedTetrahedralConstraints& Constraints) const
 		{
 			for (int32 MuscleIdx = 0; MuscleIdx < MuscleActivationElement.Num(); MuscleIdx++)
 			{
@@ -156,20 +156,20 @@ namespace Chaos::Softs
 		int32 GetMuscleVertexCount(int32 MuscleIndex) { return MuscleVertexCount[MuscleIndex]; };
 
 		/* returns muscle activation of the specified muscle */
-		float GetMuscleActivation(int32 MuscleIndex) const { 
+		FSolverReal GetMuscleActivation(int32 MuscleIndex) const { 
 			if (MuscleActivation.IsValidIndex(MuscleIndex))
 			{
 				return MuscleActivation[MuscleIndex];
 			}
-			return 0.f;  
+			return FSolverReal(0);  
 		};
 
 		/* Sets muscle activation of the specified muscle and returns bool for success*/
-		bool SetMuscleActivation(int32 MuscleIndex, float InMuscleActivation)
+		bool SetMuscleActivation(int32 MuscleIndex, FSolverReal InMuscleActivation)
 		{
 			if (MuscleActivation.IsValidIndex(MuscleIndex))
 			{
-				MuscleActivation[MuscleIndex] = FMath::Clamp(InMuscleActivation, 0.f, 1.f);
+				MuscleActivation[MuscleIndex] = FMath::Clamp(InMuscleActivation, FSolverReal(0), FSolverReal(1));
 				return true;
 			}
 			return false;
@@ -178,17 +178,19 @@ namespace Chaos::Softs
 	private:	
 		TArray<TArray<int32>> MuscleActivationElement;
 		TArray<FIntVector2> OriginInsertionPair;
-		TArray<float> MuscleRestLength;
-		TArray<float> MuscleActivation;
+		TArray<FSolverReal> MuscleRestLength;
+		TArray<FSolverReal> MuscleActivation;
 		TArray<TArray<Chaos::PMatrix33d>> FiberDirectionMatrix;
-		TArray<TArray<float>> ContractionVolumeScale;
-		TArray<float> FiberLengthRatioAtMaxActivation;
-		TArray<float> MuscleLengthRatioThresholdForMaxActivation;
-		TArray<float> InflationVolumeScale;
+		TArray<TArray<FSolverReal>> ContractionVolumeScale;
+		TArray<FSolverReal> FiberLengthRatioAtMaxActivation;
+		TArray<FSolverReal> MuscleLengthRatioThresholdForMaxActivation;
+		TArray<FSolverReal> InflationVolumeScale;
 		TArray<Chaos::FLinearCurve> LengthActivationCurves;
 		TArray<int32> MuscleVertexOffset;
 		TArray<int32> MuscleVertexCount;
 	};
 
+	template <typename T, typename ParticleType>
+	using FMuscleActivationConstraints UE_DEPRECATED(5.8, "Deprecated. this class is to be deleted, use FMuscleActivationTetrahedralConstraints instead") = FMuscleActivationTetrahedralConstraints;
 
 }// End namespace Chaos::Softs
diff --git a/Engine/Source/Runtime/Experimental/Chaos/Public/Chaos/XPBDCorotatedConstraints.h b/Engine/Source/Runtime/Experimental/Chaos/Public/Chaos/XPBDCorotatedConstraints.h
--- a/Engine/Source/Runtime/Experimental/Chaos/Public/Chaos/XPBDCorotatedConstraints.h
+++ b/Engine/Source/Runtime/Experimental/Chaos/Public/Chaos/XPBDCorotatedConstraints.h
@@ -14,43 +14,42 @@ DECLARE_CYCLE_STAT(TEXT("Chaos XPBD Corotated Constraint Det Compute"), STAT_Cha
 
 namespace Chaos::Softs
 {
-
-	template <typename T, typename ParticleType>
-	class FXPBDCorotatedConstraints 
+	// Tetrahedral constraint using fixed corotated constitutive model [Stomakhin et al., 2012]
+	class FXPBDCorotatedTetrahedralConstraints
 	{
 
 	public:
 		//this one only accepts tetmesh input and mesh
-		FXPBDCorotatedConstraints(
-			const ParticleType& InParticles,
+		FXPBDCorotatedTetrahedralConstraints(
+			const FSolverParticles& InParticles,
 			const TArray<TVector<int32, 4>>& InMesh,
 			const bool bRecordMetricIn = true,
-			const T& EMesh = (T)10.0,
-			const T& NuMesh = (T).3
+			const FSolverReal& EMesh = (FSolverReal)10.0,
+			const FSolverReal& NuMesh = (FSolverReal).3
 			)
 			: bRecordMetric(bRecordMetricIn), MeshConstraints(InMesh)
 		{	
 
-			LambdaArray.Init((T)0., 2 * MeshConstraints.Num());
-			DmInverse.Init((T)0., 9 * MeshConstraints.Num());
-			Measure.Init((T)0.,  MeshConstraints.Num());
-			Lambda = EMesh * NuMesh / (((T)1. + NuMesh) * ((T)1. - (T)2. * NuMesh));
-			Mu = EMesh / ((T)2. * ((T)1. + NuMesh));
+			LambdaArray.Init((FSolverReal)0., 2 * MeshConstraints.Num());
+			DmInverse.Init((FSolverReal)0., 9 * MeshConstraints.Num());
+			Measure.Init((FSolverReal)0.,  MeshConstraints.Num());
+			Lambda = EMesh * NuMesh / (((FSolverReal)1. + NuMesh) * ((FSolverReal)1. - (FSolverReal)2. * NuMesh));
+			Mu = EMesh / ((FSolverReal)2. * ((FSolverReal)1. + NuMesh));
 			for (int e = 0; e < InMesh.Num(); e++)
 			{
-				PMatrix<T, 3, 3> Dm = DsInit(e, InParticles);
-				PMatrix<T, 3, 3> DmInv = Dm.Inverse();
+				PMatrix<FSolverReal, 3, 3> Dm = DsInit(e, InParticles);
+				PMatrix<FSolverReal, 3, 3> DmInv = Dm.Inverse();
 				for (int r = 0; r < 3; r++) {
 					for (int c = 0; c < 3; c++) {
 						DmInverse[(3 * 3) * e + 3 * r + c] = DmInv.GetAt(r, c);
 					}
 				}
 
-				Measure[e] = Dm.Determinant() / (T)6.;
+				Measure[e] = Dm.Determinant() / (FSolverReal)6.;
 
 				//part of preprocessing: if inverted element is found, 
 				//invert it so that the measure is positive
-				if (Measure[e] < (T)0.)
+				if (Measure[e] < (FSolverReal)0.)
 				{
 
 					Measure[e] = -Measure[e];
@@ -61,21 +60,21 @@ namespace Chaos::Softs
 			InitColor(InParticles);
 		}
 
-		FXPBDCorotatedConstraints(
-			const ParticleType& InParticles,
+		FXPBDCorotatedTetrahedralConstraints(
+			const FSolverParticles& InParticles,
 			const TArray<TVector<int32, 4>>& InMesh,
-			const TArray<T>& EMeshArray,
-			const T& NuMesh = (T).3,
+			const TArray<FSolverReal>& EMeshArray,
+			const FSolverReal& NuMesh = (FSolverReal).3,
 			const bool bRecordMetricIn = false
 		)
 			: bRecordMetric(bRecordMetricIn), MeshConstraints(InMesh)
 		{
 			ensureMsgf(EMeshArray.Num() == InMesh.Num(), TEXT("Input Young Modulus Array Size is wrong"));
-			LambdaArray.Init((T)0., 2 * MeshConstraints.Num());
-			DmInverse.Init((T)0., 9 * MeshConstraints.Num());
-			Measure.Init((T)0., MeshConstraints.Num());
-			LambdaElementArray.Init((T)0., MeshConstraints.Num());
-			MuElementArray.Init((T)0., MeshConstraints.Num());
+			LambdaArray.Init((FSolverReal)0., 2 * MeshConstraints.Num());
+			DmInverse.Init((FSolverReal)0., 9 * MeshConstraints.Num());
+			Measure.Init((FSolverReal)0., MeshConstraints.Num());
+			LambdaElementArray.Init((FSolverReal)0., MeshConstraints.Num());
+			MuElementArray.Init((FSolverReal)0., MeshConstraints.Num());
 
 			for (int e = 0; e < InMesh.Num(); e++)
 			{
@@ -87,20 +86,20 @@ namespace Chaos::Softs
 			
 			for (int e = 0; e < InMesh.Num(); e++)
 			{
-				LambdaElementArray[e] = EMeshArray[e] * NuMesh / (((T)1. + NuMesh) * ((T)1. - (T)2. * NuMesh));
-				MuElementArray[e] = EMeshArray[e] / ((T)2. * ((T)1. + NuMesh));
+				LambdaElementArray[e] = EMeshArray[e] * NuMesh / (((FSolverReal)1. + NuMesh) * ((FSolverReal)1. - (FSolverReal)2. * NuMesh));
+				MuElementArray[e] = EMeshArray[e] / ((FSolverReal)2. * ((FSolverReal)1. + NuMesh));
 
-				PMatrix<T, 3, 3> Dm = DsInit(e, InParticles);
-				PMatrix<T, 3, 3> DmInv = Dm.Inverse();
+				PMatrix<FSolverReal, 3, 3> Dm = DsInit(e, InParticles);
+				PMatrix<FSolverReal, 3, 3> DmInv = Dm.Inverse();
 				for (int r = 0; r < 3; r++) {
 					for (int c = 0; c < 3; c++) {
 						DmInverse[(3 * 3) * e + 3 * r + c] = DmInv.GetAt(r, c);
 					}
 				}
 
-				Measure[e] = Dm.Determinant() / (T)6.;
+				Measure[e] = Dm.Determinant() / (FSolverReal)6.;
 
-				if (Measure[e] < (T)0.)
+				if (Measure[e] < (FSolverReal)0.)
 				{
 					Measure[e] = -Measure[e];
 				}
@@ -109,25 +108,25 @@ namespace Chaos::Softs
 			InitColor(InParticles);
 		}
 
-		FXPBDCorotatedConstraints(
-			const ParticleType& InParticles,
+		FXPBDCorotatedTetrahedralConstraints(
+			const FSolverParticles& InParticles,
 			const TArray<TVector<int32, 4>>& InMesh,
-			const TArray<T>& EMeshArray,
-			const TArray<T>& NuMeshArray,
-			TArray<T>&& AlphaJMeshArray,
+			const TArray<FSolverReal>& EMeshArray,
+			const TArray<FSolverReal>& NuMeshArray,
+			TArray<FSolverReal>&& AlphaJMeshArray,
 			const FDeformableXPBDCorotatedParams& InParams,
-			const T& NuMesh = (T).3,
+			const FSolverReal& NuMesh = (FSolverReal).3,
 			const bool bRecordMetricIn = false, 
 			const bool bDoColoring = true
 		)
 			: CorotatedParams(InParams), AlphaJArray(MoveTemp(AlphaJMeshArray)), bRecordMetric(bRecordMetricIn), MeshConstraints(InMesh)
 		{
 			ensureMsgf(EMeshArray.Num() == InMesh.Num(), TEXT("Input Young Modulus Array Size is wrong"));
-			LambdaArray.Init((T)0., 2 * MeshConstraints.Num());
-			DmInverse.Init((T)0., 9 * MeshConstraints.Num());
-			Measure.Init((T)0., MeshConstraints.Num());
-			LambdaElementArray.Init((T)0., MeshConstraints.Num());
-			MuElementArray.Init((T)0., MeshConstraints.Num());
+			LambdaArray.Init((FSolverReal)0., 2 * MeshConstraints.Num());
+			DmInverse.Init((FSolverReal)0., 9 * MeshConstraints.Num());
+			Measure.Init((FSolverReal)0., MeshConstraints.Num());
+			LambdaElementArray.Init((FSolverReal)0., MeshConstraints.Num());
+			MuElementArray.Init((FSolverReal)0., MeshConstraints.Num());
 
 			for (int e = 0; e < InMesh.Num(); e++)
 			{
@@ -135,9 +134,9 @@ namespace Chaos::Softs
 				{
 					ensure(MeshConstraints[e][j] > -1 && MeshConstraints[e][j] < int32(InParticles.Size()));
 					
-					if (InParticles.InvM(MeshConstraints[e][j]) == (T)0.)
+					if (InParticles.InvM(MeshConstraints[e][j]) == (FSolverReal)0.)
 					{
-						AlphaJArray[e] = (T)1.;
+						AlphaJArray[e] = (FSolverReal)1.;
 					}
 				}
 				
@@ -145,20 +144,20 @@ namespace Chaos::Softs
 
 			for (int e = 0; e < InMesh.Num(); e++)
 			{
-				LambdaElementArray[e] = EMeshArray[e] * NuMeshArray[e] / (((T)1. + NuMeshArray[e]) * ((T)1. - (T)2. * NuMeshArray[e]));
-				MuElementArray[e] = EMeshArray[e] / ((T)2. * ((T)1. + NuMeshArray[e]));
+				LambdaElementArray[e] = EMeshArray[e] * NuMeshArray[e] / (((FSolverReal)1. + NuMeshArray[e]) * ((FSolverReal)1. - (FSolverReal)2. * NuMeshArray[e]));
+				MuElementArray[e] = EMeshArray[e] / ((FSolverReal)2. * ((FSolverReal)1. + NuMeshArray[e]));
 
-				PMatrix<T, 3, 3> Dm = DsInit(e, InParticles);
-				PMatrix<T, 3, 3> DmInv = Dm.Inverse();
+				PMatrix<FSolverReal, 3, 3> Dm = DsInit(e, InParticles);
+				PMatrix<FSolverReal, 3, 3> DmInv = Dm.Inverse();
 				for (int r = 0; r < 3; r++) {
 					for (int c = 0; c < 3; c++) {
 						DmInverse[(3 * 3) * e + 3 * r + c] = DmInv.GetAt(r, c);
 					}
 				}
 
-				Measure[e] = Dm.Determinant() / (T)6.;
+				Measure[e] = Dm.Determinant() / (FSolverReal)6.;
 
-				if (Measure[e] < (T)0.)
+				if (Measure[e] < (FSolverReal)0.)
 				{
 					Measure[e] = -Measure[e];
 				}
@@ -169,33 +168,33 @@ namespace Chaos::Softs
 			}
 		}
 
-		FXPBDCorotatedConstraints(
-			const ParticleType& InParticles,
+		FXPBDCorotatedTetrahedralConstraints(
+			const FSolverParticles& InParticles,
 			const TArray<TVector<int32, 4>>& InMesh,
-			const T GridN = (T).1,
-			const T& EMesh = (T)10.0,
-			const T& NuMesh = (T).3
+			const FSolverReal GridN = (FSolverReal).1,
+			const FSolverReal& EMesh = (FSolverReal)10.0,
+			const FSolverReal& NuMesh = (FSolverReal).3
 		)
 			: MeshConstraints(InMesh)
 		{
-			LambdaArray.Init((T)0., 2 * MeshConstraints.Num());
-			DmInverse.Init((T)0., 9 * MeshConstraints.Num());
-			Measure.Init((T)0., MeshConstraints.Num());
-			Lambda = EMesh * NuMesh / (((T)1. + NuMesh) * ((T)1. - (T)2. * NuMesh));
-			Mu = EMesh / ((T)2. * ((T)1. + NuMesh));
+			LambdaArray.Init((FSolverReal)0., 2 * MeshConstraints.Num());
+			DmInverse.Init((FSolverReal)0., 9 * MeshConstraints.Num());
+			Measure.Init((FSolverReal)0., MeshConstraints.Num());
+			Lambda = EMesh * NuMesh / (((FSolverReal)1. + NuMesh) * ((FSolverReal)1. - (FSolverReal)2. * NuMesh));
+			Mu = EMesh / ((FSolverReal)2. * ((FSolverReal)1. + NuMesh));
 			for (int e = 0; e < InMesh.Num(); e++)
 			{
-				PMatrix<T, 3, 3> Dm = DsInit(e, InParticles);
-				PMatrix<T, 3, 3> DmInv = Dm.Inverse();
+				PMatrix<FSolverReal, 3, 3> Dm = DsInit(e, InParticles);
+				PMatrix<FSolverReal, 3, 3> DmInv = Dm.Inverse();
 				for (int r = 0; r < 3; r++) {
 					for (int c = 0; c < 3; c++) {
 						DmInverse[(3 * 3) * e + 3 * r + c] = DmInv.GetAt(r, c);
 					}
 				}
 
-				Measure[e] = Dm.Determinant() / (T)6.;
+				Measure[e] = Dm.Determinant() / (FSolverReal)6.;
 
-				if (Measure[e] < (T)0.)
+				if (Measure[e] < (FSolverReal)0.)
 				{
 					Measure[e] = -Measure[e];
 				}
@@ -203,10 +202,10 @@ namespace Chaos::Softs
 			DmInverseSave = DmInverse;
 		}
 
-		virtual ~FXPBDCorotatedConstraints() {}
+		virtual ~FXPBDCorotatedTetrahedralConstraints() {}
 
-		PMatrix<T, 3, 3> DsInit(const int e, const ParticleType& InParticles) const {
-			PMatrix<T, 3, 3> Result((T)0.);
+		PMatrix<FSolverReal, 3, 3> DsInit(const int e, const FSolverParticles& InParticles) const {
+			PMatrix<FSolverReal, 3, 3> Result((FSolverReal)0.);
 			for (int i = 0; i < 3; i++) {
 				for (int c = 0; c < 3; c++) {
 					Result.SetAt(c, i, InParticles.GetX(MeshConstraints[e][i + 1])[c] - InParticles.GetX(MeshConstraints[e][0])[c]);
@@ -215,9 +214,8 @@ namespace Chaos::Softs
 			return Result;
 		}
 
-		template<typename SolverParticlesOrRange>
-		PMatrix<T, 3, 3> Ds(const int e, const SolverParticlesOrRange& InParticles) const {
-			PMatrix<T, 3, 3> Result((T)0.);
+		PMatrix<FSolverReal, 3, 3> Ds(const int e, const FSolverParticlesRange& InParticles) const {
+			PMatrix<FSolverReal, 3, 3> Result((FSolverReal)0.);
 			for (int i = 0; i < 3; i++) {
 				for (int c = 0; c < 3; c++) {
 					Result.SetAt(c, i, InParticles.GetP(MeshConstraints[e][i+1])[c] - InParticles.GetP(MeshConstraints[e][0])[c]);
@@ -226,13 +224,12 @@ namespace Chaos::Softs
 			return Result;
 		}
 
-		template<typename SolverParticlesOrRange>
-		PMatrix<T, 3, 3> F(const int e, const SolverParticlesOrRange& InParticles) const {
+		PMatrix<FSolverReal, 3, 3> F(const int e, const FSolverParticlesRange& InParticles) const {
 			return ElementDmInv(e) * Ds(e, InParticles);
 		}
 
-		PMatrix<T, 3, 3> ElementDmInv(const int e) const {
-			PMatrix<T, 3, 3> DmInv((T)0.);
+		PMatrix<FSolverReal, 3, 3> ElementDmInv(const int e) const {
+			PMatrix<FSolverReal, 3, 3> DmInv((FSolverReal)0.);
 			for (int r = 0; r < 3; r++) {
 				for (int c = 0; c < 3; c++) {
 					DmInv.SetAt(r, c, DmInverse[(3 * 3) * e + 3 * r + c]);
@@ -241,8 +238,8 @@ namespace Chaos::Softs
 			return DmInv;
 		}
 
-		PMatrix<T, 3, 3> ElementDmInvSave(const int e) const {
-			PMatrix<T, 3, 3> DmInv((T)0.);
+		PMatrix<FSolverReal, 3, 3> ElementDmInvSave(const int e) const {
+			PMatrix<FSolverReal, 3, 3> DmInv((FSolverReal)0.);
 			for (int r = 0; r < 3; r++) {
 				for (int c = 0; c < 3; c++) {
 					DmInv.SetAt(r, c, DmInverseSave[(3 * 3) * e + 3 * r + c]);
@@ -253,32 +250,29 @@ namespace Chaos::Softs
 		
 		virtual void Init() const 
 		{
-			for (T& Lambdas : LambdaArray) { Lambdas = (T)0.; }
+			for (FSolverReal& Lambdas : LambdaArray) { Lambdas = (FSolverReal)0.; }
 		}
 
-		template<typename SolverParticlesOrRange>
-		void ApplyInSerial(SolverParticlesOrRange& Particles, const T Dt, const int32 ElementIndex) const
+		virtual void ApplyInSerial(FSolverParticlesRange& Particles, const FSolverReal Dt, const int32 ElementIndex) const
 		{
 			TRACE_CPUPROFILER_EVENT_SCOPE(STAT_ChaosXPBDCorotatedApplySingle);
 
-			TVec4<TVector<T, 3>> PolarDelta = GetPolarDelta(Particles, Dt, ElementIndex);
+			TVec4<TVector<FSolverReal, 3>> PolarDelta = GetPolarDelta(Particles, Dt, ElementIndex);
 
 			for (int i = 0; i < 4; i++)
 			{
 				Particles.P(MeshConstraints[ElementIndex][i]) += PolarDelta[i];
 			}
 
-			TVec4<TVector<T, 3>> DetDelta = GetDeterminantDelta(Particles, Dt, ElementIndex);
+			TVec4<TVector<FSolverReal, 3>> DetDelta = GetDeterminantDelta(Particles, Dt, ElementIndex);
 
 			for (int i = 0; i < 4; i++)
 			{
 				Particles.P(MeshConstraints[ElementIndex][i]) += DetDelta[i];
 			}
-
 		}
 
-		template<typename SolverParticlesOrRange>
-		void ApplyInSerial(SolverParticlesOrRange& Particles, const T Dt) const
+		void ApplyInSerial(FSolverParticlesRange& Particles, const FSolverReal Dt) const
 		{
 			SCOPE_CYCLE_COUNTER(STAT_ChaosXPBDCorotated);
 			TRACE_CPUPROFILER_EVENT_SCOPE(STAT_ChaosXPBDCorotatedApplySerial);
@@ -293,17 +287,15 @@ namespace Chaos::Softs
 					ApplyInSerial(Particles, Dt, ConstraintIndex);
 				}
 			}
-
 		}
 
-		template<typename SolverParticlesOrRange>
-		void ApplyInParallel(SolverParticlesOrRange& Particles, const T Dt) const
+		void ApplyInParallel(FSolverParticlesRange& Particles, const FSolverReal Dt) const
 		{	
 			//code for error metric:
 			if (bRecordMetric)
 			{
-				GError.Init((T)0., 3 * Particles.Size());
-				HErrorArray.Init((T)0., 2 * MeshConstraints.Num());
+				GError.Init((FSolverReal)0., 3 * Particles.Size());
+				HErrorArray.Init((FSolverReal)0., 2 * MeshConstraints.Num());
 			}
 			
 			{
@@ -340,13 +332,19 @@ namespace Chaos::Softs
 			}
 		}
 
-		TVec4<TVector<T, 3>> GetPolarGradient(const PMatrix<T, 3, 3>& Fe, const PMatrix<T, 3, 3>& Re, const PMatrix<T, 3, 3>& DmInvT, const T C1) const
+		void ApplyInParallel(FSolverParticles& Particles, const FSolverReal Dt) const
 		{
-			//TVector<T, 81> dRdF((T)0.);
+			FSolverParticlesRange Range(&Particles, /*InOffset = */ 0, /*InRange = */ Particles.Size());
+			ApplyInParallel(Range, Dt);
+		}
+
+		TVec4<TVector<FSolverReal, 3>> GetPolarGradient(const PMatrix<FSolverReal, 3, 3>& Fe, const PMatrix<FSolverReal, 3, 3>& Re, const PMatrix<FSolverReal, 3, 3>& DmInvT, const FSolverReal C1) const
+		{
+			//TVector<FSolverReal, 81> dRdF((FSolverReal)0.);
 			//Chaos::dRdFCorotated(Fe, dRdF);
-			TVec4<TVector<T, 3>> dC1(TVector<T, 3>((T)0.));
+			TVec4<TVector<FSolverReal, 3>> dC1(TVector<FSolverReal, 3>((FSolverReal)0.));
 			//dC1 = dC1dF * dFdX
-			PMatrix<T, 3, 3> A = DmInvT * (Fe - Re);
+			PMatrix<FSolverReal, 3, 3> A = DmInvT * (Fe - Re);
 			for (int alpha = 0; alpha < 3; alpha++) {
 				for (int l = 0; l < 3; l++) {
 					dC1[0][alpha] -= A.GetAt(alpha, l);
@@ -372,17 +370,17 @@ namespace Chaos::Softs
 
 		}
 
-		TVec4<TVector<T, 3>> GetDeterminantGradient(const PMatrix<T, 3, 3>& Fe, const PMatrix<T, 3, 3>& DmInvT) const
+		TVec4<TVector<FSolverReal, 3>> GetDeterminantGradient(const PMatrix<FSolverReal, 3, 3>& Fe, const PMatrix<FSolverReal, 3, 3>& DmInvT) const
 		{
 			TRACE_CPUPROFILER_EVENT_SCOPE(STAT_ChaosXPBDCorotatedGetDetGradient);
 			//SCOPE_CYCLE_COUNTER(STAT_ChaosXPBDCorotatedDet);
 			//const TVec4<int32>& Constraint = MeshConstraints[ElementIndex];
 
-			//const PMatrix<T, 3, 3> Fe = F(ElementIndex, Particles);
-			//PMatrix<T, 3, 3> DmInvT = ElementDmInv(ElementIndex).GetTransposed();
-			TVec4<TVector<T, 3>> dC2(TVector<T, 3>((T)0.));
+			//const PMatrix<FSolverReal, 3, 3> Fe = F(ElementIndex, Particles);
+			//PMatrix<FSolverReal, 3, 3> DmInvT = ElementDmInv(ElementIndex).GetTransposed();
+			TVec4<TVector<FSolverReal, 3>> dC2(TVector<FSolverReal, 3>((FSolverReal)0.));
 
-			PMatrix<T, 3, 3> JFinvT((T)0.);
+			PMatrix<FSolverReal, 3, 3> JFinvT((FSolverReal)0.);
 			JFinvT.SetAt(0, 0, Fe.GetAt(1, 1) * Fe.GetAt(2, 2) - Fe.GetAt(2, 1) * Fe.GetAt(1, 2));
 			JFinvT.SetAt(0, 1, Fe.GetAt(2, 0) * Fe.GetAt(1, 2) - Fe.GetAt(1, 0) * Fe.GetAt(2, 2));
 			JFinvT.SetAt(0, 2, Fe.GetAt(1, 0) * Fe.GetAt(2, 1) - Fe.GetAt(2, 0) * Fe.GetAt(1, 1));
@@ -393,7 +391,7 @@ namespace Chaos::Softs
 			JFinvT.SetAt(2, 1, Fe.GetAt(1, 0) * Fe.GetAt(0, 2) - Fe.GetAt(0, 0) * Fe.GetAt(1, 2));
 			JFinvT.SetAt(2, 2, Fe.GetAt(0, 0) * Fe.GetAt(1, 1) - Fe.GetAt(1, 0) * Fe.GetAt(0, 1));
 
-			PMatrix<T, 3, 3> JinvTDmInvT = DmInvT * JFinvT;
+			PMatrix<FSolverReal, 3, 3> JinvTDmInvT = DmInvT * JFinvT;
 
 			for (int ie = 0; ie < 3; ie++) {
 				for (int alpha = 0; alpha < 3; alpha++) {
@@ -408,14 +406,14 @@ namespace Chaos::Softs
 			return dC2;
 		}
 
-		void ModifyDmInverseFromMuscleLength(const int32 ElemIdx, const T FiberLengthRatio, const PMatrix<T, 3, 3>& FiberDir, const T ContractionVolumeScale) const
+		void ModifyDmInverseFromMuscleLength(const int32 ElemIdx, const FSolverReal FiberLengthRatio, const PMatrix<FSolverReal, 3, 3>& FiberDir, const FSolverReal ContractionVolumeScale) const
 		{
 			TRACE_CPUPROFILER_EVENT_SCOPE(STAT_ChaosXPBDCorotatedModifyDmInverseFromMuscleLength);
-			PMatrix<T, 3, 3> DmInv = ElementDmInvSave(ElemIdx);
+			PMatrix<FSolverReal, 3, 3> DmInv = ElementDmInvSave(ElemIdx);
 			if (FiberLengthRatio < 1)
 			{
-				PMatrix<T, 3, 3> D(1 / FiberLengthRatio, FMath::Pow(FiberLengthRatio, ContractionVolumeScale/T(2)), FMath::Pow(FiberLengthRatio, ContractionVolumeScale/T(2)));
-				PMatrix<T, 3, 3> Factor = FiberDir * D * FiberDir.GetTransposed();
+				PMatrix<FSolverReal, 3, 3> D(1 / FiberLengthRatio, FMath::Pow(FiberLengthRatio, ContractionVolumeScale/FSolverReal(2)), FMath::Pow(FiberLengthRatio, ContractionVolumeScale/FSolverReal(2)));
+				PMatrix<FSolverReal, 3, 3> Factor = FiberDir * D * FiberDir.GetTransposed();
 				DmInv = Factor * DmInv; //Matrix multiplication convention: AB = B*A
 			}
 			for (int Row = 0; Row < 3; Row++)
@@ -427,14 +425,14 @@ namespace Chaos::Softs
 			}
 		}
 
-		void ModifyDmInverseSaveFromInflationVolumeScale(const int32 ElemIdx, const T InflationVolumeScale, const PMatrix<T, 3, 3>& FiberDir)
+		void ModifyDmInverseSaveFromInflationVolumeScale(const int32 ElemIdx, const FSolverReal InflationVolumeScale, const PMatrix<FSolverReal, 3, 3>& FiberDir)
 		{
-			PMatrix<T, 3, 3> DmInv = ElementDmInvSave(ElemIdx);
+			PMatrix<FSolverReal, 3, 3> DmInv = ElementDmInvSave(ElemIdx);
 			if (FMath::Abs(InflationVolumeScale - 1) > UE_SMALL_NUMBER)
 			{
-				T InvSqrt = FMath::InvSqrt(InflationVolumeScale);
-				PMatrix<T, 3, 3> D(1, InvSqrt, InvSqrt);
-				PMatrix<T, 3, 3> Factor = FiberDir * D * FiberDir.GetTransposed();
+				FSolverReal InvSqrt = FMath::InvSqrt(InflationVolumeScale);
+				PMatrix<FSolverReal, 3, 3> D(1, InvSqrt, InvSqrt);
+				PMatrix<FSolverReal, 3, 3> Factor = FiberDir * D * FiberDir.GetTransposed();
 				DmInv = Factor * DmInv; //Matrix multiplication convention: AB = B*A
 			}
 			for (int Row = 0; Row < 3; Row++)
@@ -448,7 +446,7 @@ namespace Chaos::Softs
 		}
 	protected:
 
-		void InitColor(const ParticleType& Particles)
+		void InitColor(const FSolverParticles& Particles)
 		{
 
 			{
@@ -456,8 +454,8 @@ namespace Chaos::Softs
 
 				// Reorder constraints based on color so each array in ConstraintsPerColor contains contiguous elements.
 				TArray<TVec4<int32>> ReorderedConstraints;
-				TArray<T> ReorderedMeasure;
-				TArray<T> ReorderedDmInverse;
+				TArray<FSolverReal> ReorderedMeasure;
+				TArray<FSolverReal> ReorderedDmInverse;
 				TArray<int32> OrigToReorderedIndices; // used to reorder stiffness indices
 				ReorderedConstraints.SetNumUninitialized(MeshConstraints.Num());
 				ReorderedMeasure.SetNumUninitialized(Measure.Num());
@@ -492,30 +490,29 @@ namespace Chaos::Softs
 			}
 		}
 
-		template<typename SolverParticlesOrRange>
-		TVec4<TVector<T, 3>> GetDeterminantDelta(const SolverParticlesOrRange& Particles, const T Dt, const int32 ElementIndex, const T Tol = (T)1e-3) const
+		virtual TVec4<TVector<FSolverReal, 3>> GetDeterminantDelta(const FSolverParticlesRange& Particles, const FSolverReal Dt, const int32 ElementIndex, const FSolverReal Tol = (FSolverReal)1e-3) const
 		{
 			TRACE_CPUPROFILER_EVENT_SCOPE(STAT_ChaosXPBDCorotatedApplyDet);
 			//SCOPE_CYCLE_COUNTER(STAT_ChaosXPBDCorotatedDet);
 			////const TVec4<int32>& Constraint = MeshConstraints[ElementIndex];
 
-			const PMatrix<T, 3, 3> Fe = F(ElementIndex, Particles);
-			PMatrix<T, 3, 3> DmInvT = ElementDmInv(ElementIndex).GetTransposed();
+			const PMatrix<FSolverReal, 3, 3> Fe = F(ElementIndex, Particles);
+			PMatrix<FSolverReal, 3, 3> DmInvT = ElementDmInv(ElementIndex).GetTransposed();
 			
-			T J = Fe.Determinant();
+			FSolverReal J = Fe.Determinant();
 			if (J - AlphaJArray[ElementIndex] < Tol)
 			{
-				return TVec4<TVector<T, 3>>(TVector<T, 3>((T)0.));
+				return TVec4<TVector<FSolverReal, 3>>(TVector<FSolverReal, 3>((FSolverReal)0.));
 			}
 
-			TVec4<TVector<T, 3>> dC2 = GetDeterminantGradient(Fe, DmInvT);
+			TVec4<TVector<FSolverReal, 3>> dC2 = GetDeterminantGradient(Fe, DmInvT);
 
-			//T AlphaTilde = (T)2. / (Dt * Dt * Lambda * Measure[ElementIndex]);
-			T AlphaTilde = (T)2. / (Dt * Dt * LambdaElementArray[ElementIndex] * Measure[ElementIndex]);
+			//T AlphaTilde = (FSolverReal)2. / (Dt * Dt * Lambda * Measure[ElementIndex]);
+			FSolverReal AlphaTilde = (FSolverReal)2. / (Dt * Dt * LambdaElementArray[ElementIndex] * Measure[ElementIndex]);
 
-			if (LambdaElementArray[ElementIndex] > (T)1. / (T)UE_SMALL_NUMBER)
+			if (LambdaElementArray[ElementIndex] > (FSolverReal)1. / (FSolverReal)UE_SMALL_NUMBER)
 			{
-				AlphaTilde = (T)0.;
+				AlphaTilde = (FSolverReal)0.;
 			}
 
 			if (bRecordMetric)
@@ -532,9 +529,9 @@ namespace Chaos::Softs
 			}
 			
 
-			T DLambda = (AlphaJArray[ElementIndex] - J) - AlphaTilde * LambdaArray[2 * ElementIndex + 1];
+			FSolverReal DLambda = (AlphaJArray[ElementIndex] - J) - AlphaTilde * LambdaArray[2 * ElementIndex + 1];
 
-			T Denom = AlphaTilde;
+			FSolverReal Denom = AlphaTilde;
 			for (int i = 0; i < 4; i++)
 			{
 				for (int j = 0; j < 3; j++)
@@ -544,7 +541,7 @@ namespace Chaos::Softs
 			}
 			DLambda /= Denom;
 			LambdaArray[2 * ElementIndex + 1] += DLambda;
-			TVec4<TVector<T, 3>> Delta(TVector<T, 3>((T)0.));
+			TVec4<TVector<FSolverReal, 3>> Delta(TVector<FSolverReal, 3>((FSolverReal)0.));
 			for (int i = 0; i < 4; i++)
 			{
 				for (int j = 0; j < 3; j++)
@@ -555,21 +552,19 @@ namespace Chaos::Softs
 			return Delta;
 		}
 
-
-		template<typename SolverParticlesOrRange>
-		TVec4<TVector<T, 3>> GetPolarDelta(const SolverParticlesOrRange& Particles, const T Dt, const int32 ElementIndex, const T Tol = (T)1e-3) const
+		virtual TVec4<TVector<FSolverReal, 3>> GetPolarDelta(const FSolverParticlesRange& Particles, const FSolverReal Dt, const int32 ElementIndex, const FSolverReal Tol = (FSolverReal)1e-3) const
 		{	
 			TRACE_CPUPROFILER_EVENT_SCOPE(STAT_ChaosXPBDCorotatedApplyPolar);
 			SCOPE_CYCLE_COUNTER(STAT_ChaosXPBDCorotatedPolar);
-			const PMatrix<T, 3, 3> Fe = F(ElementIndex, Particles);
+			const PMatrix<FSolverReal, 3, 3> Fe = F(ElementIndex, Particles);
 
-			PMatrix<T, 3, 3> Re((T)0.), Se((T)0.);
+			PMatrix<FSolverReal, 3, 3> Re((FSolverReal)0.), Se((FSolverReal)0.);
 
 			Chaos::PolarDecomposition(Fe, Re, Se);
 
-			Re *= FGenericPlatformMath::Pow(AlphaJArray[ElementIndex], (T)1. / (T)3.);
+			Re *= FGenericPlatformMath::Pow(AlphaJArray[ElementIndex], (FSolverReal)1. / (FSolverReal)3.);
 
-			T C1 = (T)0.;
+			FSolverReal C1 = (FSolverReal)0.;
 			for (int i = 0; i < 3; i++)
 			{
 				for (int j = 0; j < 3; j++)
@@ -581,20 +576,20 @@ namespace Chaos::Softs
 
 			if (C1 < Tol )
 			{
-				return TVec4<TVector<T, 3>>(TVector<T, 3>((T)0.));
+				return TVec4<TVector<FSolverReal, 3>>(TVector<FSolverReal, 3>((FSolverReal)0.));
 			}
 
 
-			PMatrix<T, 3, 3> DmInvT = ElementDmInv(ElementIndex).GetTransposed();
+			PMatrix<FSolverReal, 3, 3> DmInvT = ElementDmInv(ElementIndex).GetTransposed();
 
-			TVec4<TVector<T, 3>> dC1 = GetPolarGradient(Fe, Re, DmInvT, C1);
+			TVec4<TVector<FSolverReal, 3>> dC1 = GetPolarGradient(Fe, Re, DmInvT, C1);
 
-			//T AlphaTilde = (T)1. / (Dt * Dt * Mu * Measure[ElementIndex]);
-			T AlphaTilde = (T)1. / (Dt * Dt * MuElementArray[ElementIndex] * Measure[ElementIndex]);
+			//T AlphaTilde = (FSolverReal)1. / (Dt * Dt * Mu * Measure[ElementIndex]);
+			FSolverReal AlphaTilde = (FSolverReal)1. / (Dt * Dt * MuElementArray[ElementIndex] * Measure[ElementIndex]);
 
-			if (MuElementArray[ElementIndex] > (T)1./ (T)UE_SMALL_NUMBER) 
+			if (MuElementArray[ElementIndex] > (FSolverReal)1./ (FSolverReal)UE_SMALL_NUMBER) 
 			{
-				AlphaTilde = (T)0.;
+				AlphaTilde = (FSolverReal)0.;
 			}
 
 			if (bRecordMetric)
@@ -610,9 +605,9 @@ namespace Chaos::Softs
 				HErrorArray[2 * ElementIndex + 0] = C1 + AlphaTilde * LambdaArray[2 * ElementIndex + 0];
 			}
 
-			T DLambda = -C1 - AlphaTilde* LambdaArray[2 * ElementIndex + 0];
+			FSolverReal DLambda = -C1 - AlphaTilde* LambdaArray[2 * ElementIndex + 0];
 
-			T Denom = AlphaTilde;
+			FSolverReal Denom = AlphaTilde;
 			for (int i = 0; i < 4; i++)
 			{
 				for (int j = 0; j < 3; j++)
@@ -622,7 +617,7 @@ namespace Chaos::Softs
 			}
 			DLambda /= Denom;
 			LambdaArray[2 * ElementIndex + 0] += DLambda;
-			TVec4<TVector<T, 3>> Delta(TVector<T, 3>((T)0.));
+			TVec4<TVector<FSolverReal, 3>> Delta(TVector<FSolverReal, 3>((FSolverReal)0.));
 			for (int i = 0; i < 4; i++) 
 			{
 				for (int j = 0; j < 3; j++)
@@ -634,31 +629,33 @@ namespace Chaos::Softs
 
 		}
 
-
 	protected:
-		mutable TArray<T> LambdaArray;
-		mutable TArray<T> DmInverse;
-		TArray<T> DmInverseSave;
+		mutable TArray<FSolverReal> LambdaArray;
+		mutable TArray<FSolverReal> DmInverse;
+		TArray<FSolverReal> DmInverseSave;
 		//parallel data:
 		FDeformableXPBDCorotatedParams CorotatedParams;
 
 		//material constants calculated from E:
-		T Mu;
-		T Lambda;
-		TArray<T> MuElementArray;
-		TArray<T> LambdaElementArray;
-		TArray<T> AlphaJArray;
-		mutable T HError;
-		mutable TArray<T> HErrorArray;
+		FSolverReal Mu;
+		FSolverReal Lambda;
+		TArray<FSolverReal> MuElementArray;
+		TArray<FSolverReal> LambdaElementArray;
+		TArray<FSolverReal> AlphaJArray;
+		mutable FSolverReal HError;
+		mutable TArray<FSolverReal> HErrorArray;
 		bool bRecordMetric;
 		bool VariableStiffness = false;
 
 		TArray<TVector<int32, 4>> MeshConstraints;
-		mutable TArray<T> Measure;
-		ParticleType RestParticles;
+		mutable TArray<FSolverReal> Measure;
+		FSolverParticles RestParticles;
 		TArray<int32> ConstraintsPerColorStartIndex; // Constraints are ordered so each batch is contiguous. This is ColorNum + 1 length so it can be used as start and end.
-		mutable TArray<T> GError;
-};
+		mutable TArray<FSolverReal> GError;
+	};
+
+	template <typename T, typename ParticleType>
+	using FXPBDCorotatedConstraints UE_DEPRECATED(5.8, "Deprecated. this class is to be deleted, use FXPBDCorotatedTetrahedralConstraints instead") = FXPBDCorotatedTetrahedralConstraints;
 
 }  // End namespace Chaos::Softs
 
diff --git a/Engine/Source/Runtime/Experimental/Chaos/Public/Chaos/XPBDCorotatedFiberConstraints.h b/Engine/Source/Runtime/Experimental/Chaos/Public/Chaos/XPBDCorotatedFiberConstraints.h
--- a/Engine/Source/Runtime/Experimental/Chaos/Public/Chaos/XPBDCorotatedFiberConstraints.h
+++ b/Engine/Source/Runtime/Experimental/Chaos/Public/Chaos/XPBDCorotatedFiberConstraints.h
@@ -12,11 +12,13 @@
 
 namespace Chaos::Softs
 {
-	template <typename T, typename ParticleType>
-	class FXPBDCorotatedFiberConstraints : public FXPBDCorotatedConstraints <T, ParticleType>
+	// Tetrahedral constraint using fixed corotated + muscle fiber constitutive model
+	// Fixed corotated: [Stomakhin et al., 2012]
+	// Muscle fiber: [Blemker et al., 2005]
+	class FXPBDCorotatedFiberTetrahedralConstraints : public FXPBDCorotatedTetrahedralConstraints
 	{
 
-		typedef FXPBDCorotatedConstraints<T, ParticleType> Base;
+		typedef FXPBDCorotatedTetrahedralConstraints Base;
 		using Base::MeshConstraints;
 		using Base::LambdaArray;
 		using Base::Measure;
@@ -24,25 +26,25 @@ namespace Chaos::Softs
 
 	public:
 		//this one only accepts tetmesh input and mesh
-		FXPBDCorotatedFiberConstraints(
-			const ParticleType& InParticles,
+		FXPBDCorotatedFiberTetrahedralConstraints(
+			const FSolverParticles& InParticles,
 			const TArray<TVector<int32, 4>>& InMesh,
 			const bool bRecordMetricIn = true,
-			const T& EMesh = (T)10.0,
-			const T& NuMesh = (T).3,
-			const TVector<T, 3> InFiberDir = TVector<T, 3>((T)1., (T)0., (T)0.),
-			const T InSigmaMax = (T)3e5
+			const FSolverReal& EMesh = (FSolverReal)10.0,
+			const FSolverReal& NuMesh = (FSolverReal).3,
+			const TVector<FSolverReal, 3> InFiberDir = TVector<FSolverReal, 3>((FSolverReal)1., (FSolverReal)0., (FSolverReal)0.),
+			const FSolverReal InSigmaMax = (FSolverReal)3e5
 		)
 			: Base(InParticles, InMesh, bRecordMetricIn, EMesh, NuMesh), SigmaMax(InSigmaMax), FiberDir(InFiberDir)
 		{
-			LambdaArray.Init((T)0., 3 * MeshConstraints.Num());
+			LambdaArray.Init((FSolverReal)0., 3 * MeshConstraints.Num());
 		}
 
-		virtual ~FXPBDCorotatedFiberConstraints() {}
+		virtual ~FXPBDCorotatedFiberTetrahedralConstraints() {}
 		
-		TVector<T, 3> GetFiberDir() { return FiberDir; }
+		TVector<FSolverReal, 3> GetFiberDir() { return FiberDir; }
 
-		void SetActivation(const T AlphaIn) { AlphaActivation = AlphaIn; }
+		void SetActivation(const FSolverReal AlphaIn) { AlphaActivation = AlphaIn; }
 
 		void SetTime(const float Time) const {
 			float FinalTime = 4.0f;
@@ -51,26 +53,26 @@ namespace Chaos::Softs
 			{
 				CurrentTime -= FinalTime;
 			}
-			AlphaActivation = (T)1. - (T)4. / FinalTime * FMath::Abs(CurrentTime - FinalTime / (T)2.);
+			AlphaActivation = (FSolverReal)1. - (FSolverReal)4. / FinalTime * FMath::Abs(CurrentTime - FinalTime / (FSolverReal)2.);
 		}
 
-		virtual void ApplyInSerial(ParticleType& Particles, const T Dt, const int32 ElementIndex) const override
+		virtual void ApplyInSerial(FSolverParticlesRange& Particles, const FSolverReal Dt, const int32 ElementIndex) const override
 		{
-			TVec4<TVector<T, 3>> PolarDelta = Base::GetPolarDelta(Particles, Dt, ElementIndex);
+			TVec4<TVector<FSolverReal, 3>> PolarDelta = Base::GetPolarDelta(Particles, Dt, ElementIndex);
 
 			for (int i = 0; i < 4; i++)
 			{
 				Particles.P(MeshConstraints[ElementIndex][i]) += PolarDelta[i];
 			}
 
-			TVec4<TVector<T, 3>> DetDelta = Base::GetDeterminantDelta(Particles, Dt, ElementIndex);
+			TVec4<TVector<FSolverReal, 3>> DetDelta = Base::GetDeterminantDelta(Particles, Dt, ElementIndex);
 
 			for (int i = 0; i < 4; i++)
 			{
 				Particles.P(MeshConstraints[ElementIndex][i]) += DetDelta[i];
 			}
 
-			TVec4<TVector<T, 3>> FiberDelta = GetFiberDelta(Particles, Dt, ElementIndex);
+			TVec4<TVector<FSolverReal, 3>> FiberDelta = GetFiberDelta(Particles, Dt, ElementIndex);
 
 			for (int i = 0; i < 4; i++)
 			{
@@ -80,19 +82,19 @@ namespace Chaos::Softs
 
 		}
 
-		TVec4<TVector<T, 3>> GetFiberGradient(const T dFpdL, const T dFadL,const T C3, const TVec4<TVector<T, 3>>& dLdX) const 
+		TVec4<TVector<FSolverReal, 3>> GetFiberGradient(const FSolverReal dFpdL, const FSolverReal dFadL,const FSolverReal C3, const TVec4<TVector<FSolverReal, 3>>& dLdX) const 
 		{
-			T dC3dL = (T)0.5 * (dFpdL + AlphaActivation * dFadL);
+			FSolverReal dC3dL = (FSolverReal)0.5 * (dFpdL + AlphaActivation * dFadL);
 			if (C3 == 0)
 			{
-				return TVec4<TVector<T, 3>>(TVector<T, 3>((T)0.));
+				return TVec4<TVector<FSolverReal, 3>>(TVector<FSolverReal, 3>((FSolverReal)0.));
 			}
 			else
 			{
 				dC3dL /= C3;
 			}
 
-			TVec4<TVector<T, 3>> dC3(TVector<T, 3>((T)0.));
+			TVec4<TVector<FSolverReal, 3>> dC3(TVector<FSolverReal, 3>((FSolverReal)0.));
 			for (int32 i = 0; i < 4; i++)
 			{
 				for (int32 j = 0; j < 3; j++)
@@ -108,19 +110,19 @@ namespace Chaos::Softs
 
 	private:
 
-		TVec4<TVector<T, 3>> GetFiberDelta(const ParticleType& Particles, const T Dt, const int32 ElementIndex, const T Tol = 1e-3) const
+		TVec4<TVector<FSolverReal, 3>> GetFiberDelta(const FSolverParticlesRange& Particles, const FSolverReal Dt, const int32 ElementIndex, const FSolverReal Tol = FSolverReal(1e-3)) const
 		{
-			const PMatrix<T, 3, 3> Fe = Base::F(ElementIndex, Particles);
+			const PMatrix<FSolverReal, 3, 3> Fe = Base::F(ElementIndex, Particles);
 
-			PMatrix<T, 3, 3> Re((T)0.), Se((T)0.);
+			PMatrix<FSolverReal, 3, 3> Re((FSolverReal)0.), Se((FSolverReal)0.);
 
 			Chaos::PolarDecomposition(Fe, Re, Se);
 
 			// l: fiber stretch, f = (vTCv)^(1/2)
-			TVector<T, 3> FeV = Fe.GetTransposed() * FiberDir;
-			TVector<T, 3> DmInverseV = Base::ElementDmInv(ElementIndex).GetTransposed() * FiberDir;
-			T L = FeV.Size();
-			TVec4<TVector<T, 3>> dLdX(TVector<T, 3>((T)0.));
+			TVector<FSolverReal, 3> FeV = Fe.GetTransposed() * FiberDir;
+			TVector<FSolverReal, 3> DmInverseV = Base::ElementDmInv(ElementIndex).GetTransposed() * FiberDir;
+			FSolverReal L = FeV.Size();
+			TVec4<TVector<FSolverReal, 3>> dLdX(TVector<FSolverReal, 3>((FSolverReal)0.));
 			for (int32 alpha = 0; alpha < 3; alpha++)
 			{
 				for (int32 s = 0; s < 3; s++)
@@ -137,61 +139,61 @@ namespace Chaos::Softs
 				}
 			}
 
-			T LambdaOFL = (T)1.4;
-			T P1 = (T)0.05;
-			T P2 = (T)6.6;
-			T FpIntegral = (T)0.;
-			T FaIntegral = (T)0.;
+			FSolverReal LambdaOFL = (FSolverReal)1.4;
+			FSolverReal P1 = (FSolverReal)0.05;
+			FSolverReal P2 = (FSolverReal)6.6;
+			FSolverReal FpIntegral = (FSolverReal)0.;
+			FSolverReal FaIntegral = (FSolverReal)0.;
 
-			T C3 = (T)0.;
-			T AlphaTilde = LambdaOFL / (SigmaMax * Dt * Dt * Measure[ElementIndex]);
-			T dFpdL = (T)0.;
-			T dFadL = (T)0.;
+			FSolverReal C3 = (FSolverReal)0.;
+			FSolverReal AlphaTilde = LambdaOFL / (SigmaMax * Dt * Dt * Measure[ElementIndex]);
+			FSolverReal dFpdL = (FSolverReal)0.;
+			FSolverReal dFadL = (FSolverReal)0.;
 
 			if (L > LambdaOFL)
 			{
-				FpIntegral = P1 * LambdaOFL / P2 * FMath::Exp(P2 * (L / LambdaOFL - (T)1.)) - P1 * (L - LambdaOFL);
-				dFpdL = P1 * FMath::Exp(P2 * (L / LambdaOFL - (T)1.)) - P1;
+				FpIntegral = P1 * LambdaOFL / P2 * FMath::Exp(P2 * (L / LambdaOFL - (FSolverReal)1.)) - P1 * (L - LambdaOFL);
+				dFpdL = P1 * FMath::Exp(P2 * (L / LambdaOFL - (FSolverReal)1.)) - P1;
 			}
-			if (L > (T)0.4 * LambdaOFL && L < (T)0.6 * LambdaOFL)
+			if (L > (FSolverReal)0.4 * LambdaOFL && L < (FSolverReal)0.6 * LambdaOFL)
 			{
-				FaIntegral = (T)3. * LambdaOFL * FMath::Pow((L / LambdaOFL - (T)0.4), 3);
-				dFadL = (T)9. * FMath::Pow((L / LambdaOFL - (T)0.4), 2);
+				FaIntegral = (FSolverReal)3. * LambdaOFL * FMath::Pow((L / LambdaOFL - (FSolverReal)0.4), 3);
+				dFadL = (FSolverReal)9. * FMath::Pow((L / LambdaOFL - (FSolverReal)0.4), 2);
 			}
-			else if (L >= (T)0.6 * LambdaOFL && L <= (T)1.4 * LambdaOFL)
+			else if (L >= (FSolverReal)0.6 * LambdaOFL && L <= (FSolverReal)1.4 * LambdaOFL)
 			{
-				FaIntegral = (T)3. * LambdaOFL * (T)0.008 + L - (T)4. / (T)3. * LambdaOFL * FMath::Pow(L / LambdaOFL - (T)1., 3) - (T)0.6 * LambdaOFL - (T)4. / (T)3. * LambdaOFL * FMath::Pow((T)0.4, 3);
-				dFadL = (T)1. - (T)4. * FMath::Pow(L / LambdaOFL - (T)1., 2);
+				FaIntegral = (FSolverReal)3. * LambdaOFL * (FSolverReal)0.008 + L - (FSolverReal)4. / (FSolverReal)3. * LambdaOFL * FMath::Pow(L / LambdaOFL - (FSolverReal)1., 3) - (FSolverReal)0.6 * LambdaOFL - (FSolverReal)4. / (FSolverReal)3. * LambdaOFL * FMath::Pow((FSolverReal)0.4, 3);
+				dFadL = (FSolverReal)1. - (FSolverReal)4. * FMath::Pow(L / LambdaOFL - (FSolverReal)1., 2);
 			}
-			else if (L > (T)1.4 * LambdaOFL && L <= (T)1.6 * LambdaOFL)
+			else if (L > (FSolverReal)1.4 * LambdaOFL && L <= (FSolverReal)1.6 * LambdaOFL)
 			{
-				FaIntegral = (T)3. * LambdaOFL * (T)0.008 + (T)0.8 * LambdaOFL - (T)8. / (T)3. * LambdaOFL * FMath::Pow((T)0.4, 3) + (T)3. * LambdaOFL * FMath::Pow(L / LambdaOFL - (T)1.6, 3) + (T)3. * LambdaOFL * (T)0.008;
-				dFadL = (T)9. * FMath::Pow((L / LambdaOFL - (T)1.6), 2);
+				FaIntegral = (FSolverReal)3. * LambdaOFL * (FSolverReal)0.008 + (FSolverReal)0.8 * LambdaOFL - (FSolverReal)8. / (FSolverReal)3. * LambdaOFL * FMath::Pow((FSolverReal)0.4, 3) + (FSolverReal)3. * LambdaOFL * FMath::Pow(L / LambdaOFL - (FSolverReal)1.6, 3) + (FSolverReal)3. * LambdaOFL * (FSolverReal)0.008;
+				dFadL = (FSolverReal)9. * FMath::Pow((L / LambdaOFL - (FSolverReal)1.6), 2);
 
 			}
-			else if (L > (T)1.6 * LambdaOFL)
+			else if (L > (FSolverReal)1.6 * LambdaOFL)
 			{
-				FaIntegral = (T)3. * LambdaOFL * (T)0.008 + (T)0.8 * LambdaOFL - (T)8. / (T)3. * LambdaOFL * FMath::Pow((T)0.4, 3) + (T)3. * LambdaOFL * (T)0.008;
-				dFadL = (T)0.;
+				FaIntegral = (FSolverReal)3. * LambdaOFL * (FSolverReal)0.008 + (FSolverReal)0.8 * LambdaOFL - (FSolverReal)8. / (FSolverReal)3. * LambdaOFL * FMath::Pow((FSolverReal)0.4, 3) + (FSolverReal)3. * LambdaOFL * (FSolverReal)0.008;
+				dFadL = (FSolverReal)0.;
 			}
 
 
 			C3 = FMath::Sqrt(FpIntegral + AlphaActivation * FaIntegral);
 
 
-			//T dC3dL = (T)0.5 * (dFpdL + AlphaActivation * dFadL);
+			//FSolverReal dC3dL = (FSolverReal)0.5 * (dFpdL + AlphaActivation * dFadL);
 			//if (C3 == 0)
 			//{
-			//	return TVec4<TVector<T, 3>>(TVector<T, 3>((T)0.));
+			//	return TVec4<TVector<FSolverReal, 3>>(TVector<FSolverReal, 3>((FSolverReal)0.));
 			//}
 			//else
 			//{
 			//	dC3dL /= C3;
 			//}
 
-			TVec4<TVector<T, 3>> dC3 = GetFiberGradient(dFpdL,dFadL, C3, dLdX);
+			TVec4<TVector<FSolverReal, 3>> dC3 = GetFiberGradient(dFpdL,dFadL, C3, dLdX);
 
-			//TVec4<TVector<T, 3>> dC3(TVector<T, 3>((T)0.));
+			//TVec4<TVector<FSolverReal, 3>> dC3(TVector<FSolverReal, 3>((FSolverReal)0.));
 			//for (int32 i = 0; i < 4; i++)
 			//{
 			//	for (int32 j = 0; j < 3; j++)
@@ -200,9 +202,9 @@ namespace Chaos::Softs
 			//	}
 			//}
 
-			T DLambda = -C3 - AlphaTilde * LambdaArray[2 * ElementIndex + 2];
+			FSolverReal DLambda = -C3 - AlphaTilde * LambdaArray[2 * ElementIndex + 2];
 
-			T Denom = AlphaTilde;
+			FSolverReal Denom = AlphaTilde;
 			for (int i = 0; i < 4; i++)
 			{
 				for (int j = 0; j < 3; j++)
@@ -212,7 +214,7 @@ namespace Chaos::Softs
 			}
 			DLambda /= Denom;
 			LambdaArray[2 * ElementIndex + 2] += DLambda;
-			TVec4<TVector<T, 3>> Delta(TVector<T, 3>((T)0.));
+			TVec4<TVector<FSolverReal, 3>> Delta(TVector<FSolverReal, 3>((FSolverReal)0.));
 			for (int i = 0; i < 4; i++)
 			{
 				for (int j = 0; j < 3; j++)
@@ -229,11 +231,14 @@ namespace Chaos::Softs
 	private:
 
 		//material constants calculated from E:
-		T SigmaMax;
-		mutable T AlphaActivation;
-		TVector<T, 3> FiberDir;
+		FSolverReal SigmaMax;
+		mutable FSolverReal AlphaActivation;
+		TVector<FSolverReal, 3> FiberDir;
 
 	};
 
+	template <typename T, typename ParticleType>
+	using FXPBDCorotatedFiberConstraints UE_DEPRECATED(5.8, "Deprecated. this class is to be deleted, use FXPBDCorotatedFiberTetrahedralConstraints instead") = FXPBDCorotatedFiberTetrahedralConstraints;
+
 }  // End namespace Chaos::Softs
 
diff --git a/Engine/Source/Runtime/Experimental/Chaos/Public/Chaos/XPBDGridBasedCorotatedConstraints.h b/Engine/Source/Runtime/Experimental/Chaos/Public/Chaos/XPBDGridBasedCorotatedConstraints.h
--- a/Engine/Source/Runtime/Experimental/Chaos/Public/Chaos/XPBDGridBasedCorotatedConstraints.h
+++ b/Engine/Source/Runtime/Experimental/Chaos/Public/Chaos/XPBDGridBasedCorotatedConstraints.h
@@ -18,11 +18,10 @@
 
 namespace Chaos::Softs
 {
-	template <typename T, typename ParticleType>
-	class FXPBDGridBasedCorotatedConstraints : public FXPBDCorotatedConstraints <T, ParticleType>
+	class FXPBDGridBasedCorotatedTetrahedralConstraints : public FXPBDCorotatedTetrahedralConstraints
 	{
 
-		typedef FXPBDCorotatedConstraints<T, ParticleType> Base;
+		typedef FXPBDCorotatedTetrahedralConstraints Base;
 		using Base::MeshConstraints;
 		using Base::LambdaArray;
 		using Base::Measure;
@@ -33,40 +32,40 @@ namespace Chaos::Softs
 
 	public:
 		//this one only accepts tetmesh input and mesh
-		FXPBDGridBasedCorotatedConstraints(
-			const ParticleType& InParticles,
+		FXPBDGridBasedCorotatedTetrahedralConstraints(
+			const FSolverParticles& InParticles,
 			const TArray<TVector<int32, 4>>& InMesh,
-			const T GridDx = (T).1,
+			const FSolverReal GridDx = (FSolverReal).1,
 			const bool bRecordMetricIn = true,
-			const T MaxDtIn = (T).1,
-			const T MinDtIn = (T).01,
-			const T CFLCoeffIn = (T).4,
-			const T& EMesh = (T)1000.0,
-			const T& NuMesh = (T).3
+			const FSolverReal MaxDtIn = (FSolverReal).1,
+			const FSolverReal MinDtIn = (FSolverReal).01,
+			const FSolverReal CFLCoeffIn = (FSolverReal).4,
+			const FSolverReal& EMesh = (FSolverReal)1000.0,
+			const FSolverReal& NuMesh = (FSolverReal).3
 		)
 			: Base(InParticles, InMesh, GridDx, EMesh, NuMesh), MaxDt(MaxDtIn), MinDt(MinDtIn), CFLCoeff(CFLCoeffIn), MPMGrid(GridDx), MPMTransfer(MPMGrid, InParticles.Size()), PreviousColoringPtr(nullptr)
 		{
-			//TMPMGrid<T> MPMGrid(GridN);
+			//TMPMGrid<FSolverReal> MPMGrid(GridN);
 			MPMGrid.UpdateGridFromPositions(InParticles);
-			//TMPMTransfer<T> MPMTrasnfer(MPMGrid);
-			LambdaArray.Init((T)0., 2 * MeshConstraints.Num());
-			DmInverse.Init((T)0., 9 * MeshConstraints.Num());
-			Measure.Init((T)0., MeshConstraints.Num());
-			Lambda = EMesh * NuMesh / (((T)1. + NuMesh) * ((T)1. - (T)2. * NuMesh));
-			Mu = EMesh / ((T)2. * ((T)1. + NuMesh));
+			//TMPMTransfer<FSolverReal> MPMTrasnfer(MPMGrid);
+			LambdaArray.Init((FSolverReal)0., 2 * MeshConstraints.Num());
+			DmInverse.Init((FSolverReal)0., 9 * MeshConstraints.Num());
+			Measure.Init((FSolverReal)0., MeshConstraints.Num());
+			Lambda = EMesh * NuMesh / (((FSolverReal)1. + NuMesh) * ((FSolverReal)1. - (FSolverReal)2. * NuMesh));
+			Mu = EMesh / ((FSolverReal)2. * ((FSolverReal)1. + NuMesh));
 			for (int e = 0; e < InMesh.Num(); e++)
 			{
-				PMatrix<T, 3, 3> Dm = Base::DsInit(e, InParticles);
-				PMatrix<T, 3, 3> DmInv = Dm.Inverse();
+				PMatrix<FSolverReal, 3, 3> Dm = Base::DsInit(e, InParticles);
+				PMatrix<FSolverReal, 3, 3> DmInv = Dm.Inverse();
 				for (int r = 0; r < 3; r++) {
 					for (int c = 0; c < 3; c++) {
 						DmInverse[(3 * 3) * e + 3 * r + c] = DmInv.GetAt(r, c);
 					}
 				}
 
-				Measure[e] = Dm.Determinant() / (T)6.;
+				Measure[e] = Dm.Determinant() / (FSolverReal)6.;
 
-				if (Measure[e] < (T)0.)
+				if (Measure[e] < (FSolverReal)0.)
 				{
 					Measure[e] = -Measure[e];
 				}
@@ -74,7 +73,7 @@ namespace Chaos::Softs
 			}
 
 			ElementsPerColor = FGraphColoring::ComputeGraphColoringAllDynamic(MeshConstraints, InParticles);
-			Init(InParticles, (T)0.);
+			Init(InParticles, (FSolverReal)0.);
 			WriteGridNodes();
 		}
 
@@ -84,13 +83,13 @@ namespace Chaos::Softs
 			//file.Append(TEXT("\HoudiniOuput\Test.geo"));
 			file.Append(TEXT("/HoudiniOutput/GridData.geo"));
 
-			TArray<TVector<T, 3>> ActivatedGridNodes;
+			TArray<TVector<FSolverReal, 3>> ActivatedGridNodes;
 
 			for (int32 i = 0; i < MPMGrid.Size(); i++)
 			{
-				if (GridData[i * (MPMTransfer.NTransfer+1)] > T(0.))
+				if (GridData[i * (MPMTransfer.NTransfer+1)] > FSolverReal(0.))
 				{
-					TVector<T, 3> GridNode = MPMGrid.Node(i);
+					TVector<FSolverReal, 3> GridNode = MPMGrid.Node(i);
 					ActivatedGridNodes.Emplace(GridNode);
 				}
 			}
@@ -119,9 +118,9 @@ namespace Chaos::Softs
 
 		}
 
-		void Init(const ParticleType& InParticles, const T Dt) 
+		void Init(const FSolverParticles& InParticles, const FSolverReal Dt) 
 		{
-			for (T& Lambdas : LambdaArray) { Lambdas = (T)0.; }
+			for (FSolverReal& Lambdas : LambdaArray) { Lambdas = (FSolverReal)0.; }
 			{
 				TRACE_CPUPROFILER_EVENT_SCOPE(STAT_ChaosXPBDGridBasedCorotatedUpdateGrid);
 				MPMTransfer.Grid.UpdateGridFromPositions(InParticles);
@@ -144,7 +143,7 @@ namespace Chaos::Softs
 			}
 		}
 
-		void TimeStepPostprocessing(ParticleType& InParticles, const T Dt)
+		void TimeStepPostprocessing(FSolverParticles& InParticles, const FSolverReal Dt)
 		{
 			{
 				TRACE_CPUPROFILER_EVENT_SCOPE(STAT_ChaosXPBDGridBasedCorotatedGrid::FinalG2P);
@@ -163,12 +162,12 @@ namespace Chaos::Softs
 		
 		}
 
-		T ComputeCFLDt(const ParticleType& InParticles)
+		FSolverReal ComputeCFLDt(const FSolverParticles& InParticles)
 		{
-			T VMax = (T)-1.;
-			for (int32 p = 0; p < InParticles.Size(); p++)
+			FSolverReal VMax = (FSolverReal)-1.;
+			for (int32 p = 0; p < int32(InParticles.Size()); p++)
 			{
-				T VMag = InParticles.V(p).Size();
+				FSolverReal VMag = InParticles.V(p).Size();
 				if (VMag > VMax)
 				{
 					VMax = VMag;
@@ -184,7 +183,7 @@ namespace Chaos::Softs
 				return MaxDt;
 			}
 
-			T CFLDt = MPMGrid.MDx[0] * CFLCoeff / VMax;
+			FSolverReal CFLDt = MPMGrid.Dx()[0] * CFLCoeff / VMax;
 
 			if (CFLDt > MaxDt)
 			{
@@ -203,33 +202,33 @@ namespace Chaos::Softs
 		const TArray<TArray<int32>> GetElementsPerColor () const { return ElementsPerColor; }
 		const TArray<TArray<int32>> GetPreviousColoring () const { return *PreviousColoringPtr; }
 		const TArray<TArray<TArray<int32>>> GetElementsPerSubColors() const { return ElementsPerSubColors; }
-		const TMPMTransfer<T> GetMPMTransfer() const { return MPMTransfer; }
+		const TMPMTransfer<FSolverReal> GetMPMTransfer() const { return MPMTransfer; }
 
-		void ApplyPolar(ParticleType& Particles, const T Dt, const int32 ElementIndex)  
+		void ApplyPolar(FSolverParticles& Particles, const FSolverReal Dt, const int32 ElementIndex)  
 		{
 			//step 1: sparse g2p:
-			TVec4<TVector<T, 3>> XPosFromGrid = MPMTransfer.SparseG2P(GridPositions, ElementIndex);
+			TVec4<TVector<FSolverReal, 3>> XPosFromGrid = MPMTransfer.SparseG2P(GridPositions, ElementIndex);
 
 			//step 2: gradient and constraint computation:
-			T C1 = (T)0.;
-			TVec4<TVector<T, 3>> PolarGradient = GetPolarGradient(XPosFromGrid, ElementIndex, C1);
+			FSolverReal C1 = (FSolverReal)0.;
+			TVec4<TVector<FSolverReal, 3>> PolarGradient = GetPolarGradient(XPosFromGrid, ElementIndex, C1);
 			
-			if (C1 == (T)0.)
+			if (C1 == (FSolverReal)0.)
 			{
 				return;
 			}
 
 			//step 3 : sparse p2g:
-			TArray<TVector<T, 3>> GridGradient = MPMTransfer.SparseP2G(PolarGradient, ElementIndex);
+			TArray<TVector<FSolverReal, 3>> GridGradient = MPMTransfer.SparseP2G(PolarGradient, ElementIndex);
 
 			//step 4: computation for delta:
-			T AlphaTilde = (T)1. / (Dt * Dt * Mu * Measure[ElementIndex]);
+			FSolverReal AlphaTilde = (FSolverReal)1. / (Dt * Dt * Mu * Measure[ElementIndex]);
 
-			T DeltaLambda = -C1 - AlphaTilde * LambdaArray[2 * ElementIndex + 0];
-			T Denom = AlphaTilde;
+			FSolverReal DeltaLambda = -C1 - AlphaTilde * LambdaArray[2 * ElementIndex + 0];
+			FSolverReal Denom = AlphaTilde;
 
 			for (int32 k = 0; k < MPMTransfer.ElementGridNodeIncidentElements[ElementIndex].Num(); k++) {
-				T Mass = GridData[(MPMTransfer.NTransfer + 1) * MPMTransfer.ElementGridNodes[ElementIndex][MPMTransfer.ElementGridNodeIncidentElements[ElementIndex][k][0]]];
+				FSolverReal Mass = GridData[(MPMTransfer.NTransfer + 1) * MPMTransfer.ElementGridNodes[ElementIndex][MPMTransfer.ElementGridNodeIncidentElements[ElementIndex][k][0]]];
 				if (Mass != 0) {
 					for (int32 alpha = 0; alpha < 3; alpha++) {
 						Denom += GridGradient[k][alpha] * GridGradient[k][alpha] / Mass;
@@ -242,7 +241,7 @@ namespace Chaos::Softs
 
 			for (int32 k = 0; k < MPMTransfer.ElementGridNodeIncidentElements[ElementIndex].Num(); k++) {
 				int32 CurrentNode = MPMTransfer.ElementGridNodes[ElementIndex][MPMTransfer.ElementGridNodeIncidentElements[ElementIndex][k][0]];
-				T Mass = GridData[(MPMTransfer.NTransfer + 1) * CurrentNode];
+				FSolverReal Mass = GridData[(MPMTransfer.NTransfer + 1) * CurrentNode];
 				if (Mass != 0) {
 					for (int32 alpha = 0; alpha < 3; alpha++) {
 						GridPositions[CurrentNode][alpha] += DeltaLambda * GridGradient[k][alpha] / Mass;
@@ -251,33 +250,33 @@ namespace Chaos::Softs
 			}
 		}
 
-		void ApplyDet(ParticleType& Particles, const T Dt, const int32 ElementIndex) 
+		void ApplyDet(FSolverParticles& Particles, const FSolverReal Dt, const int32 ElementIndex) 
 		{
-			//TArray<TVector<T, 3>> XPosFromGrid;
+			//TArray<TVector<FSolverReal, 3>> XPosFromGrid;
 			//XPosFromGrid.SetNum(4);
 			//step 1: sparse g2p:
-			TVec4<TVector<T, 3>> XPosFromGrid = MPMTransfer.SparseG2P(GridPositions, ElementIndex);
+			TVec4<TVector<FSolverReal, 3>> XPosFromGrid = MPMTransfer.SparseG2P(GridPositions, ElementIndex);
 
 			//step 2: gradient and constraint computation:
-			T C2 = (T)0.;
-			TVec4<TVector<T, 3>> DetGradient = GetDetGradient(XPosFromGrid, ElementIndex, C2);
+			FSolverReal C2 = (FSolverReal)0.;
+			TVec4<TVector<FSolverReal, 3>> DetGradient = GetDetGradient(XPosFromGrid, ElementIndex, C2);
 
-			if (C2 == (T)0.)
+			if (C2 == (FSolverReal)0.)
 			{
 				return;
 			}
 
 			//step 3 : sparse p2g:
-			TArray<TVector<T, 3>> GridGradient = MPMTransfer.SparseP2G(DetGradient, ElementIndex);
+			TArray<TVector<FSolverReal, 3>> GridGradient = MPMTransfer.SparseP2G(DetGradient, ElementIndex);
 
 			//step 4: computation for delta:
-			T AlphaTilde = (T)2. / (Dt * Dt * Lambda * Measure[ElementIndex]);
+			FSolverReal AlphaTilde = (FSolverReal)2. / (Dt * Dt * Lambda * Measure[ElementIndex]);
 
-			T DeltaLambda = -C2 - AlphaTilde * LambdaArray[2 * ElementIndex + 1];
-			T Denom = AlphaTilde;
+			FSolverReal DeltaLambda = -C2 - AlphaTilde * LambdaArray[2 * ElementIndex + 1];
+			FSolverReal Denom = AlphaTilde;
 
 			for (int32 k = 0; k < MPMTransfer.ElementGridNodeIncidentElements[ElementIndex].Num(); k++) {
-				T Mass = GridData[(MPMTransfer.NTransfer + 1) * MPMTransfer.ElementGridNodes[ElementIndex][MPMTransfer.ElementGridNodeIncidentElements[ElementIndex][k][0]]];
+				FSolverReal Mass = GridData[(MPMTransfer.NTransfer + 1) * MPMTransfer.ElementGridNodes[ElementIndex][MPMTransfer.ElementGridNodeIncidentElements[ElementIndex][k][0]]];
 				if (Mass != 0) {
 					for (int32 alpha = 0; alpha < 3; alpha++) {
 						Denom += GridGradient[k][alpha] * GridGradient[k][alpha] / Mass;
@@ -290,7 +289,7 @@ namespace Chaos::Softs
 
 			for (int32 k = 0; k < MPMTransfer.ElementGridNodeIncidentElements[ElementIndex].Num(); k++) {
 				int32 CurrentNode = MPMTransfer.ElementGridNodes[ElementIndex][MPMTransfer.ElementGridNodeIncidentElements[ElementIndex][k][0]];
-				T Mass = GridData[(MPMTransfer.NTransfer + 1) * CurrentNode];
+				FSolverReal Mass = GridData[(MPMTransfer.NTransfer + 1) * CurrentNode];
 				if (Mass != 0) {
 					for (int32 alpha = 0; alpha < 3; alpha++) {
 						GridPositions[CurrentNode][alpha] += DeltaLambda * GridGradient[k][alpha] / Mass;
@@ -299,7 +298,7 @@ namespace Chaos::Softs
 			}
 		}
 
-		void ApplyInParallel(ParticleType& Particles, const T Dt) 
+		void ApplyInParallel(FSolverParticles& Particles, const FSolverReal Dt) 
 		{
 			for (int32 color = 0; color < ElementsPerSubColors.Num(); color++)
 			{
@@ -319,8 +318,8 @@ namespace Chaos::Softs
 
 	private:
 
-		PMatrix<T, 3, 3> Ds(const TVec4<TVector<T, 3>>& LocalPositions) const {
-			PMatrix<T, 3, 3> Result((T)0.);
+		PMatrix<FSolverReal, 3, 3> Ds(const TVec4<TVector<FSolverReal, 3>>& LocalPositions) const {
+			PMatrix<FSolverReal, 3, 3> Result((FSolverReal)0.);
 			for (int i = 0; i < 3; i++) {
 				for (int c = 0; c < 3; c++) {
 					Result.SetAt(c, i, LocalPositions[i+1][c] - LocalPositions[0][c]);
@@ -330,19 +329,19 @@ namespace Chaos::Softs
 		}
 
 
-		PMatrix<T, 3, 3> F(const TVec4<TVector<T, 3>>& LocalPositions, const int32 ElementIndex) const 
+		PMatrix<FSolverReal, 3, 3> F(const TVec4<TVector<FSolverReal, 3>>& LocalPositions, const int32 ElementIndex) const 
 		{
 			return Base::ElementDmInv(ElementIndex) * Ds(LocalPositions);
 		}
 
-		TVec4<TVector<T, 3>> GetPolarGradient(const TVec4<TVector<T, 3>>& LocalPositions, const int32 ElementIndex, T& Constraint, const T Tol = T(1e-4)) const
+		TVec4<TVector<FSolverReal, 3>> GetPolarGradient(const TVec4<TVector<FSolverReal, 3>>& LocalPositions, const int32 ElementIndex, FSolverReal& Constraint, const FSolverReal Tol = FSolverReal(1e-4)) const
 		{
-			PMatrix<T, 3, 3> Fe = F(LocalPositions, ElementIndex);
-			PMatrix<T, 3, 3> Re((T)0.), Se((T)0.);
+			PMatrix<FSolverReal, 3, 3> Fe = F(LocalPositions, ElementIndex);
+			PMatrix<FSolverReal, 3, 3> Re((FSolverReal)0.), Se((FSolverReal)0.);
 
 			Chaos::PolarDecomposition(Fe, Re, Se);
 
-			T C1 = (T)0.;
+			FSolverReal C1 = (FSolverReal)0.;
 			for (int i = 0; i < 3; i++)
 			{
 				for (int j = 0; j < 3; j++)
@@ -355,15 +354,15 @@ namespace Chaos::Softs
 			
 			if (C1 < Tol)
 			{
-				Constraint = (T)0.;
-				return TVec4<TVector<T, 3>>(TVector<T, 3>((T)0.));
+				Constraint = (FSolverReal)0.;
+				return TVec4<TVector<FSolverReal, 3>>(TVector<FSolverReal, 3>((FSolverReal)0.));
 			}
 			Constraint = C1;
 
-			PMatrix<T, 3, 3> DmInvT = Base::ElementDmInv(ElementIndex).GetTransposed();
-			TVector<T, 81> dRdF((T)0.);
+			PMatrix<FSolverReal, 3, 3> DmInvT = Base::ElementDmInv(ElementIndex).GetTransposed();
+			TVector<FSolverReal, 81> dRdF((FSolverReal)0.);
 			Chaos::dRdFCorotated(Fe, dRdF);
-			TVec4<TVector<T, 3>> dC1(TVector<T, 3>((T)0.));
+			TVec4<TVector<FSolverReal, 3>> dC1(TVector<FSolverReal, 3>((FSolverReal)0.));
 			//dC1 = dC1dF * dFdX
 			for (int alpha = 0; alpha < 3; alpha++) {
 				for (int l = 0; l < 3; l++) {
@@ -414,22 +413,22 @@ namespace Chaos::Softs
 			return dC1;
 		}
 
-		TVec4<TVector<T, 3>> GetDetGradient(const TVec4<TVector<T, 3>>& LocalPositions, const int32 ElementIndex, T& Constraint, const T Tol = T(1e-4)) const
+		TVec4<TVector<FSolverReal, 3>> GetDetGradient(const TVec4<TVector<FSolverReal, 3>>& LocalPositions, const int32 ElementIndex, FSolverReal& Constraint, const FSolverReal Tol = FSolverReal(1e-4)) const
 		{
-			const PMatrix<T, 3, 3> Fe = F(LocalPositions, ElementIndex);
-			PMatrix<T, 3, 3> DmInvT = Base::ElementDmInv(ElementIndex).GetTransposed();
+			const PMatrix<FSolverReal, 3, 3> Fe = F(LocalPositions, ElementIndex);
+			PMatrix<FSolverReal, 3, 3> DmInvT = Base::ElementDmInv(ElementIndex).GetTransposed();
 
-			T J = Fe.Determinant();
+			FSolverReal J = Fe.Determinant();
 			if (J - 1 < Tol)
 			{
-				Constraint = (T)0.;
-				return TVec4<TVector<T, 3>>(TVector<T, 3>((T)0.));
+				Constraint = (FSolverReal)0.;
+				return TVec4<TVector<FSolverReal, 3>>(TVector<FSolverReal, 3>((FSolverReal)0.));
 			}
 			Constraint = J - 1;
 
-			TVec4<TVector<T, 3>> dC2(TVector<T, 3>((T)0.));
+			TVec4<TVector<FSolverReal, 3>> dC2(TVector<FSolverReal, 3>((FSolverReal)0.));
 
-			PMatrix<T, 3, 3> JFinvT((T)0.);
+			PMatrix<FSolverReal, 3, 3> JFinvT((FSolverReal)0.);
 			JFinvT.SetAt(0, 0, Fe.GetAt(1, 1) * Fe.GetAt(2, 2) - Fe.GetAt(2, 1) * Fe.GetAt(1, 2));
 			JFinvT.SetAt(0, 1, Fe.GetAt(2, 0) * Fe.GetAt(1, 2) - Fe.GetAt(1, 0) * Fe.GetAt(2, 2));
 			JFinvT.SetAt(0, 2, Fe.GetAt(1, 0) * Fe.GetAt(2, 1) - Fe.GetAt(2, 0) * Fe.GetAt(1, 1));
@@ -440,7 +439,7 @@ namespace Chaos::Softs
 			JFinvT.SetAt(2, 1, Fe.GetAt(1, 0) * Fe.GetAt(0, 2) - Fe.GetAt(0, 0) * Fe.GetAt(1, 2));
 			JFinvT.SetAt(2, 2, Fe.GetAt(0, 0) * Fe.GetAt(1, 1) - Fe.GetAt(1, 0) * Fe.GetAt(0, 1));
 
-			PMatrix<T, 3, 3> JinvTDmInvT = DmInvT * JFinvT;
+			PMatrix<FSolverReal, 3, 3> JinvTDmInvT = DmInvT * JFinvT;
 
 			for (int32 ie = 0; ie < 3; ie++) {
 				for (int32 alpha = 0; alpha < 3; alpha++) {
@@ -456,20 +455,20 @@ namespace Chaos::Softs
 		}
 
 		//interpolate particle positions from grid positions
-		void FinalG2P(ParticleType& Particles)
+		void FinalG2P(FSolverParticles& Particles)
 		{
 			PhysicsParallelFor(Particles.Size(), [&](const int32 p)
 				{
-					Particles.P(p) = TVec3<T>((T)0.);
+					Particles.P(p) = TVec3<FSolverReal>((FSolverReal)0.);
 					for (int32 ii = 0; ii < int32(MPMGrid.NPerDir); ii++)
 					{
-						T Nii = MPMGrid.Nijk(MPMTransfer.Weights[p][0], ii);
+						FSolverReal Nii = MPMGrid.Nijk(MPMTransfer.Weights[p][0], ii);
 						for (int32 jj = 0; jj < int32(MPMGrid.NPerDir); jj++)
 						{
-							T Njj = MPMGrid.Nijk(MPMTransfer.Weights[p][1], jj);
+							FSolverReal Njj = MPMGrid.Nijk(MPMTransfer.Weights[p][1], jj);
 							for (int32 kk = 0; kk < int32(MPMGrid.NPerDir); kk++)
 							{
-								T Nkk = MPMGrid.Nijk(MPMTransfer.Weights[p][2], kk);
+								FSolverReal Nkk = MPMGrid.Nijk(MPMTransfer.Weights[p][2], kk);
 								TVector<int32, 3> LocIndex = { ii, jj, kk };
 								TVector<int32, 3> GlobIndex = MPMGrid.Loc2GlobIndex(MPMTransfer.Indices[p], LocIndex);
 								int32 GlobIndexFlat = MPMGrid.FlatIndex(GlobIndex);
@@ -482,20 +481,22 @@ namespace Chaos::Softs
 		}
 
 	protected:
-		T MaxDt;
-		T MinDt;
-		T CFLCoeff;
-		TMPMGrid<T> MPMGrid;
-		TMPMTransfer<T> MPMTransfer;
+		FSolverReal MaxDt;
+		FSolverReal MinDt;
+		FSolverReal CFLCoeff;
+		TMPMGrid<FSolverReal> MPMGrid;
+		TMPMTransfer<FSolverReal> MPMTransfer;
 		TArray<TArray<int32>> ElementsPerColor;
 		UE_DEPRECATED(5.6, "Use PreviousColoringPtr instead")
 		TArray<TArray<int32>>* PreviousColoring = nullptr;
 		TUniquePtr<TArray<TArray<int32>>> PreviousColoringPtr;
 		TArray<TArray<TArray<int32>>> ElementsPerSubColors;
-		TArray<T> GridData;
-		TArray<TVector<T, 3>> GridPositions;
+		TArray<FSolverReal> GridData;
+		TArray<TVector<FSolverReal, 3>> GridPositions;
 		bool InitialGridDataWritten = false;
 	};
 
+	template <typename T, typename ParticleType>
+	using FXPBDGridBasedCorotatedConstraints UE_DEPRECATED(5.8, "Deprecated. this class is to be deleted, use FXPBDGridBasedCorotatedTetrahedralConstraints instead") = FXPBDGridBasedCorotatedTetrahedralConstraints;
 
 }
\ No newline at end of file
diff --git a/Engine/Source/Runtime/Experimental/Chaos/Public/Chaos/XPBDWeakConstraints.h b/Engine/Source/Runtime/Experimental/Chaos/Public/Chaos/XPBDWeakConstraints.h
--- a/Engine/Source/Runtime/Experimental/Chaos/Public/Chaos/XPBDWeakConstraints.h
+++ b/Engine/Source/Runtime/Experimental/Chaos/Public/Chaos/XPBDWeakConstraints.h
@@ -21,7 +21,7 @@
 namespace Chaos::Softs
 {
 
-	struct FDeformableXPBDWeakConstraintParams
+	struct FDeformableXPBDWeightedSpringConstraintParams
 	{
 		float DebugLineWidth = 5.f;
 		float DebugParticleWidth = 20.f;
@@ -30,34 +30,32 @@ namespace Chaos::Softs
 
 	using Chaos::TVec3;
 
-	template <typename T, typename ParticleType>
-	class FXPBDWeakConstraints
+	// Builds weighted spring constraints between two arbitrary groups of vertices
+	class FXPBDWeightedSpringConstraints
 	{
 
 	public:
-		//this one only accepts tetmesh input and mesh
-		FXPBDWeakConstraints(
-			const ParticleType& InParticles,
+		// Used for spring constraints without self targets (against a fixed constraint target instead of deformable vertices)
+		FXPBDWeightedSpringConstraints(
+			const FSolverParticles& InParticles,
 			const TArray<TArray<int32>>& InIndices,
-			const TArray<TArray<T>>& InWeights,
-			const TArray<T>& InStiffness
+			const TArray<TArray<FSolverReal>>& InWeights,
+			const TArray<FSolverReal>& InStiffness
 		)
-			: Stiffness(InStiffness), Indices(InIndices), Weights(InWeights)
+			: Indices(InIndices), Weights(InWeights), Stiffness(InStiffness)
 		{
-			SecondIndices.SetNum(0);
-			SecondWeights.SetNum(0);
 			InitColor(InParticles);
-			LambdaArray.Init((T)0., Indices.Num() * 3);
+			LambdaArray.Init((FSolverReal)0., Indices.Num() * 3);
 		}
 
-		FXPBDWeakConstraints(
-			const ParticleType& InParticles,
+		FXPBDWeightedSpringConstraints(
+			const FSolverParticles& InParticles,
 			const TArray<TArray<int32>>& InIndices,
-			const TArray<TArray<T>>& InWeights,
-			const TArray<T>& InStiffness,
+			const TArray<TArray<FSolverReal>>& InWeights,
+			const TArray<FSolverReal>& InStiffness,
 			const TArray<TArray<int32>>& InSecondIndices,
-			const TArray<TArray<T>>& InSecondWeights,
-			const FDeformableXPBDWeakConstraintParams& InParams
+			const TArray<TArray<FSolverReal>>& InSecondWeights,
+			const FDeformableXPBDWeightedSpringConstraintParams& InParams
 		)
 			: Indices(InIndices), Weights(InWeights), SecondIndices(InSecondIndices), SecondWeights(InSecondWeights), Stiffness(InStiffness), DebugDrawParams(InParams)
 		{
@@ -72,14 +70,13 @@ namespace Chaos::Softs
 				}
 			}
 			InitColor(InParticles);
-			LambdaArray.Init((T)0., Indices.Num() * 3);
+			LambdaArray.Init((FSolverReal)0., Indices.Num() * 3);
 		}
 
 
-		virtual ~FXPBDWeakConstraints() {}
+		virtual ~FXPBDWeightedSpringConstraints() {}
 
-		template<typename SolverParticlesOrRange>
-		void ApplyInParallel(SolverParticlesOrRange& Particles, const T Dt) const
+		void ApplyInParallel(FSolverParticlesRange& Particles, const FSolverReal Dt) const
 		{
 			TRACE_CPUPROFILER_EVENT_SCOPE(STAT_ChaosXPBDWeakConstraintApply);
 
@@ -110,23 +107,32 @@ namespace Chaos::Softs
 			}
 		}
 
+		void ApplyInParallel(FSolverParticles& Particles, const FSolverReal Dt) const
+		{
+			FSolverParticlesRange Range(&Particles, /*InOffset = */ 0, /*InRange = */ Particles.Size());
+			ApplyInParallel(Range, Dt);
+		}
 
 		const TArray<TArray<int32>>& GetIndices()
 		{
 			return Indices;
 		}
 
-		template<typename SolverParticlesOrRange>
-		void Init(const SolverParticlesOrRange& InParticles, const T Dt) const
+		void Init(const FSolverParticlesRange& InParticles, const FSolverReal Dt) const
 		{
-			for (T& Lambdas : LambdaArray) { Lambdas = (T)0.; }
+			for (FSolverReal& Lambdas : LambdaArray) { Lambdas = (FSolverReal)0.; }
 			if (DebugDrawParams.bVisualizeBindings)
 			{
 				VisualizeAllBindings(InParticles, Dt);
 			}
 		}
 
-		void UpdateTargets(TArray<TVector<T, 3>>&& InTargets)
+		void Init(const FSolverParticles& InParticles, const FSolverReal Dt) const
+		{
+			Init(FSolverParticlesRange(const_cast<FSolverParticles*>(&InParticles), /*InOffset = */ 0, /*InRange = */ InParticles.Size()), Dt);
+		}
+
+		void UpdateTargets(TArray<TVector<FSolverReal, 3>>&& InTargets)
 		{
 			Constraints = MoveTemp(InTargets);
 		}
@@ -134,27 +140,26 @@ namespace Chaos::Softs
 
 	private:
 
-		void InitColor(const ParticleType& Particles)
+		void InitColor(const FSolverParticles& Particles)
 		{
 			Chaos::ComputeWeakConstraintsColoring(Indices, SecondIndices, Particles, ConstraintsPerColor);
 		}
 
-		template<typename SolverParticlesOrRange>
-		void ApplySingleConstraintWithoutSelfTarget(SolverParticlesOrRange& Particles, const T Dt, const int32 ConstraintIndex, const T Tol = UE_KINDA_SMALL_NUMBER) const
+		void ApplySingleConstraintWithoutSelfTarget(FSolverParticlesRange& Particles, const FSolverReal Dt, const int32 ConstraintIndex, const FSolverReal Tol = UE_KINDA_SMALL_NUMBER) const
 		{
 			if (Stiffness[ConstraintIndex] < Tol)
 			{
 				return;
 			}
 
-			T AlphaTilde = T(2) / (Stiffness[ConstraintIndex] * Dt * Dt);
+			FSolverReal AlphaTilde = FSolverReal(2) / (Stiffness[ConstraintIndex] * Dt * Dt);
 
 			if (Stiffness[ConstraintIndex] > 1e14)
 			{
-				AlphaTilde = T(0);
+				AlphaTilde = FSolverReal(0);
 			}
 
-			TVec3<T> SpringEdge((T)0.);
+			TVec3<FSolverReal> SpringEdge((FSolverReal)0.);
 			for (int32 j = 0; j < Indices[ConstraintIndex].Num(); j++) {
 				for (int32 alpha = 0; alpha < 3; alpha++) {
 					SpringEdge[alpha] += Weights[ConstraintIndex][j] * Particles.P(Indices[ConstraintIndex][j])[alpha];
@@ -163,7 +168,7 @@ namespace Chaos::Softs
 
 			SpringEdge -= Constraints[ConstraintIndex];
 
-			T Denom = AlphaTilde;
+			FSolverReal Denom = AlphaTilde;
 			for (int32 j = 0; j < Indices[ConstraintIndex].Num(); j++) {
 				Denom += Weights[ConstraintIndex][j] * Weights[ConstraintIndex][j] * Particles.InvM(Indices[ConstraintIndex][j]);
 				
@@ -171,8 +176,8 @@ namespace Chaos::Softs
 
 			for (int32 Beta = 0; Beta < 3; Beta++)
 			{
-				T Cj = SpringEdge[Beta];
-				T DLambda = -Cj - AlphaTilde * LambdaArray[ConstraintIndex * 3 + Beta];
+				FSolverReal Cj = SpringEdge[Beta];
+				FSolverReal DLambda = -Cj - AlphaTilde * LambdaArray[ConstraintIndex * 3 + Beta];
 
 				DLambda /= Denom;
 				LambdaArray[ConstraintIndex * 3 + Beta] += DLambda;
@@ -182,8 +187,7 @@ namespace Chaos::Softs
 			}
 		}
 
-		template<typename SolverParticlesOrRange>
-		void ApplySingleConstraintWithSelfTarget(SolverParticlesOrRange& Particles, const T Dt, const int32 ConstraintIndex, const T Tol = UE_KINDA_SMALL_NUMBER) const
+		void ApplySingleConstraintWithSelfTarget(FSolverParticlesRange& Particles, const FSolverReal Dt, const int32 ConstraintIndex, const FSolverReal Tol = UE_KINDA_SMALL_NUMBER) const
 		{
 			if (Stiffness[ConstraintIndex] < Tol)
 			{
@@ -192,14 +196,14 @@ namespace Chaos::Softs
 
 			ensure(SecondIndices.Num() > 0);
 
-			T AlphaTilde = T(2) / (Stiffness[ConstraintIndex] * Dt * Dt);
+			FSolverReal AlphaTilde = FSolverReal(2) / (Stiffness[ConstraintIndex] * Dt * Dt);
 
 			if (Stiffness[ConstraintIndex] > 1e14)
 			{
-				AlphaTilde = T(0);
+				AlphaTilde = FSolverReal(0);
 			}
 
-			TVec3<T> SpringEdge((T)0.);
+			TVec3<FSolverReal> SpringEdge((FSolverReal)0.);
 			for (int32 j = 0; j < Indices[ConstraintIndex].Num(); j++) {
 				for (int32 alpha = 0; alpha < 3; alpha++) {
 					SpringEdge[alpha] += Weights[ConstraintIndex][j] * Particles.P(Indices[ConstraintIndex][j])[alpha];
@@ -212,7 +216,7 @@ namespace Chaos::Softs
 
 			}
 
-			T Denom = AlphaTilde;
+			FSolverReal Denom = AlphaTilde;
 			for (int32 j = 0; j < Indices[ConstraintIndex].Num(); j++) {
 				Denom += Weights[ConstraintIndex][j] * Weights[ConstraintIndex][j] * Particles.InvM(Indices[ConstraintIndex][j]);
 			}
@@ -224,8 +228,8 @@ namespace Chaos::Softs
 
 			for (int32 Beta = 0; Beta < 3; Beta++)
 			{
-				T Cj = SpringEdge[Beta];
-				T DLambda = -Cj - AlphaTilde * LambdaArray[ConstraintIndex * 3 + Beta];
+				FSolverReal Cj = SpringEdge[Beta];
+				FSolverReal DLambda = -Cj - AlphaTilde * LambdaArray[ConstraintIndex * 3 + Beta];
 
 				DLambda /= Denom;
 				LambdaArray[ConstraintIndex * 3 + Beta] += DLambda;
@@ -239,14 +243,13 @@ namespace Chaos::Softs
 			}
 		}
 
-		template<typename SolverParticlesOrRange>
-		void VisualizeAllBindings(const SolverParticlesOrRange& InParticles, const T Dt) const
+		void VisualizeAllBindings(const FSolverParticlesRange& InParticles, const FSolverReal Dt) const
 		{
 #if WITH_EDITOR
-			auto DoubleVert = [](Chaos::TVec3<T> V) { return FVector3d(V.X, V.Y, V.Z); };
+			auto DoubleVert = [](Chaos::TVec3<FSolverReal> V) { return FVector3d(V.X, V.Y, V.Z); };
 			for (int32 i = 0; i < Indices.Num(); i++)
 			{
-				Chaos::TVec3<T> SourcePos((T)0.), TargetPos((T)0.);
+				Chaos::TVec3<FSolverReal> SourcePos((FSolverReal)0.), TargetPos((FSolverReal)0.);
 				for (int32 j = 0; j < Indices[i].Num(); j++)
 				{
 					SourcePos += Weights[i][j] * InParticles.P(Indices[i][j]);
@@ -288,18 +291,22 @@ namespace Chaos::Softs
 
 	protected:
 		TArray<TArray<int32>> Indices;
-		TArray<TArray<T>> Weights;
-		TArray<TVector<T, 3>> Constraints;
+		TArray<TArray<FSolverReal>> Weights;
+		TArray<TVector<FSolverReal, 3>> Constraints;
 		TArray<TArray<int32>> SecondIndices;
-		TArray<TArray<T>> SecondWeights;
-		TArray<T> Stiffness;
+		TArray<TArray<FSolverReal>> SecondWeights;
+		TArray<FSolverReal> Stiffness;
 		TArray<TArray<int32>> ConstraintsPerColor;
-		mutable TArray<T> LambdaArray;
+		mutable TArray<FSolverReal> LambdaArray;
 		bool VisualizeBindings = false;
-		FDeformableXPBDWeakConstraintParams DebugDrawParams;
+		FDeformableXPBDWeightedSpringConstraintParams DebugDrawParams;
 	};
 
+	using FDeformableXPBDWeakConstraintParams = FDeformableXPBDWeightedSpringConstraintParams;
 	
+	template <typename T, typename ParticleType>
+	using FXPBDWeakConstraints UE_DEPRECATED(5.8, "Deprecated. this class is to be deleted, use FXPBDWeightedSpringConstraints instead") = FXPBDWeightedSpringConstraints;
+
 
 }  // End namespace Chaos::Softs
 
