Clean up some incorrect documentation as the API has changed.

#rnx

- The method ::Get is scnchronous and does not have requests that need to be canceled.

diff --git a/Engine/Source/Runtime/Experimental/IoStore/OnDemand/Private/IasCache.h b/Engine/Source/Runtime/Experimental/IoStore/OnDemand/Private/IasCache.h
--- a/Engine/Source/Runtime/Experimental/IoStore/OnDemand/Private/IasCache.h
+++ b/Engine/Source/Runtime/Experimental/IoStore/OnDemand/Private/IasCache.h
@@ -25,8 +25,7 @@ public:
 
 	/** Deletes the IAS object, dropping all data persisted to disk and releasing
 	OS resources. As this also deletes the object thus any unique pointers should
-	be released prior to abandonment; TUniquePtr->Release()->Abandon(). Be sure
-	to cancel and collect any Get() tasks beforehand. */
+	be released prior to abandonment; TUniquePtr->Release()->Abandon().*/
 	virtual void Abandon() = 0;
 
 	/** Clears out the cache and resets it back to the default state. */
