Constraints: better notifications

#rnx

diff --git a/Engine/Source/Runtime/Experimental/Animation/Constraints/Constraints.Build.cs b/Engine/Source/Runtime/Experimental/Animation/Constraints/Constraints.Build.cs
--- a/Engine/Source/Runtime/Experimental/Animation/Constraints/Constraints.Build.cs
+++ b/Engine/Source/Runtime/Experimental/Animation/Constraints/Constraints.Build.cs
@@ -13,7 +13,9 @@ public class Constraints : ModuleRules
 				"CoreUObject",
 				"Engine",
 				"AnimationCore",
-				"MovieScene"
+				"MovieScene",
+				"Slate",
+				"SlateCore"
 			}
 		);
 	}
diff --git a/Engine/Source/Runtime/Experimental/Animation/Constraints/Private/Transform/TransformConstraintUtil.cpp b/Engine/Source/Runtime/Experimental/Animation/Constraints/Private/Transform/TransformConstraintUtil.cpp
--- a/Engine/Source/Runtime/Experimental/Animation/Constraints/Private/Transform/TransformConstraintUtil.cpp
+++ b/Engine/Source/Runtime/Experimental/Animation/Constraints/Private/Transform/TransformConstraintUtil.cpp
@@ -6,12 +6,14 @@
 #include "Engine/World.h"
 #include "GameFramework/Actor.h"
 #include "MovieSceneSection.h"
+#include "Framework/Notifications/NotificationManager.h"
 #include "Transform/TransformableHandle.h"
 #include "Transform/TransformableRegistry.h"
 #include "Transform/TransformConstraint.h"
 #include "Transform/DependencyBuilder.h"
 #include "UObject/Object.h"
 #include "UObject/Package.h"
+#include "Widgets/Notifications/SNotificationList.h"
 
 namespace UE::TransformConstraintUtil
 {
@@ -44,6 +46,21 @@ UTransformableHandle* GetHandle(UObject* InObject, const FName& InSocketName)
 
 	return nullptr;
 }
+
+void NotifyError(const FText& InNotification)
+{
+	// output log notification
+	UE_LOG(LogAnimation, Error, TEXT("%s"), *InNotification.ToString());
+	
+#if WITH_EDITOR
+	// editor notification
+	FNotificationInfo Info(InNotification);
+	Info.Image = FAppStyle::GetBrush(TEXT("MessageLog.Error"));
+	Info.ExpireDuration = 5.0f;
+
+	FSlateNotificationManager::Get().AddNotification(Info);
+#endif
+}
 	
 }
 
@@ -238,13 +255,19 @@ UTickableTransformConstraint* CreateFromType(UWorld* InWorld, const ETransformCo
 {
 	if (!InWorld)
 	{
-		UE_LOG(LogTemp, Error, TEXT("TransformConstraintUtil::CreateFromType sanity check failed."));
+		const FText Notification =
+			NSLOCTEXT("TransformConstraintUtil", "CreateFromType_NoWorld", "Cannot create a new constraint in a null world.");
+		Private::NotifyError(Notification);
 		return nullptr;
 	}
+	
 	const UEnum* ETransformConstraintTypeEnum = StaticEnum<ETransformConstraintType>();
 	if (!ETransformConstraintTypeEnum->IsValidEnumValue(static_cast<int64>(InType)))
 	{
-		UE_LOG(LogTemp, Error, TEXT("Constraint Type %d not recognized"), int(InType));
+		const FText Notification = FText::Format(
+			NSLOCTEXT("TransformConstraintUtil", "CreateFromType_UnknownType",  "Cannot create a new constraint with unknown type '%d'."),
+			static_cast<int>(InType));
+		Private::NotifyError(Notification);
 		return nullptr;
 	}
 
@@ -330,13 +353,14 @@ bool HasConstraintDependencyWith(UWorld* InWorld, const UTransformableHandle* In
 	return false;
 }
 	
-bool AreHandlesConstrainable( UWorld* InWorld, UTransformableHandle* InParentHandle, UTransformableHandle* InChildHandle)
+bool AreHandlesConstrainable( UWorld* InWorld, const UTransformableHandle* InParentHandle, const UTransformableHandle* InChildHandle)
 {
-	static const TCHAR* ErrorPrefix = TEXT("Dependency error:");
-
 	if (InChildHandle->GetHash() == InParentHandle->GetHash())
 	{
-		UE_LOG(LogTemp, Error, TEXT("%s handles are pointing at the same object."), ErrorPrefix);
+		const FText Notification =
+			NSLOCTEXT("TransformConstraintUtil", "AreHandlesConstrainable_SameObject", "Cannot create a constraint on itself.");
+		NotifyError(Notification);
+		
 		return false;
 	}
 
@@ -344,8 +368,12 @@ bool AreHandlesConstrainable( UWorld* InWorld, UTransformableHandle* InParentHan
 	if (InParentHandle->HasDirectDependencyWith(*InChildHandle))
 	{
 #if WITH_EDITOR
-		UE_LOG(LogTemp, Error, TEXT("%s: %s has a direct dependency with %s."),
-			ErrorPrefix, *InParentHandle->GetLabel(), *InChildHandle->GetLabel());
+		FFormatNamedArguments Args;
+		Args.Add(TEXT("Parent"), FText::FromString(InParentHandle->GetLabel()));
+		Args.Add(TEXT("Child"), FText::FromString(InChildHandle->GetLabel()));
+		const FText Notification = FText::Format(NSLOCTEXT("TransformConstraintUtil", "AreHandlesConstrainable_DirectDependency",
+			"Cannot create a new constraint between '{Parent}' and '{Child}' as there is a direct dependency between them."), Args);
+		NotifyError(Notification);
 #endif
 		return false;
 	}
@@ -354,8 +382,13 @@ bool AreHandlesConstrainable( UWorld* InWorld, UTransformableHandle* InParentHan
 	if (HasConstraintDependencyWith(InWorld, InParentHandle, InChildHandle))
 	{
 #if WITH_EDITOR
-		UE_LOG(LogTemp, Error, TEXT("%s: %s has an indirect dependency with %s."),
-			ErrorPrefix, *InParentHandle->GetLabel(), *InChildHandle->GetLabel());
+		FFormatNamedArguments Args;
+		Args.Add(TEXT("Parent"), FText::FromString(InParentHandle->GetLabel()));
+		Args.Add(TEXT("Child"), FText::FromString(InChildHandle->GetLabel()));
+		
+		const FText Notification = FText::Format(NSLOCTEXT("TransformConstraintUtil", "AreHandlesConstrainable_IndirectDependency",
+			"Cannot create a new constraint between '{Parent}' and '{Child}' as there is an indirect dependency between them."), Args);
+		NotifyError(Notification);
 #endif
 		return false;
 	}
@@ -374,12 +407,12 @@ UTickableTransformConstraint* CreateAndAddFromObjects(
 	const bool bUseDefault,
 	const TFunction<void()>& InValidDependencyFunction)
 {
-	static const TCHAR* ErrorPrefix = TEXT("TransformConstraintUtil::CreateAndAddFromActors");
-	
 	// SANITY CHECK
 	if (!InWorld || !InParent || !InChild)
 	{
-		UE_LOG(LogTemp, Error, TEXT("%s sanity check failed."), ErrorPrefix);
+		const FText Notification =
+			NSLOCTEXT("TransformConstraintUtil", "CreateAndAddFromObjects_SanityCheck", "Cannot create a new constraint with invalid world on handles.");
+		Private::NotifyError(Notification);
 		return nullptr;
 	}
 
@@ -425,15 +458,22 @@ bool AddConstraint(UWorld* InWorld, UTransformableHandle* InParentHandle, UTrans
 {
 	const bool bIsValidParent = InParentHandle && InParentHandle->IsValid();
 	const bool bIsValidChild = InChildHandle && InChildHandle->IsValid();
+	
 	if (!bIsValidParent || !bIsValidChild)
 	{
-		UE_LOG(LogTemp, Error, TEXT("TransformConstraintUtil::AddConst error adding constraint"));
+		const FText Notification =
+			NSLOCTEXT("TransformConstraintUtil", "AreHandlesConstrainable_InvalidHandles", "Cannot create a new constraint between invalid handle(s).");
+		Private::NotifyError(Notification);
+		
 		return false;
 	}
 
 	if (InNewConstraint == nullptr)
 	{
-		UE_LOG(LogTemp, Error, TEXT("TransformConstraintUtil::AddConst error creating constraint"));
+		const FText Notification =
+			NSLOCTEXT("TransformConstraintUtil", "AreHandlesConstrainable_InvalidConstraint", "Cannot add a null constraint.");
+		Private::NotifyError(Notification);
+		
 		return false;
 	}
 
