bughawk follow up

#rnx

diff --git a/Engine/Source/Runtime/Experimental/Animation/Constraints/Private/Transform/TransformConstraintUtil.cpp b/Engine/Source/Runtime/Experimental/Animation/Constraints/Private/Transform/TransformConstraintUtil.cpp
--- a/Engine/Source/Runtime/Experimental/Animation/Constraints/Private/Transform/TransformConstraintUtil.cpp
+++ b/Engine/Source/Runtime/Experimental/Animation/Constraints/Private/Transform/TransformConstraintUtil.cpp
@@ -265,8 +265,8 @@ UTickableTransformConstraint* CreateFromType(UWorld* InWorld, const ETransformCo
 	if (!ETransformConstraintTypeEnum->IsValidEnumValue(static_cast<int64>(InType)))
 	{
 		const FText Notification = FText::Format(
-			NSLOCTEXT("TransformConstraintUtil", "CreateFromType_UnknownType",  "Cannot create a new constraint with unknown type '%d'."),
-			static_cast<int>(InType));
+			NSLOCTEXT("TransformConstraintUtil", "CreateFromType_UnknownType",  "Cannot create a new constraint with unknown type '{0}'."),
+			FText::AsNumber(static_cast<int32>(InType)));
 		Private::NotifyError(Notification);
 		return nullptr;
 	}
