Minor safety fixes to FBitWindow

#rnx

- Calling FBitWindow::Reset with 0 as the parameter will now correctly reset the object to default settings and not break the mask.
- Calliing FBitWindow::Add on a completely empty object will no longer write out of bounds.
- Calling FBitWindow::AvgSetBits will no longer throw an exception when called on an empty object.

diff --git a/Engine/Source/Runtime/Experimental/IoStore/OnDemand/Private/IasHostGroup.cpp b/Engine/Source/Runtime/Experimental/IoStore/OnDemand/Private/IasHostGroup.cpp
--- a/Engine/Source/Runtime/Experimental/IoStore/OnDemand/Private/IasHostGroup.cpp
+++ b/Engine/Source/Runtime/Experimental/IoStore/OnDemand/Private/IasHostGroup.cpp
@@ -37,20 +37,35 @@ struct FBitWindow
 
 	void Reset(uint32 Count)
 	{
-		Count = FMath::RoundUpToPowerOfTwo(Count);
-		Bits.Init(false, Count);
-		Counter = 0;
-		Mask = Count - 1;
+		if (Count > 0)
+		{
+			Count = FMath::RoundUpToPowerOfTwo(Count);
+			Bits.Init(false, Count);
+			Counter = 0;
+			Mask = Count - 1 ;
+		}
+		else
+		{
+			*this = FBitWindow();
+		}
 	}
 
 	void Add(bool bValue)
 	{
-		const uint32 Idx = Counter++ & Mask;
-		Bits[Idx] = bValue;
+		if (!Bits.IsEmpty())
+		{
+			const uint32 Idx = Counter++ & Mask;
+			Bits[Idx] = bValue;
+		}
 	}
 
 	float AvgSetBits() const
 	{
+		if (Bits.IsEmpty())
+		{
+			return 0.0f;
+		}
+
 		return float(Bits.CountSetBits()) / float(Bits.Num());
 	}
 
