[NetPhysics] Fix misspelling (AccumulateErrorXR) and add CVar (p.RenderInterp.SnapToLatest) to always snap render interpolation to latest physics state but still take post-resim error correction into account.

diff --git a/Engine/Source/Runtime/Experimental/Chaos/Private/Framework/PhysicsProxyBase.cpp b/Engine/Source/Runtime/Experimental/Chaos/Private/Framework/PhysicsProxyBase.cpp
--- a/Engine/Source/Runtime/Experimental/Chaos/Private/Framework/PhysicsProxyBase.cpp
+++ b/Engine/Source/Runtime/Experimental/Chaos/Private/Framework/PhysicsProxyBase.cpp
@@ -46,6 +46,9 @@ namespace RenderInterpolationCVars
 
 	float RenderInterpMinimumAngularThreshold = 0.001f;
 	static FAutoConsoleVariableRef CVarRenderInterpMinimumAngularThreshold(TEXT("p.RenderInterp.MinimumAngularThreshold"), RenderInterpMinimumAngularThreshold, TEXT("When the remaining render error angle is below this we clear it, if ApplyExponentialDecay is enabled."));
+
+	bool bSnapToLatest = false;
+	static FAutoConsoleVariableRef CVarRenderInterpSnapToLatest(TEXT("p.RenderInterp.SnapToLatest"), bSnapToLatest, TEXT("Always snap to the latest physics state, instead of interpolating."));
 }
 
 IPhysicsProxyBase::~IPhysicsProxyBase()
@@ -82,7 +85,7 @@ const bool FProxyInterpolationError::IsErrorSmoothing() const
 }
 
 
-void FProxyInterpolationError::AccumlateErrorXR(const Chaos::FVec3 X, const FQuat R, const int32 CurrentSimTick, const int32 ErrorSmoothDuration)
+void FProxyInterpolationError::AccumulateErrorXR(const Chaos::FVec3 X, const FQuat R, const int32 CurrentSimTick, const int32 ErrorSmoothDuration)
 {
 	LastSimTick = CurrentSimTick - 1; // Error is from the previous simulation tick, not the current
 	EndDecayTick = LastSimTick + ErrorSmoothDuration;
diff --git a/Engine/Source/Runtime/Experimental/Chaos/Private/PhysicsProxy/ClusterUnionPhysicsProxy.cpp b/Engine/Source/Runtime/Experimental/Chaos/Private/PhysicsProxy/ClusterUnionPhysicsProxy.cpp
--- a/Engine/Source/Runtime/Experimental/Chaos/Private/PhysicsProxy/ClusterUnionPhysicsProxy.cpp
+++ b/Engine/Source/Runtime/Experimental/Chaos/Private/PhysicsProxy/ClusterUnionPhysicsProxy.cpp
@@ -488,6 +488,8 @@ namespace Chaos
 			return false;
 		}
 
+		const Chaos::FRealSingle LerpAlpha = RenderInterpolationCVars::bSnapToLatest ? 1.0f : *Alpha;
+		
 		FProxyInterpolationBase* InterpData = GetInterpolationData();
 		if (Error)
 		{
@@ -522,7 +524,7 @@ namespace Chaos
 				RenderInterpErrorCorrectionDurationTicks = FMath::FloorToInt32(ErrorCorrectionDuration / AsyncFixedTimeStep); // Convert duration from seconds to simulation ticks
 			}
 
-			InterpData->AccumlateErrorXR(Error->ErrorX, Error->ErrorR, SolverSyncTimestamp, RenderInterpErrorCorrectionDurationTicks);
+			InterpData->AccumulateErrorXR(Error->ErrorX, Error->ErrorR, SolverSyncTimestamp, RenderInterpErrorCorrectionDurationTicks);
 		}
 		const bool bHasInterpolationData = InterpData != nullptr;
 
@@ -610,8 +612,8 @@ namespace Chaos
 				}
 
 				const FVec3 NewX = bIsReplicationErrorSmoothing ?
-					FMath::Lerp(*Prev, NextPullData->X, *Alpha) + InterpData->GetErrorX(*Alpha) :
-					FMath::Lerp(*Prev, NextPullData->X, *Alpha);
+					FMath::Lerp(*Prev, NextPullData->X, LerpAlpha) + InterpData->GetErrorX(LerpAlpha) :
+					FMath::Lerp(*Prev, NextPullData->X, LerpAlpha);
 
 				Particle_External->SetX(NewX, false);
 			}
@@ -619,21 +621,21 @@ namespace Chaos
 			if (const FQuat* Prev = LerpHelper(PullData.R, ProxyTimestamp->OverWriteR))
 			{
 				const FQuat NewR = bIsReplicationErrorSmoothing ? 
-					FMath::Lerp(*Prev, NextPullData->R, *Alpha) * InterpData->GetErrorR(*Alpha) :
-					FMath::Lerp(*Prev, NextPullData->R, *Alpha); // Add rotational error offset in local space
+					FMath::Lerp(*Prev, NextPullData->R, LerpAlpha) * InterpData->GetErrorR(LerpAlpha) :
+					FMath::Lerp(*Prev, NextPullData->R, LerpAlpha); // Add rotational error offset in local space
 
 				Particle_External->SetR(NewR, false);
 			}
 
 			if (const FVec3* Prev = LerpHelper(PullData.V, ProxyTimestamp->OverWriteV))
 			{
-				const FVec3 NewV = FMath::Lerp(*Prev, NextPullData->V, *Alpha);
+				const FVec3 NewV = FMath::Lerp(*Prev, NextPullData->V, LerpAlpha);
 				Particle_External->SetV(NewV, false);
 			}
 
 			if (const FVec3* Prev = LerpHelper(PullData.W, ProxyTimestamp->OverWriteW))
 			{
-				const FVec3 NewW = FMath::Lerp(*Prev, NextPullData->W, *Alpha);
+				const FVec3 NewW = FMath::Lerp(*Prev, NextPullData->W, LerpAlpha);
 				Particle_External->SetW(NewW, false);
 			}
 
@@ -653,7 +655,7 @@ namespace Chaos
 						Chaos::FDebugDrawQueue::GetInstance().DrawDebugDirectionalArrow(ZOffset + (PullData.X + Error->ErrorX), ZOffset + PullData.X, 1, FColor::Red, false, 5.0f, 0, 0.5f);
 					}
 
-					Chaos::FDebugDrawQueue::GetInstance().DrawDebugDirectionalArrow(ZOffset + (Particle_External->GetX() - InterpData->GetErrorX(*Alpha)), ZOffset + Particle_External->GetX(), 1, FColor::Blue, false, 5.0f, 0, 0.5f);
+					Chaos::FDebugDrawQueue::GetInstance().DrawDebugDirectionalArrow(ZOffset + (Particle_External->GetX() - InterpData->GetErrorX(LerpAlpha)), ZOffset + Particle_External->GetX(), 1, FColor::Blue, false, 5.0f, 0, 0.5f);
 				}
 			}
 #endif // CHAOS_DEBUG_DRAW
diff --git a/Engine/Source/Runtime/Experimental/Chaos/Private/PhysicsProxy/GeometryCollectionPhysicsProxy.cpp b/Engine/Source/Runtime/Experimental/Chaos/Private/PhysicsProxy/GeometryCollectionPhysicsProxy.cpp
--- a/Engine/Source/Runtime/Experimental/Chaos/Private/PhysicsProxy/GeometryCollectionPhysicsProxy.cpp
+++ b/Engine/Source/Runtime/Experimental/Chaos/Private/PhysicsProxy/GeometryCollectionPhysicsProxy.cpp
@@ -4399,6 +4399,8 @@ bool FGeometryCollectionPhysicsProxy::PullFromPhysicsState(const Chaos::FDirtyGe
 {
 	if(IsObjectDeleting) return false;
 
+	const Chaos::FRealSingle LerpAlpha = RenderInterpolationCVars::bSnapToLatest ? 1.0f : *Alpha;
+
 	/**
 	 * CONTEXT: GAMETHREAD (Read Locked)
 	 * Perform a similar operation to Sync, but take the data from a gamethread-safe buffer. This will be called
@@ -4421,7 +4423,7 @@ bool FGeometryCollectionPhysicsProxy::PullFromPhysicsState(const Chaos::FDirtyGe
 
 	if (NumTransforms != GameThreadCollection.GetNumTransforms())
 	{
-		// At this point teh component may be gone and we have no way to get the name of the asset/component
+		// At this point the component may be gone and we have no way to get the name of the asset/component
 		UE_LOG(LogChaos, Error, 
 			TEXT("Attempt to update physics state on thegamethread where the number of transform (%d) does not match anymore the rest collection (%d) - skipping physics update"),
 			NumTransforms, 
@@ -4450,7 +4452,7 @@ bool FGeometryCollectionPhysicsProxy::PullFromPhysicsState(const Chaos::FDirtyGe
 		{
 			RenderInterpErrorCorrectionDurationTicks = FMath::FloorToInt32(ErrorCorrectionDuration / AsyncFixedTimeStep); // Convert duration from seconds to simulation ticks
 		}
-		InterpData->AccumlateErrorXR(Error->ErrorX, Error->ErrorR, SolverSyncTimestamp, RenderInterpErrorCorrectionDurationTicks);
+		InterpData->AccumulateErrorXR(Error->ErrorX, Error->ErrorR, SolverSyncTimestamp, RenderInterpErrorCorrectionDurationTicks);
 	}
 	const bool bHasInterpolationData = InterpData != nullptr;
 
@@ -4512,7 +4514,7 @@ bool FGeometryCollectionPhysicsProxy::PullFromPhysicsState(const Chaos::FDirtyGe
 				}
 
 				FParticle& GTParticle = *GTParticles[ParticleIndex];
-				const FResultInterpolator ResultInterpolator(TransformGroupIndex, PrevResults, NextResults, GTParticle, *Alpha);
+				const FResultInterpolator ResultInterpolator(TransformGroupIndex, PrevResults, NextResults, GTParticle, LerpAlpha);
 				if (ResultInterpolator.HasNoEntry())
 				{
 					continue;
@@ -4530,8 +4532,8 @@ bool FGeometryCollectionPhysicsProxy::PullFromPhysicsState(const Chaos::FDirtyGe
 
 					if (bIsReplicationErrorSmoothing)
 					{
-						NewX = NewX + InterpData->GetErrorX(*Alpha);
-						NewR = NewR * InterpData->GetErrorR(*Alpha); // Add rotational error offset in local space
+						NewX = NewX + InterpData->GetErrorX(LerpAlpha);
+						NewR = NewR * InterpData->GetErrorR(LerpAlpha); // Add rotational error offset in local space
 					}
 
 #if CHAOS_DEBUG_DRAW
@@ -4559,7 +4561,7 @@ bool FGeometryCollectionPhysicsProxy::PullFromPhysicsState(const Chaos::FDirtyGe
 									Chaos::FDebugDrawQueue::GetInstance().DrawDebugDirectionalArrow(ZOffset + PullDataT.ParticleX + Error->ErrorX, ZOffset + PullDataT.ParticleX, 1, FColor::Red, false, 3.0f, 0, 0.5f);
 								}
 
-								Chaos::FDebugDrawQueue::GetInstance().DrawDebugDirectionalArrow(ZOffset + (NewX - InterpData->GetErrorX(*Alpha)), ZOffset + NewX, 1, FColor::Blue, false, 3.0f, 0, 0.5f);
+								Chaos::FDebugDrawQueue::GetInstance().DrawDebugDirectionalArrow(ZOffset + (NewX - InterpData->GetErrorX(LerpAlpha)), ZOffset + NewX, 1, FColor::Blue, false, 3.0f, 0, 0.5f);
 							}
 						}
 					}
diff --git a/Engine/Source/Runtime/Experimental/Chaos/Private/PhysicsProxy/SingleParticlePhysicsProxy.cpp b/Engine/Source/Runtime/Experimental/Chaos/Private/PhysicsProxy/SingleParticlePhysicsProxy.cpp
--- a/Engine/Source/Runtime/Experimental/Chaos/Private/PhysicsProxy/SingleParticlePhysicsProxy.cpp
+++ b/Engine/Source/Runtime/Experimental/Chaos/Private/PhysicsProxy/SingleParticlePhysicsProxy.cpp
@@ -306,6 +306,8 @@ bool FSingleParticlePhysicsProxy::PullFromPhysicsState(const Chaos::FDirtyRigidP
 {
 	using namespace Chaos;
 
+	const Chaos::FRealSingle LerpAlpha = RenderInterpolationCVars::bSnapToLatest ? 1.0f : *Alpha;
+
 	// Move buffered data into the TPBDRigidParticle without triggering invalidation of the physics state.
 	Chaos::FPBDRigidParticle* Rigid = Particle ? Particle->CastToRigidParticle() : nullptr;
 	if (Rigid)
@@ -359,7 +361,7 @@ bool FSingleParticlePhysicsProxy::PullFromPhysicsState(const Chaos::FDirtyRigidP
 				RenderInterpErrorCorrectionDurationTicks = FMath::FloorToInt32(ErrorCorrectionDuration / AsyncFixedTimeStep); // Convert duration from seconds to simulation ticks
 			}
 
-			InterpData->AccumlateErrorXR(Error->ErrorX, Error->ErrorR, SolverSyncTimestamp, RenderInterpErrorCorrectionDurationTicks);
+			InterpData->AccumulateErrorXR(Error->ErrorX, Error->ErrorR, SolverSyncTimestamp, RenderInterpErrorCorrectionDurationTicks);
 			
 			const int32 RenderInterpErrorVelocitySmoothingDurationTicks = FMath::FloorToInt32(RenderInterpolationCVars::RenderInterpErrorVelocitySmoothingDuration / AsyncFixedTimeStep); // Convert duration from seconds to simulation ticks
 			InterpData->SetVelocitySmoothing(Rigid->V(), Rigid->X(), RenderInterpErrorVelocitySmoothingDurationTicks);
@@ -390,12 +392,12 @@ bool FSingleParticlePhysicsProxy::PullFromPhysicsState(const Chaos::FDirtyRigidP
 
 				if (const FVec3* Prev = LerpHelper(PullData.X, ProxyTimestamp->OverWriteX))
 				{
-					FVec3 Target = FMath::Lerp(*Prev, NextPullData->X, *Alpha);
+					FVec3 Target = FMath::Lerp(*Prev, NextPullData->X, LerpAlpha);
 					if (bIsReplicationErrorSmoothing)
 					{
 						DirectionalDecayPerformed = InterpData->DirectionalDecay(NextPullData->X - *Prev, RenderInterpolationCVars::RenderInterpErrorDirectionalDecayMultiplier);
 
-						Target += InterpData->GetErrorX(*Alpha);
+						Target += InterpData->GetErrorX(LerpAlpha);
 
 						if (InterpData->IsErrorVelocitySmoothing())
 						{
@@ -403,14 +405,14 @@ bool FSingleParticlePhysicsProxy::PullFromPhysicsState(const Chaos::FDirtyRigidP
 							if (RenderInterpolationCVars::bRenderInterpDebugDraw)
 							{
 								const FVector ZOffset = FVector(0,0, RenderInterpolationCVars::RenderInterpDebugDrawZOffset);
-								Chaos::FDebugDrawQueue::GetInstance().DrawDebugDirectionalArrow(ZOffset + (Target - InterpData->GetErrorX(*Alpha)), ZOffset + Target, 1, FColor::Blue, false, 5.0f, 0, 0.5f);
-								Chaos::FDebugDrawQueue::GetInstance().DrawDebugDirectionalArrow(ZOffset + InterpData->GetErrorVelocitySmoothingX(*Alpha), ZOffset + Target, 1, FColor::Blue, false, 5.0f, 0, 0.5f);
+								Chaos::FDebugDrawQueue::GetInstance().DrawDebugDirectionalArrow(ZOffset + (Target - InterpData->GetErrorX(LerpAlpha)), ZOffset + Target, 1, FColor::Blue, false, 5.0f, 0, 0.5f);
+								Chaos::FDebugDrawQueue::GetInstance().DrawDebugDirectionalArrow(ZOffset + InterpData->GetErrorVelocitySmoothingX(LerpAlpha), ZOffset + Target, 1, FColor::Blue, false, 5.0f, 0, 0.5f);
 								Chaos::FDebugDrawQueue::GetInstance().DrawDebugBox(ZOffset + Target, FVector(2, 1, 1), Rigid->R(), FColor::Cyan, false, 5.f, 0, 0.25f);
 							}
 #endif // CHAOS_DEBUG_DRAW
 
 							const int32 RenderInterpErrorVelocitySmoothingDurationTicks = FMath::FloorToInt32(RenderInterpolationCVars::RenderInterpErrorVelocitySmoothingDuration / AsyncFixedTimeStep); // Convert duration from seconds to simulation ticks
-							Target = FMath::Lerp(Target, InterpData->GetErrorVelocitySmoothingX(*Alpha), InterpData->GetErrorVelocitySmoothingAlpha(RenderInterpErrorVelocitySmoothingDurationTicks));
+							Target = FMath::Lerp(Target, InterpData->GetErrorVelocitySmoothingX(LerpAlpha), InterpData->GetErrorVelocitySmoothingAlpha(RenderInterpErrorVelocitySmoothingDurationTicks));
 						}
 					}
 
@@ -420,10 +422,10 @@ bool FSingleParticlePhysicsProxy::PullFromPhysicsState(const Chaos::FDirtyRigidP
 
 				if (const FQuat* Prev = LerpHelper(PullData.R, ProxyTimestamp->OverWriteR))
 				{
-					FQuat Target = FMath::Lerp(*Prev, NextPullData->R, *Alpha);
+					FQuat Target = FMath::Lerp(*Prev, NextPullData->R, LerpAlpha);
 					if (bIsReplicationErrorSmoothing)
 					{
-						Target = Target * InterpData->GetErrorR(*Alpha); // Add rotational error offset in local space
+						Target = Target * InterpData->GetErrorR(LerpAlpha); // Add rotational error offset in local space
 					}
 					Rigid->SetR(Target, false);
 					bNeedUpdateShapeBounds = true;
@@ -447,11 +449,11 @@ bool FSingleParticlePhysicsProxy::PullFromPhysicsState(const Chaos::FDirtyRigidP
 
 						if (InterpData->IsErrorVelocitySmoothing())
 						{
-							Chaos::FDebugDrawQueue::GetInstance().DrawDebugBox(ZOffset + InterpData->GetErrorVelocitySmoothingX(*Alpha), FVector(2, 2, 2), Rigid->R(), FColor::Purple, false, 5.f, 0, 0.5f);
+							Chaos::FDebugDrawQueue::GetInstance().DrawDebugBox(ZOffset + InterpData->GetErrorVelocitySmoothingX(LerpAlpha), FVector(2, 2, 2), Rigid->R(), FColor::Purple, false, 5.f, 0, 0.5f);
 						}
 						else
 						{
-							Chaos::FDebugDrawQueue::GetInstance().DrawDebugDirectionalArrow(ZOffset + (Rigid->X() - InterpData->GetErrorX(*Alpha)), ZOffset + Rigid->X(), 1, FColor::Blue, false, 5.0f, 0, 0.5f);
+							Chaos::FDebugDrawQueue::GetInstance().DrawDebugDirectionalArrow(ZOffset + (Rigid->X() - InterpData->GetErrorX(LerpAlpha)), ZOffset + Rigid->X(), 1, FColor::Blue, false, 5.0f, 0, 0.5f);
 						}
 					}
 				}
@@ -460,13 +462,13 @@ bool FSingleParticlePhysicsProxy::PullFromPhysicsState(const Chaos::FDirtyRigidP
 
 			if (const FVec3* Prev = LerpHelper(PullData.V, ProxyTimestamp->OverWriteV))
 			{
-				FVec3 Target = FMath::Lerp(*Prev, NextPullData->V, *Alpha);
+				FVec3 Target = FMath::Lerp(*Prev, NextPullData->V, LerpAlpha);
 				Rigid->SetV(Target, false);
 			}
 
 			if (const FVec3* Prev = LerpHelper(PullData.W, ProxyTimestamp->OverWriteW))
 			{
-				FVec3 Target = FMath::Lerp(*Prev, NextPullData->W, *Alpha);
+				FVec3 Target = FMath::Lerp(*Prev, NextPullData->W, LerpAlpha);
 				Rigid->SetW(Target, false);
 			}
 
@@ -476,7 +478,7 @@ bool FSingleParticlePhysicsProxy::PullFromPhysicsState(const Chaos::FDirtyRigidP
 			{
 				Rigid->SetObjectState(PullData.ObjectState, true, /*bInvalidate=*/false);
 			}
-			else if(ProxyTimestamp->ObjectStateTimestamp == SolverSyncTimestamp && *Alpha == 1.f)
+			else if(ProxyTimestamp->ObjectStateTimestamp == SolverSyncTimestamp && LerpAlpha == 1.f)
 			{
 				//if timestamp is the same as next, AND alpha is exactly 1, we are exactly at Next's time
 				//so we can use its sleep state
diff --git a/Engine/Source/Runtime/Experimental/Chaos/Public/Chaos/Framework/PhysicsProxyBase.h b/Engine/Source/Runtime/Experimental/Chaos/Public/Chaos/Framework/PhysicsProxyBase.h
--- a/Engine/Source/Runtime/Experimental/Chaos/Public/Chaos/Framework/PhysicsProxyBase.h
+++ b/Engine/Source/Runtime/Experimental/Chaos/Public/Chaos/Framework/PhysicsProxyBase.h
@@ -43,6 +43,7 @@ namespace RenderInterpolationCVars
 	extern float RenderInterpExponentialDecayAngularHalfLife;
 	extern float RenderInterpMinimumLinearThreshold;
 	extern float RenderInterpMinimumAngularThreshold;
+	extern bool bSnapToLatest;
 }
 
 enum class EProxyInterpolationType : uint32
@@ -291,7 +292,7 @@ public:
 	virtual const FQuat GetErrorR(const Chaos::FRealSingle Alpha) const { return FQuat::Identity; }
 
 	/** Add X and R error onto current error to correct through interpolation */
-	virtual void AccumlateErrorXR(const Chaos::FVec3 X, const FQuat R, const int32 CurrentSimTick, const int32 ErrorSmoothDuration) { }
+	virtual void AccumulateErrorXR(const Chaos::FVec3 X, const FQuat R, const int32 CurrentSimTick, const int32 ErrorSmoothDuration) { }
 
 	/** Tick current error data and decay error */
 	virtual const bool UpdateError(const int32 CurrentSimTick, const Chaos::FReal AsyncFixedTimeStep) { return false; }
@@ -338,7 +339,7 @@ struct FProxyInterpolationError : FProxyInterpolationBase
 
 	virtual const FQuat GetErrorR(const Chaos::FRealSingle Alpha) const override { return FQuat::Slerp(ErrorRPrev, ErrorR, Alpha); }
 
-	CHAOS_API virtual void AccumlateErrorXR(const Chaos::FVec3 X, const FQuat R, const int32 CurrentSimTick, const int32 ErrorSmoothDuration) override;
+	CHAOS_API virtual void AccumulateErrorXR(const Chaos::FVec3 X, const FQuat R, const int32 CurrentSimTick, const int32 ErrorSmoothDuration) override;
 
 	CHAOS_API virtual const bool UpdateError(const int32 CurrentSimTick, const Chaos::FReal AsyncFixedTimeStep) override;
 
