Console variables for bubble resim

diff --git a/Engine/Source/Runtime/Experimental/Chaos/Private/Chaos/ResimConsoleVariables.cpp b/Engine/Source/Runtime/Experimental/Chaos/Private/Chaos/ResimConsoleVariables.cpp
new file mode 100644
--- /dev/null
+++ b/Engine/Source/Runtime/Experimental/Chaos/Private/Chaos/ResimConsoleVariables.cpp
@@ -0,0 +1,16 @@
+// Copyright Epic Games, Inc. All Rights Reserved.
+
+#include "Chaos/ResimConsoleVariables.h"
+#include "HAL/IConsoleManager.h"
+
+namespace Chaos
+{
+	namespace ResimConsoleVars
+	{
+		bool bIsBubbleResimEnabled = false;
+		static FAutoConsoleVariableRef CVarIsBubbleResimEnabled(
+			TEXT("p.Chaos.BubbleResim.Enable"),
+			bIsBubbleResimEnabled,
+			TEXT("If true bubble resim is enabled, otherwise a resim causes the whole world to resimulate, not just bubbles around server corrected particles"));
+	}
+}
\ No newline at end of file
diff --git a/Engine/Source/Runtime/Experimental/Chaos/Public/Chaos/ResimConsoleVariables.h b/Engine/Source/Runtime/Experimental/Chaos/Public/Chaos/ResimConsoleVariables.h
new file mode 100644
--- /dev/null
+++ b/Engine/Source/Runtime/Experimental/Chaos/Public/Chaos/ResimConsoleVariables.h
@@ -0,0 +1,12 @@
+// Copyright Epic Games, Inc. All Rights Reserved.
+
+#pragma once
+#include "Chaos/Core.h"
+
+namespace Chaos
+{
+	namespace ResimConsoleVars
+	{
+		extern CHAOS_API bool bIsBubbleResimEnabled;
+	}
+}
\ No newline at end of file
