Clean up IAX LatencyTest code.

#rnx

-We can remove a fair bit of code by switching from FPlatformTime::Cycles64 to FPlatformTime::Seconds as it eliminates the need for a second conversion pass.
- Also avoids the multiple 64->32 bit casting, we now only cast the final ms value to 32 bit at the end.

diff --git a/Engine/Source/Runtime/Experimental/IoStore/OnDemand/Private/LatencyTesting.cpp b/Engine/Source/Runtime/Experimental/IoStore/OnDemand/Private/LatencyTesting.cpp
--- a/Engine/Source/Runtime/Experimental/IoStore/OnDemand/Private/LatencyTesting.cpp
+++ b/Engine/Source/Runtime/Experimental/IoStore/OnDemand/Private/LatencyTesting.cpp
@@ -76,27 +76,15 @@ void LatencyTest(FAnsiStringView InUrl, FAnsiStringView InPath, uint32 TimeOutMs
 				}
 			});
 
-		uint64 Cycles = FPlatformTime::Cycles64();
+		const double StartTime = FPlatformTime::Seconds();
 		
 		while (Loop.Tick(TimeOutMs) != 0)
 		{
-
-		}
-
-		Cycles = FPlatformTime::Cycles64() - Cycles;
-
-		Result = Ok ? int32(Cycles) : -1;
-	}
-
-	int64 Freq = int64(1.0 / FPlatformTime::GetSecondsPerCycle64());
-	for (int32& Result : OutResults)
-	{
-		if (Result == -1)
-		{
-			continue;
+			// Do nothing, we want this loop to be as tight as possible.
 		}
 
-		Result = int32((int64(Result) * 1000) / Freq);
+		const int32 TimeMs = FMath::RoundToInt32((FPlatformTime::Seconds() - StartTime) * 1000.0);
+		Result = Ok ? TimeMs : -1;
 	}
 }
 
@@ -110,7 +98,7 @@ bool ConnectionTest(FAnsiStringView Url, FAnsiStringView Path, uint32 TimeoutMs)
 	if (Results[0] >= 0 || Results[1] >= 0 || Results[2] >= 0 || Results[3] >= 0)
 	{
 #if !UE_BUILD_SHIPPING
-		UE_LOGFMT(LogIas, Log, "Endpoint '{Url}' latency test (ms): {Result0} {Result1} {Result2} {Result3}",
+		UE_LOGFMT(LogIas, Display, "Endpoint '{Url}' latency test (ms): {Result0} {Result1} {Result2} {Result3}",
 			Url, Results[0], Results[1], Results[2], Results[3]);
 #endif // !UE_BUILD_SHIPPING
 
