Fixed a typo in code documentation

#rnx

diff --git a/Engine/Source/Runtime/Experimental/IoStore/OnDemand/Private/OnDemandConfig.h b/Engine/Source/Runtime/Experimental/IoStore/OnDemand/Private/OnDemandConfig.h
--- a/Engine/Source/Runtime/Experimental/IoStore/OnDemand/Private/OnDemandConfig.h
+++ b/Engine/Source/Runtime/Experimental/IoStore/OnDemand/Private/OnDemandConfig.h
@@ -23,4 +23,4 @@ FIasCacheConfig GetStreamingCacheConfig(const TCHAR* CommandLine);
 
 TIoStatusOr<FOnDemandInstallCacheConfig> TryParseInstallCacheConfig(const TCHAR* CommandLine);
 
-} // namespace UE::IoStore
+} // namespace UE::IoStore::Config
