Minor improvements to UE::IoStore::Config::ParseSizeParam

#rnx

- Added code documentation clarifying the return values. Looking at the existing code a few places expect that parsing an invalid integer would return -1 instead of the 0 that we do return.
- Added support for using IEC units of measurement as well as SI units.
- Added unit tests to further clarify behavior expectations.

diff --git a/Engine/Source/Runtime/Experimental/IoStore/OnDemand/Private/OnDemandConfig.cpp b/Engine/Source/Runtime/Experimental/IoStore/OnDemand/Private/OnDemandConfig.cpp
--- a/Engine/Source/Runtime/Experimental/IoStore/OnDemand/Private/OnDemandConfig.cpp
+++ b/Engine/Source/Runtime/Experimental/IoStore/OnDemand/Private/OnDemandConfig.cpp
@@ -1,13 +1,12 @@
 // Copyright Epic Games, Inc. All Rights Reserved.
 
 #include "OnDemandConfig.h"
-#include "IO/IoStatus.h"
-#include "OnDemandIoDispatcherBackend.h"
-#include "OnDemandInstallCache.h"
 
 #include "HAL/PlatformProcess.h"
-#include "IasCache.h"
+#include "IO/IoStatus.h"
 #include "IO/IoStoreOnDemand.h"
+#include "IO/IoStoreOnDemandInternals.h"
+#include "IasCache.h"
 #include "Misc/Base64.h"
 #include "Misc/ConfigCacheIni.h"
 #include "Misc/CoreDelegates.h"
@@ -15,10 +14,18 @@
 #include "Misc/EncryptionKeyManager.h"
 #include "Misc/Fork.h"
 #include "Misc/Parse.h"
-#include "Misc/Paths.h"
 #include "Misc/PathViews.h"
+#include "Misc/Paths.h"
+#include "OnDemandInstallCache.h"
+#include "OnDemandIoDispatcherBackend.h"
 #include "String/LexFromString.h"
 
+#if WITH_IOSTORE_ONDEMAND_TESTS
+#include "TestHarness.h"
+#include "TestMacros/Assertions.h"
+#include <catch2/generators/catch_generators.hpp>
+#endif
+
 namespace UE::IoStore::Config
 {
 
@@ -140,11 +147,25 @@ int64 ParseSizeParam(FStringView Value)
 
 	int64 Size = -1;
 	LexFromString(Size, Value);
-	if (Size >= 0)
+
+	if (Size > 0)
 	{
-		if (Value.EndsWith(TEXT("GB"))) return Size << 30;
-		if (Value.EndsWith(TEXT("MB"))) return Size << 20;
-		if (Value.EndsWith(TEXT("KB"))) return Size << 10;
+		// Now we have a valid integer value we should check to see if any size extension was in the string
+
+		if (Value.EndsWith(TEXT("GB")) || Value.EndsWith(TEXT("GiB")))
+		{
+			return Size << 30;
+		}
+
+		if (Value.EndsWith(TEXT("MB")) || Value.EndsWith(TEXT("MiB")))
+		{
+			return Size << 20;
+		}
+
+		if (Value.EndsWith(TEXT("KB")) || Value.EndsWith(TEXT("KiB")))
+		{
+			return Size << 10;
+		}
 	}
 	return Size;
 }
@@ -462,4 +483,49 @@ TIoStatusOr<FOnDemandInstallCacheConfig> TryParseInstallCacheConfig(const TCHAR*
 	return OutConfig;
 }
 
+#if WITH_IOSTORE_ONDEMAND_TESTS
+
+TEST_CASE("IoStore::OnDemand::ParseSizeParam", "[IoStoreOnDemand][Ias]")
+{
+	SECTION("ParseFromString")
+	{
+		CHECK(ParseSizeParam(TEXT("0")) == 0);
+		CHECK(ParseSizeParam(TEXT("1")) == 1);
+		CHECK(ParseSizeParam(TEXT("1")) != 2);
+
+		CHECK(ParseSizeParam(TEXT("1024")) == 1024);
+		CHECK(ParseSizeParam(TEXT("1KB")) == 1024);
+		CHECK(ParseSizeParam(TEXT("1KiB")) == 1024);
+		CHECK(ParseSizeParam(TEXT("2KB")) != 1024);
+		CHECK(ParseSizeParam(TEXT("2KiB")) != 1024);
+
+		CHECK(ParseSizeParam(TEXT("1048576")) == 1048576);
+		CHECK(ParseSizeParam(TEXT("1MB")) == 1048576);
+		CHECK(ParseSizeParam(TEXT("1MiB")) == 1048576);
+		CHECK(ParseSizeParam(TEXT("2MB")) != 1048576);
+		CHECK(ParseSizeParam(TEXT("2MiB")) != 1048576);
+
+		CHECK(ParseSizeParam(TEXT("1073741824")) == 1073741824);
+		CHECK(ParseSizeParam(TEXT("1GB")) == 1073741824);
+		CHECK(ParseSizeParam(TEXT("1GiB")) == 1073741824);
+		CHECK(ParseSizeParam(TEXT("2GB")) != 1073741824);
+		CHECK(ParseSizeParam(TEXT("2GiB")) != 1073741824);
+
+		CHECK(ParseSizeParam(TEXT("Invalid")) == 0);
+	}
+
+	SECTION("ParseFromCommandLine")
+	{
+		const TCHAR* CommandLine = TEXT("-BasicTestValue=1092387456 -SmallTestValue=1KB -MediumTestValue=2MiB -LargeTestValue=3GiB");
+
+		CHECK(ParseSizeParam(CommandLine, TEXT("MissingValue=")) == -1);
+		CHECK(ParseSizeParam(CommandLine, TEXT("BasicTestValue=")) == 1092387456);
+		CHECK(ParseSizeParam(CommandLine, TEXT("SmallTestValue=")) == 1024);
+		CHECK(ParseSizeParam(CommandLine, TEXT("MediumTestValue=")) == 2097152);
+		CHECK(ParseSizeParam(CommandLine, TEXT("LargeTestValue=")) == 3221225472);
+	}
+}
+
+#endif // WITH_IOSTORE_ONDEMAND_TESTS
+
 } // namespace UE::IoStore
diff --git a/Engine/Source/Runtime/Experimental/IoStore/OnDemand/Private/OnDemandConfig.h b/Engine/Source/Runtime/Experimental/IoStore/OnDemand/Private/OnDemandConfig.h
--- a/Engine/Source/Runtime/Experimental/IoStore/OnDemand/Private/OnDemandConfig.h
+++ b/Engine/Source/Runtime/Experimental/IoStore/OnDemand/Private/OnDemandConfig.h
@@ -13,8 +13,37 @@ namespace UE::IoStore { struct FOnDemandInstallCacheConfig; }
 namespace UE::IoStore::Config
 {
 
+/*
+ * Attempts to convert the provided string to an integer.
+ * Unlike LexFromString this method supports postfixing the number with one of the following units of measurement:
+ * KB/KiB
+ * MB/MiB
+ * GB/GiB
+ * For example "1KB" would return 1024 or "1MiB" would return 1048576.
+ * Note that the SI units (KB) are treated as the IEC version (KiB).
+ * 
+ * @param	Value	The string to parse the integer from.
+ * 
+ * @return	The integer in the string, 0 if no valid integer was found.
+ */
 int64 ParseSizeParam(FStringView Value);
 
+/*
+ * Attempt to find the given Parameter in the provided CommandLine and return it's value as an integer.
+ * Unlike using FParse::Value directly this method supports postfixing the number with one of the following
+ * units of measurement:
+ * KB/KiB
+ * MB/MiB
+ * GB/GiB
+ * For example "-Param=1KB" would return 1024 or "-Param=1MiB" would return 1048576.
+ * Note that the SI units (KB) are treated as the IEC version (KiB).
+ * 
+ * @param	CommandLine	The commandline to parse from.
+ * @param	Param		The parameter on the commandline to parse.
+ * 
+ * @return	The parameter value as an integer, -1 if the parameter could not be found on the commandline
+ *			and 0 if the parameter was not a valid integer.
+ */
 int64 ParseSizeParam(const TCHAR* CommandLine, const TCHAR* Param);
 
 TIoStatusOr<FOnDemandEndpointConfig> TryParseEndpointConfig(const TCHAR* CommandLine);
