Bubble Resim: Adding some missing functions to SetResimType

diff --git a/Engine/Source/Runtime/Experimental/Chaos/Private/Chaos/PBDJointConstraints.cpp b/Engine/Source/Runtime/Experimental/Chaos/Private/Chaos/PBDJointConstraints.cpp
--- a/Engine/Source/Runtime/Experimental/Chaos/Private/Chaos/PBDJointConstraints.cpp
+++ b/Engine/Source/Runtime/Experimental/Chaos/Private/Chaos/PBDJointConstraints.cpp
@@ -84,7 +84,7 @@ namespace Chaos
 
 	void FPBDJointConstraintHandle::ClearConstraintBreaking()
 	{
-		return ConcreteContainer()->ClearConstraintBreaking(ConstraintIndex);
+		ConcreteContainer()->ClearConstraintBreaking(ConstraintIndex);
 	}
 
 	bool FPBDJointConstraintHandle::IsConstraintViolating() const
@@ -94,7 +94,7 @@ namespace Chaos
 
 	void FPBDJointConstraintHandle::ClearConstraintViolating()
 	{
-		return ConcreteContainer()->ClearConstraintViolating(ConstraintIndex);
+		ConcreteContainer()->ClearConstraintViolating(ConstraintIndex);
 	}
 
 	bool FPBDJointConstraintHandle::IsDriveTargetChanged() const
@@ -104,7 +104,7 @@ namespace Chaos
 
 	void FPBDJointConstraintHandle::ClearDriveTargetChanged()
 	{
-		return ConcreteContainer()->ClearDriveTargetChanged(ConstraintIndex);
+		ConcreteContainer()->ClearDriveTargetChanged(ConstraintIndex);
 	}
 
 	bool FPBDJointConstraintHandle::IsConstraintEnabled() const
@@ -139,12 +139,12 @@ namespace Chaos
 
 	void FPBDJointConstraintHandle::SetSyncState(ESyncState SyncState)
 	{
-		return ConcreteContainer()->SetConstraintSyncState(ConstraintIndex, SyncState);
+		ConcreteContainer()->SetConstraintSyncState(ConstraintIndex, SyncState);
 	}
 
 	void FPBDJointConstraintHandle::SetEnabledDuringResim(bool bEnabled)
 	{
-		return ConcreteContainer()->SetConstraintEnabledDuringResim(ConstraintIndex, bEnabled);
+		ConcreteContainer()->SetConstraintEnabledDuringResim(ConstraintIndex, bEnabled);
 	}
 
 	bool FPBDJointConstraintHandle::IsEnabledDuringResim() const
@@ -157,6 +157,11 @@ namespace Chaos
 		return ConcreteContainer()->GetConstraintResimType(ConstraintIndex);
 	}
 
+	void FPBDJointConstraintHandle::SetResimType(EResimType ResimType)
+	{
+		ConcreteContainer()->SetConstraintResimType(ConstraintIndex, ResimType);
+	}
+
 	constexpr int32 ConstraintChildIndex = 0;
 	constexpr int32 ConstraintParentIndex = 1;
 
@@ -1006,6 +1011,11 @@ namespace Chaos
 		return ConstraintStates[ConstraintIndex].ResimType;
 	}
 
+	void FPBDJointConstraints::SetConstraintResimType(int32 ConstraintIndex, EResimType ResimType)
+	{
+		ConstraintStates[ConstraintIndex].ResimType = ResimType;
+	}
+
 	void FPBDJointConstraints::PrepareTick()
 	{
 		SCOPE_CYCLE_COUNTER(STAT_Joints_PrepareTick);
diff --git a/Engine/Source/Runtime/Experimental/Chaos/Public/Chaos/PBDJointConstraints.h b/Engine/Source/Runtime/Experimental/Chaos/Public/Chaos/PBDJointConstraints.h
--- a/Engine/Source/Runtime/Experimental/Chaos/Public/Chaos/PBDJointConstraints.h
+++ b/Engine/Source/Runtime/Experimental/Chaos/Public/Chaos/PBDJointConstraints.h
@@ -113,6 +113,7 @@ namespace Private
 		CHAOS_API void SetEnabledDuringResim(bool bEnabled);
 		CHAOS_API bool IsEnabledDuringResim() const;
 		CHAOS_API EResimType ResimType() const;
+		CHAOS_API void SetResimType(EResimType ResimType);
 
 		UE_DEPRECATED(5.2, "No longer used")
 		CHAOS_API int32 GetConstraintIsland() const;
@@ -384,6 +385,7 @@ namespace Private
 		CHAOS_API bool IsConstraintEnabledDuringResim(int32 ConstraintIndex) const;
 		
 		CHAOS_API EResimType GetConstraintResimType(int32 ConstraintIndex) const;
+		CHAOS_API void SetConstraintResimType(int32 ConstraintIndex, EResimType ResimType);
 
 		//
 		// FConstraintContainer Implementation
