Chaos - add normlization checkSlow to SimRaycastImplicit to make contract more explicit

#rnx

diff --git a/Engine/Source/Runtime/Experimental/Chaos/Private/Chaos/Collision/SimRaycast.cpp b/Engine/Source/Runtime/Experimental/Chaos/Private/Chaos/Collision/SimRaycast.cpp
--- a/Engine/Source/Runtime/Experimental/Chaos/Private/Chaos/Collision/SimRaycast.cpp
+++ b/Engine/Source/Runtime/Experimental/Chaos/Private/Chaos/Collision/SimRaycast.cpp
@@ -7,6 +7,9 @@ namespace Chaos::Private
 {
 	bool SimRaycastImplicit(const FSimRaycastRay& InRay, const FImplicitObject* InImplicit, const FRigidTransform3& InImplicitTransform, FSimRaycastHit& InOutHit)
 	{
+		// We assume the ray direction is normalized but that's fairly expensive to test...
+		checkSlow(InRay.Dir.IsNormalized());
+
 		// @todo(chaos): could optimize uniform scale rather than just unit scale...
 		const bool bIsImplicitUnitScale = FVec3::IsNearlyEqual(InImplicitTransform.GetScale3D(), FVec3(1, 1, 1), UE_SMALL_NUMBER);
 
