Keep track of in-bound amount

#rnx

diff --git a/Engine/Source/Runtime/Experimental/IoStore/HttpClient/Private/Peer.inl b/Engine/Source/Runtime/Experimental/IoStore/HttpClient/Private/Peer.inl
--- a/Engine/Source/Runtime/Experimental/IoStore/HttpClient/Private/Peer.inl
+++ b/Engine/Source/Runtime/Experimental/IoStore/HttpClient/Private/Peer.inl
@@ -39,6 +39,7 @@ class FPeer
 public:
 				FPeer() = default;
 				FPeer(FSocket InSocket);
+	uint64		GetRecvSize() const					{ return Socket.GetRecvSize(); }
 	FWaitable	GetWaitable() const					{ return Socket.GetWaitable(); }
 	FOutcome	Send(const char* Data, int32 Size)	{ return Socket.Send(Data, Size); }
 	FOutcome	Recv(char* Out, int32 MaxSize)		{ return Socket.Recv(Out, MaxSize); }
diff --git a/Engine/Source/Runtime/Experimental/IoStore/HttpClient/Private/Socket.inl b/Engine/Source/Runtime/Experimental/IoStore/HttpClient/Private/Socket.inl
--- a/Engine/Source/Runtime/Experimental/IoStore/HttpClient/Private/Socket.inl
+++ b/Engine/Source/Runtime/Experimental/IoStore/HttpClient/Private/Socket.inl
@@ -372,6 +372,7 @@ public:
 	FSocket&	operator = (FSocket&& Rhs)	{ Move(MoveTemp(Rhs)); return *this; }
 	bool		IsValid() const				{ return Socket != InvalidSocket; }
 	FWaitable	GetWaitable() const			{ return { Socket }; }
+	uint64		GetRecvSize() const;
 	bool		Create();
 	void		Destroy();
 	FOutcome	Connect(uint32 Ip, uint32 Port);
@@ -384,6 +385,7 @@ public:
 
 private:
 	void		Move(FSocket&& Rhs);
+	uint64		RecvSize = 0;
 	SocketType	Socket = InvalidSocket;
 
 				FSocket(const FSocket&) = delete;
@@ -397,6 +399,12 @@ void FSocket::Move(FSocket&& Rhs)
 	Swap(Socket, Rhs.Socket);
 }
 
+////////////////////////////////////////////////////////////////////////////////
+uint64 FSocket::GetRecvSize() const
+{
+	return RecvSize;
+}
+
 ////////////////////////////////////////////////////////////////////////////////
 bool FSocket::Create()
 {
@@ -513,6 +521,11 @@ FOutcome FSocket::Recv(char* Dest, uint32 Size)
 		return Outcome;
 	};
 
+	if (Result > 0)
+	{
+		RecvSize += Result;
+	}
+
 	if (Result > 0)						return Return(FOutcome::Ok(Result));
 	if (Result == 0)					return Return(FOutcome::Error("Recv ATH0"));
 	if (IsSocketResult(EWOULDBLOCK))	return Return(FOutcome::Waiting());
