Method to gather how much was in-bound each tick

#rnx

diff --git a/Engine/Source/Runtime/Experimental/IoStore/HttpClient/Private/Loop.inl b/Engine/Source/Runtime/Experimental/IoStore/HttpClient/Private/Loop.inl
--- a/Engine/Source/Runtime/Experimental/IoStore/HttpClient/Private/Loop.inl
+++ b/Engine/Source/Runtime/Experimental/IoStore/HttpClient/Private/Loop.inl
@@ -414,6 +414,7 @@ class FPeerGroup
 {
 public:
 						FPeerGroup() = default;
+	uint32				CollectRecvSizeKiB();
 	void				SetMaxInflight(uint32 Maximum);
 	void				Unwait()			{ check(bWaiting); bWaiting = false; }
 	FWaiter				GetWaiter() const;
@@ -429,6 +430,7 @@ private:
 	FActivityList		Send;
 	FActivityList		Recv;
 	FHttpPeer			Peer;
+	uint32				LastRecvSizeKiB = 0;
 	uint32				LastUseMs = 0;
 	uint8				IsKeepAlive = 0;
 	uint8				InflightNum = 0;
@@ -439,6 +441,23 @@ private:
 	UE_NONCOPYABLE(FPeerGroup);
 };
 
+////////////////////////////////////////////////////////////////////////////////
+uint32 FPeerGroup::CollectRecvSizeKiB()
+{
+	if (!Peer.IsValid())
+	{
+		return 0;
+	}
+
+	uint64 Size = Peer.GetRecvSize();
+	Size += 512;
+	Size >>= 10;
+
+	uint32 Ret = uint32(Size) - LastRecvSizeKiB;
+	LastRecvSizeKiB = uint32(Size);
+	return Ret;
+}
+
 ////////////////////////////////////////////////////////////////////////////////
 void FPeerGroup::SetMaxInflight(uint32 Maximum)
 {
