Missed reset of tracking variable.

#rnx

diff --git a/Engine/Source/Runtime/Experimental/IoStore/HttpClient/Private/Loop.inl b/Engine/Source/Runtime/Experimental/IoStore/HttpClient/Private/Loop.inl
--- a/Engine/Source/Runtime/Experimental/IoStore/HttpClient/Private/Loop.inl
+++ b/Engine/Source/Runtime/Experimental/IoStore/HttpClient/Private/Loop.inl
@@ -745,6 +745,7 @@ void FPeerGroup::TickSend(FTickState& State, FHost& Host, FPoller& Poller)
 		bNegotiating = true;
 		bWillBlock = Outcome.IsWaiting();
 		InflightNum = 0;
+		LastRecvSizeKiB = 0;
 
 		FHttpPeer::FParams Params = {
 			.Socket = MoveTemp(Socket),
diff --git a/Engine/Source/Runtime/Experimental/IoStore/HttpClient/Private/Socket.inl b/Engine/Source/Runtime/Experimental/IoStore/HttpClient/Private/Socket.inl
--- a/Engine/Source/Runtime/Experimental/IoStore/HttpClient/Private/Socket.inl
+++ b/Engine/Source/Runtime/Experimental/IoStore/HttpClient/Private/Socket.inl
@@ -552,7 +552,7 @@ bool FSocket::SetBlocking(bool bBlocking)
 	int32 NewFlags = bBlocking
 		? (Flags & ~int32(O_NONBLOCK))
 		: (Flags |	int32(O_NONBLOCK));
-	return (Flags == NewFlags) || (fcntl(Socket, F_SETFL, Flags) >= 0);
+	return (Flags == NewFlags) || (fcntl(Socket, F_SETFL, NewFlags) >= 0);
 #endif
 }
 
