Fix a potential issue with the InstallCache directory for server when it is set via the '-ServerIOInstallCacheDir' cmdline option.

#rnx

- Normalize the provided direcotry path to fix incorrect directory seperators that may be present.
- Attempt to collapse any relative directories that might be in the directory path (/../)
- If the relative directories cannot be collapsed it means that the final directory path will end up outside of the PersistentDownloadDirectory which we don't want, in this case we will log an error and not use the invalid path.

diff --git a/Engine/Source/Runtime/Experimental/IoStore/OnDemand/Private/OnDemandConfig.cpp b/Engine/Source/Runtime/Experimental/IoStore/OnDemand/Private/OnDemandConfig.cpp
--- a/Engine/Source/Runtime/Experimental/IoStore/OnDemand/Private/OnDemandConfig.cpp
+++ b/Engine/Source/Runtime/Experimental/IoStore/OnDemand/Private/OnDemandConfig.cpp
@@ -366,11 +366,18 @@ static FString GetInstallCacheDirectory(const TCHAR* CommandLine)
 			bool bUsePathFromCommandLine = FParse::Value(CommandLine, TEXT("ServerIOInstallCacheDir="), CommandLineDir);
 			if (bUsePathFromCommandLine)
 			{
+				FPaths::NormalizeDirectoryName(CommandLineDir);
+
 				if (!FPaths::ValidatePath(CommandLineDir))
 				{
 					bUsePathFromCommandLine = false;
 					UE_LOG(LogIoStoreOnDemand, Error, TEXT("Invalid ServerIOInstallCacheDir from command line: %s"), *CommandLineDir);
 				}
+				else if (!FPaths::CollapseRelativeDirectories(CommandLineDir))
+				{
+					bUsePathFromCommandLine = false;
+					UE_LOG(LogIoStoreOnDemand, Error, TEXT("ServerIOInstallCacheDir from command line would end up outside of the PersistentDownloadDirectory: %s"), *CommandLineDir);
+				}
 				else if (!FPaths::IsRelative(CommandLineDir))
 				{
 					bUsePathFromCommandLine = false;
