Improve some IAX cvar help documentation.

#rnx

- Help documentation for 'iostore.DebugRequest' was a copy/paste and did not describe what the cvar does.
- The help documentation for 'iostore.ErrorOnRequest' implied that a random error would be introduced, but the error is always the same each time.

diff --git a/Engine/Source/Runtime/Experimental/IoStore/OnDemand/Private/OnDemandContentInstaller.cpp b/Engine/Source/Runtime/Experimental/IoStore/OnDemand/Private/OnDemandContentInstaller.cpp
--- a/Engine/Source/Runtime/Experimental/IoStore/OnDemand/Private/OnDemandContentInstaller.cpp
+++ b/Engine/Source/Runtime/Experimental/IoStore/OnDemand/Private/OnDemandContentInstaller.cpp
@@ -32,16 +32,16 @@ namespace UE::IoStore
 		static FAutoConsoleVariableRef CVar_IoStoreErrorOnRequest(
 			TEXT("iostore.ErrorOnRequest"),
 			IoStoreErrorOnRequest,
-			TEXT("When the request with a debug name partially matching this cvar is found iostore will error with a random error.")
+			TEXT("When the request with a debug name partially matching this cvar is found iostore will error with a debug 'InstallArgs' error.")
 			);
 
 		FString IoStoreDebugRequest;
 		static FAutoConsoleVariableRef CVar_IoStoreDebugRequest(
 			TEXT("iostore.DebugRequest"),
 			IoStoreDebugRequest,
-			TEXT("When the request with a debug name partially matching this cvar is found iostore will error with a random error.")
+			TEXT("When the request with a debug name partially matching this cvar is found iostore will print a message to the log on which a breakpoint can be placed.")
 		);
-#endif
+#endif // !UE_BUILD_SHIPPING
 
 		static bool bContentInstallerFlushOnIdle = false;
 		static FAutoConsoleVariableRef CVar_ContentInstallerFlushOnIdle(
