Break out time fetch code into free function.

#rnx

diff --git a/Engine/Source/Runtime/Experimental/IoStore/HttpClient/Private/Loop.inl b/Engine/Source/Runtime/Experimental/IoStore/HttpClient/Private/Loop.inl
--- a/Engine/Source/Runtime/Experimental/IoStore/HttpClient/Private/Loop.inl
+++ b/Engine/Source/Runtime/Experimental/IoStore/HttpClient/Private/Loop.inl
@@ -941,6 +941,25 @@ void FHostGroup::AddActivity(FActivityNode* Activity)
 ////////////////////////////////////////////////////////////////////////////////
 static const FEventLoop::FRequestParams GDefaultParams;
 
+////////////////////////////////////////////////////////////////////////////////
+static uint32 GetNowMs()
+{
+	// 4.2MM seconds will give us 50 days of uptime.
+	static uint64 Freq = 0;
+	static uint64 Base = 0;
+	if (Freq == 0)
+	{
+		Freq = uint64(1.0 / FPlatformTime::GetSecondsPerCycle64());
+		Base = FPlatformTime::Cycles64();
+	}
+	uint64 NowBig = ((FPlatformTime::Cycles64() - Base) * 1000) / Freq;
+	uint32 NowMs = uint32(NowBig);
+	check(NowMs == NowBig);
+	return NowMs;
+}
+
+
+
 ////////////////////////////////////////////////////////////////////////////////
 class FEventLoop::FImpl
 {
@@ -1120,28 +1139,13 @@ uint32 FEventLoop::FImpl::Tick(int32 PollTimeoutMs)
 
 	uint64 CancelsLoad = Cancels.load(std::memory_order_relaxed);
 
-	uint32 NowMs;
-	{
-		// 4.2MM seconds will give us 50 days of uptime.
-		static uint64 Freq = 0;
-		static uint64 Base = 0;
-		if (Freq == 0)
-		{
-			Freq = uint64(1.0 / FPlatformTime::GetSecondsPerCycle64());
-			Base = FPlatformTime::Cycles64();
-		}
-		uint64 NowBig = ((FPlatformTime::Cycles64() - Base) * 1000) / Freq;
-		NowMs = uint32(NowBig);
-		check(NowMs == NowBig);
-	}
-
 	// Tick groups and then remove ones that are idle
 	FTickState TickState = {
 		.Cancels = CancelsLoad,
 		.RecvAllowance = RecvAllowance,
 		.PollTimeoutMs = PollTimeoutMs,
 		.FailTimeoutMs = FailTimeoutMs,
-		.NowMs = NowMs,
+		.NowMs = GetNowMs(),
 	};
 	for (FHostGroup& Group : Groups)
 	{
