Fix incorrect code documentation.

#rnx

diff --git a/Engine/Source/Runtime/Experimental/IoStore/OnDemand/Private/OnDemandPackageStoreBackend.h b/Engine/Source/Runtime/Experimental/IoStore/OnDemand/Private/OnDemandPackageStoreBackend.h
--- a/Engine/Source/Runtime/Experimental/IoStore/OnDemand/Private/OnDemandPackageStoreBackend.h
+++ b/Engine/Source/Runtime/Experimental/IoStore/OnDemand/Private/OnDemandPackageStoreBackend.h
@@ -31,4 +31,4 @@ public:
 
 TSharedPtr<IOnDemandPackageStoreBackend> MakeOnDemandPackageStoreBackend(TWeakPtr<FOnDemandIoStore> OnDemandIoStore);
 
-} // namespace UE
+} // namespace UE::IoStore
