Chaos, Bubble resim: When bubble resim is on, only particles marked as FullResim can integrate. Also adding some scopes to better understand the cost of resims.

diff --git a/Engine/Source/Runtime/Experimental/Chaos/Private/Chaos/PBDRigidsEvolutionGBF.cpp b/Engine/Source/Runtime/Experimental/Chaos/Private/Chaos/PBDRigidsEvolutionGBF.cpp
--- a/Engine/Source/Runtime/Experimental/Chaos/Private/Chaos/PBDRigidsEvolutionGBF.cpp
+++ b/Engine/Source/Runtime/Experimental/Chaos/Private/Chaos/PBDRigidsEvolutionGBF.cpp
@@ -19,6 +19,7 @@
 #include "Chaos/PerParticlePBDUpdateFromDeltaPosition.h"
 #include "ChaosStats.h"
 #include "Chaos/EvolutionResimCache.h"
+#include "Chaos/ResimConsoleVariables.h"
 
 #include "ProfilingDebugging/ScopedTimers.h"
 #include "Chaos/DebugDrawQueue.h"
@@ -621,7 +622,10 @@ void FPBDRigidsEvolutionGBF::AdvanceOneTimeStepImpl(const FReal Dt, const FSubSt
 
 	// Collision detection is sensitive to duplication bugs in the particle views
 	// so this is here to help us track them down when they happen
-	CheckParticleViewsForErrors(Particles);
+	{
+		QUICK_SCOPE_CYCLE_COUNTER(Chaos_CheckParticleViewsForErrors);
+		CheckParticleViewsForErrors(Particles);
+	}
 
 	{
 		SCOPE_CYCLE_COUNTER(STAT_Evolution_DetectCollisions);
@@ -630,6 +634,7 @@ void FPBDRigidsEvolutionGBF::AdvanceOneTimeStepImpl(const FReal Dt, const FSubSt
 
 		{
 			CVD_SCOPE_TRACE_SOLVER_STEP(CVDDC_CollisionDetectionBroadPhase, TEXT("Collision Detection Broad Phase"));
+			QUICK_SCOPE_CYCLE_COUNTER(Chaos_RunBroadPhase);
 			CollisionDetector.RunBroadPhase(Dt, GetCurrentStepResimCache());
 		}
 
@@ -641,6 +646,7 @@ void FPBDRigidsEvolutionGBF::AdvanceOneTimeStepImpl(const FReal Dt, const FSubSt
 
 		{
 			CVD_SCOPE_TRACE_SOLVER_STEP(CVDDC_CollisionDetectionNarrowPhase, TEXT("Collision Detection Narrow Phase"));
+			QUICK_SCOPE_CYCLE_COUNTER(Chaos_RunNarrowPhase);
 			CollisionDetector.RunNarrowPhase(Dt, GetCurrentStepResimCache());
 		}
 	}
@@ -802,13 +808,18 @@ void FPBDRigidsEvolutionGBF::AdvanceOneTimeStepImpl(const FReal Dt, const FSubSt
 
 	if (CaptureRewindData)
 	{
+		QUICK_SCOPE_CYCLE_COUNTER(Chaos_CaptureRewindData);
 		CaptureRewindData(Particles.GetDirtyParticlesView());
 	}
 
-	ParticleUpdatePosition(Particles.GetDirtyParticlesView(), Dt);
+	{
+		QUICK_SCOPE_CYCLE_COUNTER(Chaos_ParticleUpdatePosition);
+		ParticleUpdatePosition(Particles.GetDirtyParticlesView(), Dt);
+	}
 
 	{
 		CSV_SCOPED_TIMING_STAT(PhysicsVerbose, StepSolver_IslandEndTick);
+		QUICK_SCOPE_CYCLE_COUNTER(Chaos_IslandsEndTick);
 		// Clean up the transient data from the constraint graph (e.g., Islands get cleared here)
 		GetIslandManager().EndTick();
 	}
@@ -825,6 +836,7 @@ void FPBDRigidsEvolutionGBF::AdvanceOneTimeStepImpl(const FReal Dt, const FSubSt
 		// Run contact updates on probe constraints
 		// NOTE: This happens after particles have been moved to their new locations
 		// so that we get contacts that are correct for end-of-frame positions
+		QUICK_SCOPE_CYCLE_COUNTER(Chaos_DetectProbeCollisions);
 		GetCollisionConstraints().DetectProbeCollisions(Dt);
 	}
 	
@@ -866,19 +878,26 @@ void FPBDRigidsEvolutionGBF::Integrate(FReal Dt)
 	const FReal HackMaxLinearSpeedSq = CVars::HackMaxVelocity * CVars::HackMaxVelocity;
 	const bool bAllowMACD = CVars::bChaosUseMACD;
 	const bool bForceMACD = CVars::bChaosForceMACD;
+	const bool bIsResimming = IsResimming();
+	const bool bIsBubbleResimEnabled = Chaos::ResimConsoleVars::bIsBubbleResimEnabled;
 
 	const TArray<FPBDRigidParticleHandle*>& ParticlesArray = Particles.GetActiveParticlesArray();
 
 	auto IntegrateWork = [this, &ParticlesArray,
 		Dt, HackMaxAngularSpeedSq, HackMaxLinearSpeedSq, bAllowMACD, bForceMACD, VelocityBoundsMultiplierMACD, MaxVelocityBoundsExpansionMACD,
-		VelocityBoundsMultiplier, MaxVelocityBoundsExpansion, BoundsThickness](int32 BatchIndex, int32 StartIndex, int32 EndIndex)
+		VelocityBoundsMultiplier, MaxVelocityBoundsExpansion, BoundsThickness, bIsResimming, bIsBubbleResimEnabled](int32 BatchIndex, int32 StartIndex, int32 EndIndex)
 	{
 		for (int32 Index = StartIndex; Index < EndIndex; Index++)
 		{
 			//question: can we enforce this at the API layer? Right now islands contain non dynamic which makes this hard
-			if (FPBDRigidParticleHandle* PBDParticle = ParticlesArray[Index]->CastToRigidParticle();  PBDParticle)
+			if (FPBDRigidParticleHandle* PBDParticle = ParticlesArray[Index]->CastToRigidParticle(); PBDParticle)
 			{
-				if (PBDParticle->ObjectState() == EObjectStateType::Dynamic)
+				const bool bIsParticleDynamic = (PBDParticle->ObjectState() == EObjectStateType::Dynamic);
+				const bool bIsParticleFrozenDuringBubbleResim = (PBDParticle->ResimType() == EResimType::FrozenDuringResim);
+				// We only integrate dynamic particles
+				// and when we are resimming, if bubble resim is enabled, we skip integrating particles frozen during resim
+				bool bShouldIntegrate = bIsParticleDynamic && !(bIsResimming && bIsBubbleResimEnabled && bIsParticleFrozenDuringBubbleResim);
+				if (bShouldIntegrate)
 				{
 					FPBDRigidParticleHandle& Particle = *PBDParticle;
 
