Collision Filter Refactor Part 3.6: Deprecating old Get/Set ShapeSim/QueryData methods to prevent new usage from being added.

diff --git a/Engine/Plugins/ChaosVD/Source/ChaosVD/Private/DetailsCustomizations/ChaosVDShapeDataCustomization.cpp b/Engine/Plugins/ChaosVD/Source/ChaosVD/Private/DetailsCustomizations/ChaosVDShapeDataCustomization.cpp
--- a/Engine/Plugins/ChaosVD/Source/ChaosVD/Private/DetailsCustomizations/ChaosVDShapeDataCustomization.cpp
+++ b/Engine/Plugins/ChaosVD/Source/ChaosVD/Private/DetailsCustomizations/ChaosVDShapeDataCustomization.cpp
@@ -251,7 +251,9 @@ FText FChaosVDShapeDataCustomization::GetCurrentCollisionChannelAsText(ECollisio
 	{
 		FChaosVDCollisionFilterData& FilterDataRef = FilterDataType == ECollisionFilterDataType::Query ? Data->QueryData : Data->SimData;
 		uint8 QueryExtraFilter = 0;
+		PRAGMA_DISABLE_DEPRECATION_WARNINGS
 		return UEnum::GetDisplayValueAsText(GetCollisionChannelAndExtraFilter(FilterDataRef.Word3, QueryExtraFilter));
+		PRAGMA_ENABLE_DEPRECATION_WARNINGS
 	}
 
 	return FText::GetEmpty();
@@ -263,7 +265,9 @@ uint8 FChaosVDShapeDataCustomization::GetCurrentExtraFilter(ECollisionFilterData
 	{
 		FChaosVDCollisionFilterData& FilterDataRef = FilterDataType== ECollisionFilterDataType::Query ? Data->QueryData : Data->SimData;
 		uint8 QueryExtraFilter = 0;
+		PRAGMA_DISABLE_DEPRECATION_WARNINGS
 		GetCollisionChannelAndExtraFilter(FilterDataRef.Word3, QueryExtraFilter);
+		PRAGMA_ENABLE_DEPRECATION_WARNINGS
 	
 		return QueryExtraFilter;
 	}
@@ -337,11 +341,13 @@ void FChaosVDShapeDataCustomization::CacheCurrentCollisionChannelData()
 {
 	if (FChaosVDShapeCollisionData* Data = CurrentShapeDataHandle ? CurrentShapeDataHandle->GetDataInstance() : nullptr)
 	{
+		PRAGMA_DISABLE_DEPRECATION_WARNINGS
 		FCollisionFilterData QueryFilterData = FChaosVDDetailsCustomizationUtils::ConvertToEngineFilteringData(Data->QueryData);
 		CachedQueryCollisionResponsesFlagsPerChannel = ExtractQueryCollisionResponseContainer(QueryFilterData);
 
 		FCollisionFilterData SimFilterData = FChaosVDDetailsCustomizationUtils::ConvertToEngineFilteringData(Data->SimData);
 		CachedSimCollisionResponsesFlagsPerChannel = ExtractSimCollisionResponseContainer(SimFilterData);
+		PRAGMA_ENABLE_DEPRECATION_WARNINGS
 	}
 }
 
diff --git a/Engine/Source/Runtime/Engine/Private/Collision/PhysicsFiltering.cpp b/Engine/Source/Runtime/Engine/Private/Collision/PhysicsFiltering.cpp
--- a/Engine/Source/Runtime/Engine/Private/Collision/PhysicsFiltering.cpp
+++ b/Engine/Source/Runtime/Engine/Private/Collision/PhysicsFiltering.cpp
@@ -172,3 +172,29 @@ ECollisionChannel GetCollisionChannel(const Chaos::Filter::FQueryFilterData& Que
 {
 	return (ECollisionChannel)QueryFilterData.GetCollisionChannelIndex();
 }
+
+FCollisionResponseContainer ExtractCollisionResponseContainer(const Chaos::Filter::FShapeFilterData& ShapeFilterData)
+{
+	FCollisionResponseContainer CollisionResponseContainer;
+	const uint64 BlockChannels = ShapeFilterData.GetQueryBlockChannels();
+	const uint64 OverlapChannels = ShapeFilterData.GetQueryOverlapChannels();
+
+	for (int32 ChannelIndex = 0; ChannelIndex < UE_ARRAY_COUNT(CollisionResponseContainer.EnumArray); ++ChannelIndex)
+	{
+		const uint64 ChannelMask = ((uint64)1 << ChannelIndex);
+		if (BlockChannels & ChannelMask)
+		{
+			CollisionResponseContainer.EnumArray[ChannelIndex] = ECR_Block;
+		}
+		else if (OverlapChannels & ChannelMask)
+		{
+			CollisionResponseContainer.EnumArray[ChannelIndex] = ECR_Overlap;
+		}
+		else
+		{
+			CollisionResponseContainer.EnumArray[ChannelIndex] = ECR_Ignore;
+		}
+	}
+
+	return CollisionResponseContainer;
+}
diff --git a/Engine/Source/Runtime/Engine/Private/PhysicsEngine/CollisionQueryFilterCallback.cpp b/Engine/Source/Runtime/Engine/Private/PhysicsEngine/CollisionQueryFilterCallback.cpp
--- a/Engine/Source/Runtime/Engine/Private/PhysicsEngine/CollisionQueryFilterCallback.cpp
+++ b/Engine/Source/Runtime/Engine/Private/PhysicsEngine/CollisionQueryFilterCallback.cpp
@@ -12,10 +12,12 @@ ECollisionQueryHitType FCollisionQueryFilterCallback::CalcQueryHitType(const FCo
 {
 	ECollisionQuery QueryType = (ECollisionQuery)QueryFilter.Word0;
 	FMaskFilter QuerierMaskFilter;
+	PRAGMA_DISABLE_DEPRECATION_WARNINGS
 	const ECollisionChannel QuerierChannel = GetCollisionChannelAndExtraFilter(QueryFilter.Word3, QuerierMaskFilter);
 
 	FMaskFilter ShapeMaskFilter;
 	const ECollisionChannel ShapeChannel = GetCollisionChannelAndExtraFilter(ShapeFilter.Word3, ShapeMaskFilter);
+	PRAGMA_ENABLE_DEPRECATION_WARNINGS
 
 	if ((QuerierMaskFilter & ShapeMaskFilter) != 0)	//If ignore mask hit something, ignore it
 	{
diff --git a/Engine/Source/Runtime/Engine/Public/Physics/PhysicsFiltering.h b/Engine/Source/Runtime/Engine/Public/Physics/PhysicsFiltering.h
--- a/Engine/Source/Runtime/Engine/Public/Physics/PhysicsFiltering.h
+++ b/Engine/Source/Runtime/Engine/Public/Physics/PhysicsFiltering.h
@@ -209,6 +209,7 @@ inline ECollisionChannel GetCollisionChannel(uint32 Word3)
 	return (ECollisionChannel)ChannelMask;
 }
 
+UE_DEPRECATED(5.8, "Use the methods on Chaos::Filter::FShapeFilterData instead")
 inline ECollisionChannel GetCollisionChannelAndExtraFilter(uint32 Word3, FMaskFilter& OutMaskFilter)
 {
 	uint32 ChannelMask = GetCollisionChannel(Word3);
@@ -216,12 +217,14 @@ inline ECollisionChannel GetCollisionChannelAndExtraFilter(uint32 Word3, FMaskFi
 	return (ECollisionChannel)ChannelMask;
 }
 
+UE_DEPRECATED(5.8, "Use FPhysicsFilterBuilder directly instead")
 inline uint32 CreateChannelAndFilter(ECollisionChannel CollisionChannel, FMaskFilter MaskFilter)
 {
 	uint32 ResultMask = (uint32(MaskFilter) << NumCollisionChannelBits) | (uint32)CollisionChannel;
 	return ResultMask << NumFilterDataFlagBits;
 }
 
+UE_DEPRECATED(5.8, "Use FPhysicsFilterBuilder directly instead")
 inline void UpdateMaskFilter(uint32& Word3, FMaskFilter NewMaskFilter)
 {
 	static_assert(NumExtraFilterBits <= 8, "Only up to 8 extra filter bits are supported.");
@@ -229,6 +232,7 @@ inline void UpdateMaskFilter(uint32& Word3, FMaskFilter NewMaskFilter)
 	Word3 |= uint32(NewMaskFilter) << (32 - NumExtraFilterBits);
 }
 
+UE_DEPRECATED(5.8, "Use ExtractCollisionResponseContainer instead")
 inline FCollisionResponseContainer ExtractSimCollisionResponseContainer(const FCollisionFilterData& InSimFilterData)
 {
 	FCollisionResponseContainer CollisionResponseContainer;
@@ -248,6 +252,7 @@ inline FCollisionResponseContainer ExtractSimCollisionResponseContainer(const FC
 	return CollisionResponseContainer;
 }
 
+UE_DEPRECATED(5.8, "Use ExtractCollisionResponseContainer instead")
 inline FCollisionResponseContainer ExtractQueryCollisionResponseContainer(const FCollisionFilterData& InQueryFilterData)
 {
 	FCollisionResponseContainer CollisionResponseContainer;
@@ -270,3 +275,5 @@ inline FCollisionResponseContainer ExtractQueryCollisionResponseContainer(const
 
 	return CollisionResponseContainer;
 }
+
+UE_INTERNAL ENGINE_API FCollisionResponseContainer ExtractCollisionResponseContainer(const Chaos::Filter::FShapeFilterData& ShapeFilterData);
diff --git a/Engine/Source/Runtime/Experimental/Chaos/Private/Chaos/Collision/SimRaycast.cpp b/Engine/Source/Runtime/Experimental/Chaos/Private/Chaos/Collision/SimRaycast.cpp
--- a/Engine/Source/Runtime/Experimental/Chaos/Private/Chaos/Collision/SimRaycast.cpp
+++ b/Engine/Source/Runtime/Experimental/Chaos/Private/Chaos/Collision/SimRaycast.cpp
@@ -108,7 +108,9 @@ namespace Chaos::Private
 
 		// Run the channel filter
 		// NOTE: PrePreSimFilterImp return true if we want to skip
+		PRAGMA_DISABLE_DEPRECATION_WARNINGS
 		return !PrePreSimFilterImp(FilterData, Shape->GetSimData());
+		PRAGMA_ENABLE_DEPRECATION_WARNINGS
 	}
 
 	bool SimRaycastBatchFirstHits(
diff --git a/Engine/Source/Runtime/Experimental/Chaos/Private/PhysicsProxy/ClusterUnionPhysicsProxy.cpp b/Engine/Source/Runtime/Experimental/Chaos/Private/PhysicsProxy/ClusterUnionPhysicsProxy.cpp
--- a/Engine/Source/Runtime/Experimental/Chaos/Private/PhysicsProxy/ClusterUnionPhysicsProxy.cpp
+++ b/Engine/Source/Runtime/Experimental/Chaos/Private/PhysicsProxy/ClusterUnionPhysicsProxy.cpp
@@ -64,8 +64,10 @@ namespace Chaos
 			for (const TUniquePtr<Chaos::FPerShapeData>& ShapeData : ShapeArray)
 			{
 				BufferData.CollisionData.Add(ShapeData->GetCollisionData());
-				BufferData.QueryData.Add(ShapeData->GetQueryData());
-				BufferData.SimData.Add(ShapeData->GetSimData());
+				FCollisionFilterData QueryData, SimData;
+				Chaos::Filter::FShapeFilterBuilder::GetLegacyShapeFilter(ShapeData->GetCombinedShapeFilterData(), QueryData, SimData);
+				BufferData.QueryData.Add(QueryData);
+				BufferData.SimData.Add(SimData);
 			}
 
 			if constexpr (std::is_base_of_v<FClusterUnionPhysicsProxy::FInternalParticle, TParticle>)
@@ -569,7 +571,7 @@ namespace Chaos
 				{
 					ShapeData->SetCollisionData(CurrentPullData.CollisionData[ShapeIndex]);
 				}
-
+				PRAGMA_DISABLE_DEPRECATION_WARNINGS
 				if (CurrentPullData.QueryData.IsValidIndex(ShapeIndex))
 				{
 					ShapeData->SetQueryData(CurrentPullData.QueryData[ShapeIndex]);
@@ -579,6 +581,7 @@ namespace Chaos
 				{
 					ShapeData->SetSimData(CurrentPullData.SimData[ShapeIndex]);
 				}
+				PRAGMA_ENABLE_DEPRECATION_WARNINGS
 			}
 		}
 
diff --git a/Engine/Source/Runtime/Experimental/Chaos/Private/PhysicsProxy/GeometryCollectionPhysicsProxy.cpp b/Engine/Source/Runtime/Experimental/Chaos/Private/PhysicsProxy/GeometryCollectionPhysicsProxy.cpp
--- a/Engine/Source/Runtime/Experimental/Chaos/Private/PhysicsProxy/GeometryCollectionPhysicsProxy.cpp
+++ b/Engine/Source/Runtime/Experimental/Chaos/Private/PhysicsProxy/GeometryCollectionPhysicsProxy.cpp
@@ -966,8 +966,8 @@ void FGeometryCollectionPhysicsProxy::CreateGTParticles(TManagedArray<Chaos::FIm
 			for (int32 ShapeIndex = 0; ShapeIndex < NumShapes; ++ShapeIndex)
 			{
 				Chaos::FPerShapeData* Shape = Shapes[ShapeIndex].Get();
-				Shape->SetSimData(SimFilter);
-				Shape->SetQueryData(QueryFilter);
+				const Chaos::Filter::FCombinedShapeFilterData CombinedShapeFilter = Chaos::Filter::FShapeFilterBuilder::BuildFromLegacyShapeFilter(QueryFilter, SimFilter);
+				Shape->SetCombinedShapeFilterData(CombinedShapeFilter);
 				Shape->SetProxy(this);
 				Shape->SetMaterial(Parameters.PhysicalMaterialHandle);
 			}
diff --git a/Engine/Source/Runtime/Experimental/Chaos/Public/Chaos/Collision/CollisionFilter.h b/Engine/Source/Runtime/Experimental/Chaos/Public/Chaos/Collision/CollisionFilter.h
--- a/Engine/Source/Runtime/Experimental/Chaos/Public/Chaos/Collision/CollisionFilter.h
+++ b/Engine/Source/Runtime/Experimental/Chaos/Public/Chaos/Collision/CollisionFilter.h
@@ -108,7 +108,9 @@ namespace Chaos
 	*/
 	inline bool FilterHasSimEnabled(const FPerShapeData* Shape)
 	{
+		PRAGMA_DISABLE_DEPRECATION_WARNINGS
 		return (!Shape || (Shape->GetSimEnabled() && IsFilterValid(Shape->GetSimData())));
+		PRAGMA_ENABLE_DEPRECATION_WARNINGS
 	}
 
 	inline bool DoCollide(EImplicitObjectType ImplicitType, const FPerShapeData* Shape)
@@ -154,7 +156,7 @@ namespace Chaos
 		//
 		if (Shape0 && Shape1)
 		{
-
+			PRAGMA_DISABLE_DEPRECATION_WARNINGS
 			if (IsFilterValid(Shape0->GetSimData()) && IsFilterValid(Shape1->GetSimData()))
 			{
 				FMaskFilter Filter0Mask, Filter1Mask;
@@ -165,6 +167,7 @@ namespace Chaos
 				uint32 const Filter0Bit = 1 << (Filter0Channel); // SIMDATA_TO_BITFIELD
 				return (Filter0Bit & Shape1->GetSimData().Word1) && (Filter1Bit & Shape0->GetSimData().Word1);
 			}
+			PRAGMA_ENABLE_DEPRECATION_WARNINGS
 		}
 
 		return true;
diff --git a/Engine/Source/Runtime/Experimental/Chaos/Public/Chaos/ParticleHandle.h b/Engine/Source/Runtime/Experimental/Chaos/Public/Chaos/ParticleHandle.h
--- a/Engine/Source/Runtime/Experimental/Chaos/Public/Chaos/ParticleHandle.h
+++ b/Engine/Source/Runtime/Experimental/Chaos/Public/Chaos/ParticleHandle.h
@@ -2907,13 +2907,16 @@ public:
 		}
 	}
 
+	UE_DEPRECATED(5.8, "Use SetShapeFilterData instead")
 	void SetShapeSimData(int32 InShapeIndex,const FCollisionFilterData& SimData)
 	{
+		PRAGMA_DISABLE_DEPRECATION_WARNINGS
 		const FCollisionFilterData& Current = MShapesArray[InShapeIndex]->GetSimData();
 		if(Current != SimData)
 		{
 			MShapesArray[InShapeIndex]->SetSimData(SimData);
 		}
+		PRAGMA_ENABLE_DEPRECATION_WARNINGS
 	}
 
 	void SetShapeFilterData(int32 InShapeIndex, const Chaos::Filter::FShapeFilterData& FilterData)
@@ -3979,7 +3982,9 @@ template <typename TParticle>
 		// For queries we can ignore any shapes that are not query-enabled (e.g., sim- and probe-only)
 		if (Shape->GetQueryEnabled())
 		{
+			PRAGMA_DISABLE_DEPRECATION_WARNINGS
 			const FCollisionFilterData& ShapeQueryData = Shape->GetQueryData();
+			PRAGMA_ENABLE_DEPRECATION_WARNINGS
 			OutQueryFilterData.Word0 |= ShapeQueryData.Word0;
 			OutQueryFilterData.Word1 |= ShapeQueryData.Word1;
 			OutQueryFilterData.Word2 |= ShapeQueryData.Word2;
@@ -3998,7 +4003,9 @@ template <typename TParticle>
 		// In the simulation we can ignore query-only shapes, but need to know about sim- and probe-enabled shapes
 		if (Shape->GetSimEnabled() || Shape->GetIsProbe())
 		{
+			PRAGMA_DISABLE_DEPRECATION_WARNINGS
 			const FCollisionFilterData& ShapeSimData = Shape->GetSimData();
+			PRAGMA_ENABLE_DEPRECATION_WARNINGS
 
 			// Word3 contains the channel index, this cannot be or'd together. Instead, expand the index to the bitmask and or that. 
 			// No one currently uses word2 (the component ID) which would also be invalid to or, so re-use that memory.
diff --git a/Engine/Source/Runtime/Experimental/Chaos/Public/Chaos/ShapeInstance.h b/Engine/Source/Runtime/Experimental/Chaos/Public/Chaos/ShapeInstance.h
--- a/Engine/Source/Runtime/Experimental/Chaos/Public/Chaos/ShapeInstance.h
+++ b/Engine/Source/Runtime/Experimental/Chaos/Public/Chaos/ShapeInstance.h
@@ -88,11 +88,11 @@ namespace Chaos
 		void* GetUserData() const;
 		void SetUserData(void* InUserData);
 
-		const FCollisionFilterData& GetQueryData() const;
-		void SetQueryData(const FCollisionFilterData& InQueryData);
+		UE_DEPRECATED(5.8, "Use GetCombinedShapeFilterData instead")  const FCollisionFilterData& GetQueryData() const;
+		UE_DEPRECATED(5.8, "Use SetCombinedShapeFilterData instead")  void SetQueryData(const FCollisionFilterData& InQueryData);
 
-		const FCollisionFilterData& GetSimData() const;
-		void SetSimData(const FCollisionFilterData& InSimData);
+		UE_DEPRECATED(5.8, "Use GetCombinedShapeFilterData instead") const FCollisionFilterData& GetSimData() const;
+		UE_DEPRECATED(5.8, "Use SetCombinedShapeFilterData instead") void SetSimData(const FCollisionFilterData& InSimData);
 
 		CHAOS_API const Filter::FShapeFilterData GetShapeFilterData() const;
 		CHAOS_API void SetShapeFilterData(const Filter::FShapeFilterData& InShapeFilter);
@@ -173,7 +173,9 @@ namespace Chaos
 		int32 GetShapeIndex() const;
 		void ModifyShapeIndex(int32 NewShapeIndex);
 
-		template <typename Lambda> void ModifySimData(const Lambda& LambdaFunc);
+		template <typename Lambda>
+		UE_DEPRECATED(5.8, "Get/SetCombinedShapeFilterData instead")
+		void ModifySimData(const Lambda& LambdaFunc);
 		template <typename Lambda> void ModifyMaterials(const Lambda& LambdaFunc);
 		template <typename Lambda> void ModifyMaterialMasks(const Lambda& LambdaFunc);
 		template <typename Lambda> void ModifyMaterialMaskMaps(const Lambda& LambdaFunc);
@@ -286,19 +288,19 @@ namespace Chaos
 			CollisionData.Modify(true, DirtyFlags, Proxy, ShapeIdx, [InUserData](FCollisionData& Data) { Data.UserData = InUserData; });
 		}
 
-		PRAGMA_DISABLE_INTERNAL_WARNINGS
-		const FCollisionFilterData& GetQueryData() const { return CollisionData.Read().GetQueryData(); }
-		void SetQueryData(const FCollisionFilterData& InQueryData)
+		PRAGMA_DISABLE_DEPRECATION_WARNINGS
+		UE_DEPRECATED(5.8, "Use GetCombinedShapeFilterData instead") const FCollisionFilterData& GetQueryData() const { return CollisionData.Read().GetQueryData(); }
+		UE_DEPRECATED(5.8, "Use SetCombinedShapeFilterData instead") void SetQueryData(const FCollisionFilterData& InQueryData)
 		{
 			CollisionData.Modify(true, DirtyFlags, Proxy, ShapeIdx, [InQueryData](FCollisionData& Data) { Data.SetQueryData(InQueryData); });
 		}
 
-		const FCollisionFilterData& GetSimData() const { return CollisionData.Read().GetSimData(); }
-		void SetSimData(const FCollisionFilterData& InSimData)
+		UE_DEPRECATED(5.8, "Use GetCombinedShapeFilterData instead") const FCollisionFilterData& GetSimData() const { return CollisionData.Read().GetSimData(); }
+		UE_DEPRECATED(5.8, "Use SetCombinedShapeFilterData instead") void SetSimData(const FCollisionFilterData& InSimData)
 		{
 			CollisionData.Modify(true, DirtyFlags, Proxy, ShapeIdx, [InSimData](FCollisionData& Data) { Data.SetSimData(InSimData); });
 		}
-		PRAGMA_ENABLE_INTERNAL_WARNINGS
+		PRAGMA_ENABLE_DEPRECATION_WARNINGS
 
 		CHAOS_API const Filter::FShapeFilterData GetShapeFilterData() const;
 		CHAOS_API void SetShapeFilterData(const Filter::FShapeFilterData& InShapeFilter);
@@ -573,13 +575,13 @@ namespace Chaos
 		void* GetUserData() const { return CollisionData.UserData; }
 		void SetUserData(void* InUserData) { CollisionData.UserData = InUserData; }
 
-		PRAGMA_DISABLE_INTERNAL_WARNINGS
-		const FCollisionFilterData& GetQueryData() const { return CollisionData.GetQueryData(); }
-		void SetQueryData(const FCollisionFilterData& InQueryData) { CollisionData.SetQueryData(InQueryData); }
+		PRAGMA_DISABLE_DEPRECATION_WARNINGS
+		UE_DEPRECATED(5.8, "Use GetCombinedShapeFilterData instead") const FCollisionFilterData& GetQueryData() const { return CollisionData.GetQueryData(); }
+		UE_DEPRECATED(5.8, "Use SetCombinedShapeFilterData instead") void SetQueryData(const FCollisionFilterData& InQueryData) { CollisionData.SetQueryData(InQueryData); }
 
-		const FCollisionFilterData& GetSimData() const { return CollisionData.GetSimData(); }
-		void SetSimData(const FCollisionFilterData& InSimData) { CollisionData.SetSimData(InSimData); }
-		PRAGMA_ENABLE_INTERNAL_WARNINGS
+		UE_DEPRECATED(5.8, "Use GetCombinedShapeFilterData instead") const FCollisionFilterData& GetSimData() const { return CollisionData.GetSimData(); }
+		UE_DEPRECATED(5.8, "Use SetCombinedShapeFilterData instead") void SetSimData(const FCollisionFilterData& InSimData) { CollisionData.SetSimData(InSimData); }
+		PRAGMA_ENABLE_DEPRECATION_WARNINGS
 
 		CHAOS_API const Filter::FShapeFilterData GetShapeFilterData() const;
 		CHAOS_API void SetShapeFilterData(const Filter::FShapeFilterData& InShapeFilter);
@@ -1024,22 +1026,30 @@ namespace Chaos
 
 	inline const FCollisionFilterData& FPerShapeData::GetQueryData() const
 	{
+		PRAGMA_DISABLE_DEPRECATION_WARNINGS
 		return DownCast([](const auto& ShapeInstance) -> const auto& { return ShapeInstance.GetQueryData(); });
+		PRAGMA_ENABLE_DEPRECATION_WARNINGS
 	}
 
 	inline void FPerShapeData::SetQueryData(const FCollisionFilterData& InQueryData)
 	{
+		PRAGMA_DISABLE_DEPRECATION_WARNINGS
 		DownCast([&InQueryData](auto& ShapeInstance) { ShapeInstance.SetQueryData(InQueryData); });
+		PRAGMA_ENABLE_DEPRECATION_WARNINGS
 	}
 
 	inline const FCollisionFilterData& FPerShapeData::GetSimData() const
 	{
+		PRAGMA_DISABLE_DEPRECATION_WARNINGS
 		return DownCast([](const auto& ShapeInstance) -> const auto& { return ShapeInstance.GetSimData(); });
+		PRAGMA_ENABLE_DEPRECATION_WARNINGS
 	}
 
 	inline void FPerShapeData::SetSimData(const FCollisionFilterData& InSimData)
 	{
+		PRAGMA_DISABLE_DEPRECATION_WARNINGS
 		return DownCast([&InSimData](auto& ShapeInstance) { ShapeInstance.SetSimData(InSimData); });
+		PRAGMA_ENABLE_DEPRECATION_WARNINGS
 	}
 
 	inline FImplicitObjectRef FPerShapeData::GetGeometry() const
@@ -1240,7 +1250,9 @@ namespace Chaos
 	template <typename Lambda> 
 	void FPerShapeData::ModifySimData(const Lambda& LambdaFunc)
 	{
+		PRAGMA_DISABLE_DEPRECATION_WARNINGS
 		DownCast([&LambdaFunc](auto& ShapeInstance) { ShapeInstance.ModifySimData(LambdaFunc); });
+		PRAGMA_ENABLE_DEPRECATION_WARNINGS
 	}
 
 	template <typename Lambda> 
diff --git a/Engine/Source/Runtime/Experimental/Chaos/Public/PhysicsProxy/SingleParticlePhysicsProxy.h b/Engine/Source/Runtime/Experimental/Chaos/Public/PhysicsProxy/SingleParticlePhysicsProxy.h
--- a/Engine/Source/Runtime/Experimental/Chaos/Public/PhysicsProxy/SingleParticlePhysicsProxy.h
+++ b/Engine/Source/Runtime/Experimental/Chaos/Public/PhysicsProxy/SingleParticlePhysicsProxy.h
@@ -1170,7 +1170,14 @@ public:
 
 	void SetShapeSimCollisionEnabled(int32 InShapeIndex, bool bInEnabled) { VerifyContext(); GetParticle_LowLevel()->SetShapeSimCollisionEnabled(InShapeIndex, bInEnabled); }
 	void SetShapeQueryCollisionEnabled(int32 InShapeIndex, bool bInEnabled) { VerifyContext(); GetParticle_LowLevel()->SetShapeQueryCollisionEnabled(InShapeIndex, bInEnabled); }
-	void SetShapeSimData(int32 InShapeIndex, const FCollisionFilterData& SimData) { VerifyContext(); GetParticle_LowLevel()->SetShapeSimData(InShapeIndex, SimData); }
+	UE_DEPRECATED(5.8, "Use SetShapeFilterData instead")
+	void SetShapeSimData(int32 InShapeIndex, const FCollisionFilterData& SimData)
+	{
+		PRAGMA_DISABLE_DEPRECATION_WARNINGS
+		VerifyContext();
+		GetParticle_LowLevel()->SetShapeSimData(InShapeIndex, SimData);
+		PRAGMA_ENABLE_DEPRECATION_WARNINGS
+	}
 	CHAOS_API void SetShapeFilterData(int32 InShapeIndex, const Filter::FShapeFilterData& FilterData);
 	CHAOS_API void SetShapeFilterInstanceData(int32 InShapeIndex, const Filter::FInstanceData& InstanceData);
 
diff --git a/Engine/Source/Runtime/Experimental/Chaos/Public/RewindData.h b/Engine/Source/Runtime/Experimental/Chaos/Public/RewindData.h
--- a/Engine/Source/Runtime/Experimental/Chaos/Public/RewindData.h
+++ b/Engine/Source/Runtime/Experimental/Chaos/Public/RewindData.h
@@ -1002,9 +1002,9 @@ struct FPerShapeDataStateBase
 	template <typename TParticle>
 	static const FCollisionFilterData& GetQueryData(const FPerShapeDataStateBase* State, const TParticle& Particle, int32 ShapeIdx)
 	{
-		PRAGMA_DISABLE_INTERNAL_WARNINGS
+		PRAGMA_DISABLE_DEPRECATION_WARNINGS
 		return State && State->CollisionData.IsSet() ? State->CollisionData.Read().GetQueryData() : Particle.ShapesArray()[ShapeIdx]->GetQueryData();
-		PRAGMA_ENABLE_INTERNAL_WARNINGS
+		PRAGMA_ENABLE_DEPRECATION_WARNINGS
 	}
 };
 
