Gather some stats as the loop ticks.

#rnx

diff --git a/Engine/Source/Runtime/Experimental/IoStore/HttpClient/Private/Loop.inl b/Engine/Source/Runtime/Experimental/IoStore/HttpClient/Private/Loop.inl
--- a/Engine/Source/Runtime/Experimental/IoStore/HttpClient/Private/Loop.inl
+++ b/Engine/Source/Runtime/Experimental/IoStore/HttpClient/Private/Loop.inl
@@ -5,6 +5,105 @@
 namespace UE::IoStore::HTTP
 {
 
+// {{{1 recv-stats .............................................................
+
+////////////////////////////////////////////////////////////////////////////////
+class FRecvStats
+{
+public:
+							FRecvStats() = default;
+	void					Update(uint32 NowMs, uint32 RecvSizeKiB);
+	void					Flush();
+
+private:
+	enum
+	{
+		QuantizeBy			= 255,
+		BucketMs			= 1000,
+	};
+
+	union FBucket
+	{
+		struct FData
+		{
+			uint32			RecvKiB;
+			uint32			TimeMs;
+		};
+		std::atomic<uint64> Value;
+	};
+
+	FBucket					Buckets[8] = {};
+	FBucket::FData			Accumulator = {};
+	std::atomic<uint32>		TotalRecvKiB = 0;
+	uint32					LastMs = 0;
+	uint8					Cursor = 0;
+};
+
+////////////////////////////////////////////////////////////////////////////////
+void FRecvStats::Update(uint32 NowMs, uint32 RecvSizeKiB)
+{
+	if (RecvSizeKiB == 0)
+	{
+		return;
+	}
+
+	TotalRecvKiB.fetch_add(RecvSizeKiB, std::memory_order_relaxed);
+
+	// Quantize time for a modicum of stability.
+	static_assert(((QuantizeBy + 1) & QuantizeBy) == 0);
+	// NowMs = NowMs & ~QuantizeBy;
+
+	uint32 DeltaMs = NowMs - LastMs;
+	if (DeltaMs == 0)
+	{
+		return;
+	}
+	LastMs = NowMs;
+
+	Accumulator.TimeMs += DeltaMs;
+	if (Accumulator.TimeMs < BucketMs)
+	{
+		Accumulator.RecvKiB += RecvSizeKiB;
+		return;
+	}
+
+	uint32 RecvKiB = Accumulator.RecvKiB;
+	uint32 TimeMs = NowMs - DeltaMs;
+	if (Accumulator.TimeMs >= BucketMs)
+	{
+		Accumulator = {
+			.RecvKiB = RecvSizeKiB,
+			.TimeMs = 0,
+		};
+	}
+	else
+	{
+		uint32 Half = RecvSizeKiB; // (RecvSizeKiB + 1) >> 1;
+		Accumulator = {
+			.RecvKiB = RecvSizeKiB - Half,
+			.TimeMs = BucketMs - Accumulator.TimeMs,
+		};
+		RecvKiB += Half;
+		TimeMs = NowMs - Accumulator.TimeMs;
+	}
+
+	uint64 Value = (uint64(TimeMs) << 32u) + RecvKiB;
+
+	Cursor += 1;
+	uint32 Index = (Cursor - 1) % UE_ARRAY_COUNT(Buckets);
+	Buckets[Index].Value.store(Value, std::memory_order_release);
+}
+
+////////////////////////////////////////////////////////////////////////////////
+void FRecvStats::Flush()
+{
+	uint64 Value = (uint64(LastMs) << 32u) + Accumulator.RecvKiB;
+	uint32 Index = Cursor % UE_ARRAY_COUNT(Buckets);
+	Buckets[Index].Value.store(Value, std::memory_order_release);
+}
+
+
+
 // {{{1 work-queue .............................................................
 
 /*
@@ -321,6 +420,7 @@ struct FTickState
 	int32						PollTimeoutMs;
 	int32						FailTimeoutMs;
 	uint32						NowMs;
+	uint32						RecvSizeKiB = 0;
 	class FWorkQueue*			Work;
 };
 
@@ -926,6 +1026,12 @@ void FHostGroup::Tick(FTickState& State)
 	{
 		BusyCount += (Group.Tick(State) == true);
 	}
+
+	// Gather how much each peer has recv'd since the last time
+	for (FPeerGroup& Group : PeerGroups)
+	{
+		State.RecvSizeKiB += Group.CollectRecvSizeKiB();
+	}
 }
 
 ////////////////////////////////////////////////////////////////////////////////
@@ -975,6 +1081,7 @@ public:
 
 private:
 	void					ReceiveWork();
+	FRecvStats				RecvStats;
 	FCriticalSection		Lock;
 	std::atomic<uint64>		FreeSlots		= ~0ull;
 	std::atomic<uint64>		Cancels			= 0;
@@ -1139,7 +1246,7 @@ uint32 FEventLoop::FImpl::Tick(int32 PollTimeoutMs)
 
 	uint64 CancelsLoad = Cancels.load(std::memory_order_relaxed);
 
-	// Tick groups and then remove ones that are idle
+	// Tick groups
 	FTickState TickState = {
 		.Cancels = CancelsLoad,
 		.RecvAllowance = RecvAllowance,
@@ -1152,6 +1259,10 @@ uint32 FEventLoop::FImpl::Tick(int32 PollTimeoutMs)
 		Group.Tick(TickState);
 	}
 
+	// Track some stats
+	RecvStats.Update(TickState.NowMs, TickState.RecvSizeKiB);
+
+	// Remove groups that are idle
 	for (uint32 i = 0, n = Groups.Num(); i < n; ++i)
 	{
 		FHostGroup& Group = Groups[i];
@@ -1186,12 +1297,23 @@ uint32 FEventLoop::FImpl::Tick(int32 PollTimeoutMs)
 		PrevFreeSlots += ReturnedSlots;
 	}
 
+	// To cover the scenario where one of the completed slots gets cancelled
+	// after the earlier load but before now,  we will unset returned slots
+	// from the cancel mask too.
 	if (uint64 Mask = CancelsLoad | ReturnedSlots; Mask)
 	{
 		Cancels.fetch_and(~Mask, std::memory_order_relaxed);
 	}
 
-	return BusyCount + BusyBias;
+	// There is still some accounting in the stats object but there is no reason
+	// to get any further ticks. Make sure the accounting is externally visible.
+	uint32 Result = BusyCount + BusyBias;
+	if (Result == 0)
+	{
+		RecvStats.Flush();
+	}
+
+	return Result;
 }
 
 
