Surface stats to the public

#rnx

diff --git a/Engine/Source/Runtime/Experimental/IoStore/HttpClient/Internal/IO/Http/Client.h b/Engine/Source/Runtime/Experimental/IoStore/HttpClient/Internal/IO/Http/Client.h
--- a/Engine/Source/Runtime/Experimental/IoStore/HttpClient/Internal/IO/Http/Client.h
+++ b/Engine/Source/Runtime/Experimental/IoStore/HttpClient/Internal/IO/Http/Client.h
@@ -56,6 +56,7 @@ enum class EStatusCodeClass
 using	FCertRootsRef	= UPTRINT;
 using	FTicket			= uint64;
 struct	FActivityNode;
+class	FRecvStats;
 
 ////////////////////////////////////////////////////////////////////////////////
 class FCertRoots
@@ -222,6 +223,24 @@ private:
 
 using FTicketSink = TFunction<void (const FTicketStatus&)>;
 
+////////////////////////////////////////////////////////////////////////////////
+class FLoopPerf
+{
+public:
+	enum { RecvBucketCount = 8 };
+
+	UE_API	uint32	GetTotalRecvMiB() const;
+	UE_API	uint32	GetRecvKiBpsCount() const;
+	UE_API	uint32	GetRecvKiBps(uint32 Index) const;
+	UE_API	void	GetRecvKiBps(uint32 (&Out)[RecvBucketCount]) const;
+	UE_API	uint32	GetRecvKiBps() const;
+
+private:
+	friend			class FRecvStats;
+	uint32			RecvMiB;
+	uint16			RecvKiBps[RecvBucketCount];
+};
+
 ////////////////////////////////////////////////////////////////////////////////
 class FEventLoop
 {
@@ -245,6 +264,7 @@ public:
 
 	UE_API 							FEventLoop();
 	UE_API 							~FEventLoop();
+	UE_API void						SamplePerf(FLoopPerf& Perf) const;
 	UE_API uint32					Tick(int32 PollTimeoutMs=0);
 	UE_API void						Throttle(uint32 KiBPerSec);
 	UE_API void						SetFailTimeout(int32 TimeoutMs);
diff --git a/Engine/Source/Runtime/Experimental/IoStore/HttpClient/Private/Loop.inl b/Engine/Source/Runtime/Experimental/IoStore/HttpClient/Private/Loop.inl
--- a/Engine/Source/Runtime/Experimental/IoStore/HttpClient/Private/Loop.inl
+++ b/Engine/Source/Runtime/Experimental/IoStore/HttpClient/Private/Loop.inl
@@ -14,6 +14,7 @@ public:
 							FRecvStats() = default;
 	void					Update(uint32 NowMs, uint32 RecvSizeKiB);
 	void					Flush();
+	void					Sample(uint32 NowMs, FLoopPerf& Out) const;
 
 private:
 	enum
@@ -102,6 +103,113 @@ void FRecvStats::Flush()
 	Buckets[Index].Value.store(Value, std::memory_order_release);
 }
 
+////////////////////////////////////////////////////////////////////////////////
+void FRecvStats::Sample(uint32 NowMs, FLoopPerf& Out) const
+{
+	static_assert((BucketMs * UE_ARRAY_COUNT(Buckets)) / 1000 == UE_ARRAY_COUNT(Out.RecvKiBps));
+
+	// NowMs = (NowMs + QuantizeBy) & ~QuantizeBy;
+
+	union FData
+	{
+		FBucket::FData	Data;
+		uint64			Value;
+	};
+
+	// Sample all the buckets
+	FData Datas[UE_ARRAY_COUNT(Buckets)];
+	for (uint32 i = UE_ARRAY_COUNT(Buckets); i--;)
+	{
+		Datas[i].Value = Buckets[i].Value.load(std::memory_order_acquire);
+	}
+
+	// Find the newest bucket
+	int32 Index = -1;
+	for (uint32 i = 0, PrevMs = 0u; const FData& Data : Datas)
+	{
+		if (uint32 TimeMs = Data.Data.TimeMs; TimeMs >= PrevMs)
+		{
+			Index = int32(i);
+			PrevMs = TimeMs;
+		}
+		++i;
+	}
+
+	// Gather samples for the last N seconds.
+	for (uint32 i = 0, n = UE_ARRAY_COUNT(Out.RecvKiBps), TimeMs = NowMs; i < n; ++i)
+	{
+		const FBucket::FData& Data = Datas[Index].Data;
+
+		TimeMs -= BucketMs;
+		if (TimeMs > Data.TimeMs)
+		{
+			Out.RecvKiBps[i] = 0;
+			continue;
+		}
+
+		uint32 Value = Data.RecvKiB;
+		Value = (Value > 0xffff) ? 0xffff : Value;
+		Out.RecvKiBps[i] = uint16(Value);
+
+		Index = (Index - 1) % UE_ARRAY_COUNT(Buckets);
+	}
+
+	Out.RecvMiB = TotalRecvKiB.load(std::memory_order_relaxed) >> 10;
+}
+
+
+
+////////////////////////////////////////////////////////////////////////////////
+uint32 FLoopPerf::GetTotalRecvMiB() const
+{
+	return RecvMiB;
+}
+
+////////////////////////////////////////////////////////////////////////////////
+uint32 FLoopPerf::GetRecvKiBpsCount() const
+{
+	return UE_ARRAY_COUNT(RecvKiBps);
+}
+
+////////////////////////////////////////////////////////////////////////////////
+uint32 FLoopPerf::GetRecvKiBps(uint32 Index) const
+{
+	check(Index < GetRecvKiBpsCount());
+	return RecvKiBps[Index];
+}
+
+////////////////////////////////////////////////////////////////////////////////
+void FLoopPerf::GetRecvKiBps(uint32 (&Out)[RecvBucketCount]) const
+{
+	static_assert(UE_ARRAY_COUNT(Out) == UE_ARRAY_COUNT(RecvKiBps));
+
+	for (uint32 i = 0; uint32& Item : Out)
+	{
+		Item = GetRecvKiBps(i);
+		++i;
+	}
+}
+
+////////////////////////////////////////////////////////////////////////////////
+uint32 FLoopPerf::GetRecvKiBps() const
+{
+	// EMA to bias towards most recent data
+
+	uint32 Denom = 0;
+	uint32 Numer = 0;
+	for (uint32 i = RecvBucketCount; uint32 Sample : RecvKiBps)
+	{
+		Numer += Sample * i;
+		Denom += i;
+		if (i >>= 1; i == 0)
+		{
+			break;
+		}
+	}
+
+	return Numer / Denom;
+}
+
 
 
 // {{{1 work-queue .............................................................
@@ -1071,6 +1179,7 @@ class FEventLoop::FImpl
 {
 public:
 							~FImpl();
+	void					SamplePerf(FLoopPerf& Perf) const;
 	uint32					Tick(int32 PollTimeoutMs=0);
 	bool					IsIdle() const;
 	void					Throttle(uint32 KiBPerSec);
@@ -1099,6 +1208,13 @@ FEventLoop::FImpl::~FImpl()
 	check(BusyCount == 0);
 }
 
+////////////////////////////////////////////////////////////////////////////////
+void FEventLoop::FImpl::SamplePerf(FLoopPerf& Out) const
+{
+	uint32 NowMs = GetNowMs();
+	RecvStats.Sample(NowMs, Out);
+}
+
 ////////////////////////////////////////////////////////////////////////////////
 FRequest FEventLoop::FImpl::Request(FActivityNode* Activity)
 {
@@ -1321,6 +1437,7 @@ uint32 FEventLoop::FImpl::Tick(int32 PollTimeoutMs)
 ////////////////////////////////////////////////////////////////////////////////
 FEventLoop::FEventLoop()						{ Impl = new FEventLoop::FImpl(); Trace(Impl, ETrace::LoopCreate); }
 FEventLoop::~FEventLoop()						{ Trace(Impl, ETrace::LoopDestroy); delete Impl; }
+void FEventLoop::SamplePerf(FLoopPerf& Out) const { return Impl->SamplePerf(Out); }
 uint32 FEventLoop::Tick(int32 PollTimeoutMs)	{ return Impl->Tick(PollTimeoutMs); }
 bool FEventLoop::IsIdle() const					{ return Impl->IsIdle(); }
 void FEventLoop::Cancel(FTicket Ticket)			{ return Impl->Cancel(Ticket); }
