Lawn consistency

#rnx

diff --git a/Engine/Source/Runtime/Experimental/IoStore/OnDemand/Private/JournaledCache.cpp b/Engine/Source/Runtime/Experimental/IoStore/OnDemand/Private/JournaledCache.cpp
--- a/Engine/Source/Runtime/Experimental/IoStore/OnDemand/Private/JournaledCache.cpp
+++ b/Engine/Source/Runtime/Experimental/IoStore/OnDemand/Private/JournaledCache.cpp
@@ -1787,6 +1787,7 @@ void FServiceThread::UnregisterCache(FCache* Cache)
 	}
 }
 
+////////////////////////////////////////////////////////////////////////////////
 void FServiceThread::SetGovernor(FDiskCacheGovernor* Gov)
 {
 	FWork Work;
