[Chaos] Fixed HeadlessChaos after implicit include removal. Added deprecated include blocks to support upgrade path.

diff --git a/Engine/Source/Programs/HeadlessChaos/Private/HeadlessChaosTestSleeping.cpp b/Engine/Source/Programs/HeadlessChaos/Private/HeadlessChaosTestSleeping.cpp
--- a/Engine/Source/Programs/HeadlessChaos/Private/HeadlessChaosTestSleeping.cpp
+++ b/Engine/Source/Programs/HeadlessChaos/Private/HeadlessChaosTestSleeping.cpp
@@ -7,6 +7,7 @@
 #include "PhysicsProxy/JointConstraintProxy.h"
 #include "Chaos/PBDJointConstraintData.h"
 #include "HAL/ConsoleManager.h"
+#include "PhysicsInterfaceDeclaresCore.h"
 
 namespace ChaosTest
 {
diff --git a/Engine/Source/Programs/HeadlessChaos/Public/GeometryCollection/GeometryCollectionTestFramework.h b/Engine/Source/Programs/HeadlessChaos/Public/GeometryCollection/GeometryCollectionTestFramework.h
--- a/Engine/Source/Programs/HeadlessChaos/Public/GeometryCollection/GeometryCollectionTestFramework.h
+++ b/Engine/Source/Programs/HeadlessChaos/Public/GeometryCollection/GeometryCollectionTestFramework.h
@@ -9,6 +9,7 @@
 #include "Chaos/PBDRigidsEvolutionFwd.h"
 #include "Chaos/Defines.h"
 #include "PhysicsInterfaceDeclaresCore.h"
+#include "PhysicsInterfaceTypesCore.h"
 
 namespace GeometryCollectionTest
 {
diff --git a/Engine/Source/Runtime/Engine/Classes/PhysicsEngine/BodySetup.h b/Engine/Source/Runtime/Engine/Classes/PhysicsEngine/BodySetup.h
--- a/Engine/Source/Runtime/Engine/Classes/PhysicsEngine/BodySetup.h
+++ b/Engine/Source/Runtime/Engine/Classes/PhysicsEngine/BodySetup.h
@@ -15,6 +15,12 @@
 #include "HAL/ThreadSafeBool.h"
 #include "Async/TaskGraphInterfaces.h"
 #include "BodySetupCore.h"
+
+#if UE_ENABLE_INCLUDE_ORDER_DEPRECATED_IN_5_8
+#include "Factories.h"
+#include "PhysicsEngine/BodySetupObjectTextFactory.h"
+#endif // UE_ENABLE_INCLUDE_ORDER_DEPRECATED_IN_5_8
+
 #include "BodySetup.generated.h"
 
 class ITargetPlatform;
diff --git a/Engine/Source/Runtime/Experimental/Chaos/Public/Chaos/ParticleHandle.h b/Engine/Source/Runtime/Experimental/Chaos/Public/Chaos/ParticleHandle.h
--- a/Engine/Source/Runtime/Experimental/Chaos/Public/Chaos/ParticleHandle.h
+++ b/Engine/Source/Runtime/Experimental/Chaos/Public/Chaos/ParticleHandle.h
@@ -17,6 +17,11 @@
 #include "Chaos/ChaosDebugDraw.h"
 #endif
 
+#if UE_ENABLE_INCLUDE_ORDER_DEPRECATED_IN_5_8
+#include "Chaos/ChaosUserEntity.h"
+#include "PhysicsInterfaceTypesCore.h"
+#endif // UE_ENABLE_INCLUDE_ORDER_DEPRECATED_IN_5_8
+
 #ifndef UE_DEBUG_DANGLING_HANDLES
 #define UE_DEBUG_DANGLING_HANDLES 0 // Will deliberately cause a memory leak
 #endif
