IAX backends (both old and new) now handle invalid requests with invalid FIoReadOptions ranges.

#rnx

- If a request with an invalid range (offset past the end of the chunk for example) the old backend would try to read the entire chunk before failing and the new backend would just assert/crash.
- Now both backends will fail the request with the EIoErrorCode::InvalidParameter status.
- Fixed the old backend to correctly process reads that go past the end of the chunk by capping the read at the end. This mirrors the behavior of the other IoDispatcher backends.

diff --git a/Engine/Source/Runtime/Experimental/IoStore/OnDemand/Private/OnDemandHttpIoDispatcherBackend.cpp b/Engine/Source/Runtime/Experimental/IoStore/OnDemand/Private/OnDemandHttpIoDispatcherBackend.cpp
--- a/Engine/Source/Runtime/Experimental/IoStore/OnDemand/Private/OnDemandHttpIoDispatcherBackend.cpp
+++ b/Engine/Source/Runtime/Experimental/IoStore/OnDemand/Private/OnDemandHttpIoDispatcherBackend.cpp
@@ -92,7 +92,7 @@ private:
 	FOnDemandIoBackendStats	Stats;
 	FOnDemandIoStore&		IoStore;
 	FIoRequestList			CompletedRequests;
-	FMutex					Mutex;
+	FMutex					CompletedRequestMutex;
 	FSharedBackendCtx		BackendContext;
 };
 
@@ -154,6 +154,21 @@ void FHttpIoDispatcherBackend::ResolveIoRequests(FIoRequestList Requests, FIoReq
 			continue;
 		}
 
+		if (Request->Options.GetOffset() > ChunkInfo.RawSize())
+		{
+			// Mark the request as failed (FIoRequestImpl::SetFailed) but with a specific error code.
+			Request->SetLastBackendError(EIoErrorCode::InvalidParameter);
+			Stats.OnIoRequestError();
+
+			{
+				TUniqueLock _(CompletedRequestMutex);
+				CompletedRequests.AddTail(Request);
+			}
+
+			BackendContext->WakeUpDispatcherThreadDelegate.Execute();
+			continue;
+		}
+
 		const uint64 ResolvedSize = FMath::Min<uint64>(Request->Options.GetSize(), ChunkInfo.RawSize() - Request->Options.GetOffset());
 		Request->Options.SetRange(Request->Options.GetOffset(), ResolvedSize);
 		
@@ -195,13 +210,18 @@ FIoRequestImpl* FHttpIoDispatcherBackend::GetCompletedIoRequests()
 {
 	FIoRequestList Out;
 	{
-		TUniqueLock Lock(Mutex);
+		TUniqueLock Lock(CompletedRequestMutex);
 		Out = MoveTemp(CompletedRequests);
 		CompletedRequests = FIoRequestList();
 	}
 
 	for (FIoRequestImpl& R : Out)
 	{
+		if (R.BackendData == nullptr)
+		{
+			continue;
+		}
+
 		FResolvedRequest::Destroy(reinterpret_cast<FResolvedRequest*>(R.BackendData));
 		check(R.BackendData == nullptr);
 	}
@@ -320,7 +340,7 @@ void FHttpIoDispatcherBackend::CompleteRequest(FResolvedRequest& ResolvedRequest
 	LogRequest(ResolvedRequest, DispatcherRequest, bDecoded, HttpResponse.IsCached(), bCancelled, DurationMs);
 
 	{
-		TUniqueLock Lock(Mutex);
+		TUniqueLock Lock(CompletedRequestMutex);
 		CompletedRequests.AddTail(&DispatcherRequest);
 	}
 
diff --git a/Engine/Source/Runtime/Experimental/IoStore/OnDemand/Private/OnDemandIoDispatcherBackend.cpp b/Engine/Source/Runtime/Experimental/IoStore/OnDemand/Private/OnDemandIoDispatcherBackend.cpp
--- a/Engine/Source/Runtime/Experimental/IoStore/OnDemand/Private/OnDemandIoDispatcherBackend.cpp
+++ b/Engine/Source/Runtime/Experimental/IoStore/OnDemand/Private/OnDemandIoDispatcherBackend.cpp
@@ -856,7 +856,7 @@ void FOnDemandIoBackend::CompleteRequest(FChunkRequest* ChunkRequest)
 		if (Chunk.GetSize() > 0)
 		{
 			TRACE_CPUPROFILER_EVENT_SCOPE(IasBackend::DecodeBlocks);
-			const uint64 RawSize = FMath::Min<uint64>(Request->Options.GetSize(), ChunkRequest->Params.ChunkInfo.RawSize());
+			const uint64 RawSize = FMath::Min<uint64>(Request->Options.GetSize(), ChunkRequest->Params.ChunkInfo.RawSize() - Request->Options.GetOffset());
 			Request->CreateBuffer(RawSize);
 			DecodingParams.RawOffset = Request->Options.GetOffset();
 
@@ -976,6 +976,17 @@ bool FOnDemandIoBackend::Resolve(FIoRequestImpl* Request)
 		return false;
 	}
 
+	if (Request->Options.GetOffset() > ChunkInfo.RawSize())
+	{
+		// Mark the request as failed (FIoRequestImpl::SetFailed) but with a specific error code.
+		Request->SetLastBackendError(EIoErrorCode::InvalidParameter);
+		CompletedRequests.Enqueue(Request);
+
+		BackendContext->WakeUpDispatcherThreadDelegate.Execute();
+
+		return true;
+	}
+
 	FChunkRequestParams RequestParams = FChunkRequestParams::Create(Request, ChunkInfo);
 
 	if (BackendStatus.IsHttpEnabled(Request->ChunkId.GetChunkType()) == false || ChunkInfo.HostGroup().IsConnected() == false)
@@ -1147,6 +1158,11 @@ FIoRequestImpl* FOnDemandIoBackend::GetCompletedIoRequests()
 
 	for (FIoRequestImpl* It = Requests; It != nullptr; It = It->NextRequest)
 	{
+		if (It->BackendData == nullptr)
+		{
+			continue;
+		}
+
 		TUniquePtr<FBackendData> BackendData = FBackendData::Detach(It);
 		check(It->BackendData == nullptr);
 	}
