Development CDNs are now opt in for IAX rather than opt out.

#rnx

- By default any CDN url that is classified as a development CDN will be automatically excluded from HostGroups, meaning they will not be used.
-- To be classified as a development CDN the url must match the pattern found at 'ini:Engine:[IAS]:DevelopmentCDNPattern'
- A new cmdline 'Iax.EnableDevCDNs' has been added. When this is found on the command line, development CDNs will be allowed.
- The old opt out cmdline 'Ias.SkipDevCDNs' has been removed as this is now the default behavior.

diff --git a/Engine/Source/Runtime/Experimental/IoStore/OnDemandCore/Private/OnDemandHostGroup.cpp b/Engine/Source/Runtime/Experimental/IoStore/OnDemandCore/Private/OnDemandHostGroup.cpp
--- a/Engine/Source/Runtime/Experimental/IoStore/OnDemandCore/Private/OnDemandHostGroup.cpp
+++ b/Engine/Source/Runtime/Experimental/IoStore/OnDemandCore/Private/OnDemandHostGroup.cpp
@@ -36,15 +36,17 @@ static bool ShouldSkipUrl(const FString& Url)
 	{
 		FDevCDNCheck()
 		{
-			if (FParse::Param(FCommandLine::Get(), TEXT("Ias.SkipDevCDNs")))
+			if (FParse::Param(FCommandLine::Get(), TEXT("Iax.EnableDevCDNs")))
 			{
-				FString ConfigRegex;
-				GConfig->GetString(TEXT("Ias"), TEXT("DevelopmentCDNPattern"), ConfigRegex, GEngineIni);
+				return;
+			}
+
+			FString ConfigRegex;
+			GConfig->GetString(TEXT("Ias"), TEXT("DevelopmentCDNPattern"), ConfigRegex, GEngineIni);
 
-				if (!ConfigRegex.IsEmpty())
-				{
-					Pattern = MakeUnique<FRegexPattern>(MoveTemp(ConfigRegex));
-				}
+			if (!ConfigRegex.IsEmpty())
+			{
+				Pattern = MakeUnique<FRegexPattern>(MoveTemp(ConfigRegex));
 			}
 		}
 
