Consistent units

#rnx

diff --git a/Engine/Source/Runtime/Experimental/IoStore/HttpClient/Internal/IO/Http/Client.h b/Engine/Source/Runtime/Experimental/IoStore/HttpClient/Internal/IO/Http/Client.h
--- a/Engine/Source/Runtime/Experimental/IoStore/HttpClient/Internal/IO/Http/Client.h
+++ b/Engine/Source/Runtime/Experimental/IoStore/HttpClient/Internal/IO/Http/Client.h
@@ -229,7 +229,7 @@ class FLoopPerf
 public:
 	enum { RecvBucketCount = 8 };
 
-	UE_API	uint32	GetTotalRecvMiB() const;
+	UE_API	uint32	GetTotalRecvKiB() const;
 	UE_API	uint32	GetRecvKiBpsCount() const;
 	UE_API	uint32	GetRecvKiBps(uint32 Index) const;
 	UE_API	void	GetRecvKiBps(uint32 (&Out)[RecvBucketCount]) const;
@@ -237,7 +237,7 @@ public:
 
 private:
 	friend			class FRecvStats;
-	uint32			RecvMiB;
+	uint32			RecvKiB;
 	uint16			RecvKiBps[RecvBucketCount];
 };
 
diff --git a/Engine/Source/Runtime/Experimental/IoStore/HttpClient/Private/Loop.inl b/Engine/Source/Runtime/Experimental/IoStore/HttpClient/Private/Loop.inl
--- a/Engine/Source/Runtime/Experimental/IoStore/HttpClient/Private/Loop.inl
+++ b/Engine/Source/Runtime/Experimental/IoStore/HttpClient/Private/Loop.inl
@@ -154,15 +154,15 @@ void FRecvStats::Sample(uint32 NowMs, FLoopPerf& Out) const
 		Index = (Index - 1) % UE_ARRAY_COUNT(Buckets);
 	}
 
-	Out.RecvMiB = TotalRecvKiB.load(std::memory_order_relaxed) >> 10;
+	Out.RecvKiB = TotalRecvKiB.load(std::memory_order_relaxed);
 }
 
 
 
 ////////////////////////////////////////////////////////////////////////////////
-uint32 FLoopPerf::GetTotalRecvMiB() const
+uint32 FLoopPerf::GetTotalRecvKiB() const
 {
-	return RecvMiB;
+	return RecvKiB;
 }
 
 ////////////////////////////////////////////////////////////////////////////////
