Surface stats out of the IAx module.

#rnx

diff --git a/Engine/Source/Runtime/Experimental/IoStore/OnDemand/Private/OnDemandHttpClient.cpp b/Engine/Source/Runtime/Experimental/IoStore/OnDemand/Private/OnDemandHttpClient.cpp
--- a/Engine/Source/Runtime/Experimental/IoStore/OnDemand/Private/OnDemandHttpClient.cpp
+++ b/Engine/Source/Runtime/Experimental/IoStore/OnDemand/Private/OnDemandHttpClient.cpp
@@ -8,7 +8,7 @@
 #include "HAL/FileManager.h"
 #include "HAL/IConsoleManager.h"
 #include "HAL/PlatformTime.h"
-#include "IO/IoStoreOnDemand.h"
+#include "IO/IoStoreOnDemandInternals.h"
 #include "Misc/ConfigCacheIni.h"
 #include "Misc/Paths.h"
 #include "Statistics.h"
@@ -357,6 +357,29 @@ void FMultiEndpointHttpClient::UpdateConnections()
 	}
 }
 
+void FMultiEndpointHttpClient::GetHttpStats(FOnDemandHttpStats& Out) const
+{
+	using namespace HTTP;
+
+	struct FLoopHttpStats
+		: public IOnDemandInternalHttpStats
+	{
+		static_assert(SAMPLE_COUNT == FLoopPerf::RecvBucketCount);
+		virtual uint32 GetRecvKiBps() const { return Perf.GetRecvKiBps(); }
+		virtual void GetRecvKiBps(uint32 (&Out)[SAMPLE_COUNT]) const { return Perf.GetRecvKiBps(Out); }
+		virtual uint32 GetTotalRecvKiB() const { return Perf.GetTotalRecvKiB(); }
+		FLoopPerf Perf;
+	};
+
+	if (!Out.Internal.IsValid())
+	{
+		Out.Internal = MakeUnique<FLoopHttpStats>();
+	}
+
+	auto* HttpStats = static_cast<FLoopHttpStats*>(Out.Internal.Get());
+	EventLoop.SamplePerf(HttpStats->Perf);
+}
+
 FMultiEndpointHttpClient::FHttpTicketId FMultiEndpointHttpClient::IssueRequest(FRequest&& Request)
 {
 	using namespace UE::IoStore::HTTP;
diff --git a/Engine/Source/Runtime/Experimental/IoStore/OnDemand/Private/OnDemandHttpClient.h b/Engine/Source/Runtime/Experimental/IoStore/OnDemand/Private/OnDemandHttpClient.h
--- a/Engine/Source/Runtime/Experimental/IoStore/OnDemand/Private/OnDemandHttpClient.h
+++ b/Engine/Source/Runtime/Experimental/IoStore/OnDemand/Private/OnDemandHttpClient.h
@@ -27,7 +27,7 @@
 namespace UE::IoStore
 {
 
-namespace HTTP { class FTicketStatus; }
+class FOnDemandHttpStats;
 
 /** Utility function to join two FIoBuffers together and return the combined result as a new FIoBuffer */
 FIoBuffer JoinIoBuffers(const FIoBuffer& LHS, const FIoBuffer& RHS);
@@ -151,6 +151,7 @@ public:
 	void CancelRequest(FHttpTicketId TicketId);
 
 	void UpdateConnections();
+	void GetHttpStats(FOnDemandHttpStats& Out) const;
 
 private:
 	using FHttpConnectionPools = TArray<TUniquePtr<HTTP::FConnectionPool>, TInlineAllocator<4>>;
diff --git a/Engine/Source/Runtime/Experimental/IoStore/OnDemand/Private/OnDemandHttpIoDispacher.cpp b/Engine/Source/Runtime/Experimental/IoStore/OnDemand/Private/OnDemandHttpIoDispacher.cpp
--- a/Engine/Source/Runtime/Experimental/IoStore/OnDemand/Private/OnDemandHttpIoDispacher.cpp
+++ b/Engine/Source/Runtime/Experimental/IoStore/OnDemand/Private/OnDemandHttpIoDispacher.cpp
@@ -658,6 +658,7 @@ private:
 	virtual void					Tick() override;
 	// IHttpIoDispatcher
 	virtual void					Shutdown() override;
+	virtual void					GetHttpStats(FOnDemandHttpStats& Out) const override;
 	virtual FIoStatus				RegisterHostGroup(const FName& HostGroupName, TConstArrayView<FAnsiString> HostNames, FAnsiStringView TestUrl) override;
 	virtual bool					IsHostGroupRegistered(const FName& HostGroup) override;
 	virtual bool					IsHostGroupOk(const FName& HostGroup) override;
@@ -722,8 +723,6 @@ private:
 #endif // UE_ALLOW_DROP_CACHE
 };
 
-static TUniquePtr<FHttpDispatcher> HttpDispatcher;
-
 ////////////////////////////////////////////////////////////////////////////////
 FHttpDispatcher::FHttpDispatcher(TUniquePtr<IIasCache>&& InCache)
 	: Cache(MoveTemp(InCache))
@@ -851,6 +850,11 @@ void FHttpDispatcher::Shutdown()
 	}
 }
 
+void FHttpDispatcher::GetHttpStats(FOnDemandHttpStats& Out) const
+{
+	return HttpClient->GetHttpStats(Out);
+}
+
 FIoStatus FHttpDispatcher::RegisterHostGroup(const FName& HostGroupName, TConstArrayView<FAnsiString> HostNames, FAnsiStringView TestUrl)
 {
 	FIASHostGroup Existing = FHostGroupManager::Get().Find(HostGroupName);
diff --git a/Engine/Source/Runtime/Experimental/IoStore/OnDemand/Private/OnDemandHttpThread.cpp b/Engine/Source/Runtime/Experimental/IoStore/OnDemand/Private/OnDemandHttpThread.cpp
--- a/Engine/Source/Runtime/Experimental/IoStore/OnDemand/Private/OnDemandHttpThread.cpp
+++ b/Engine/Source/Runtime/Experimental/IoStore/OnDemand/Private/OnDemandHttpThread.cpp
@@ -321,6 +321,11 @@ FOnDemandHttpThread::~FOnDemandHttpThread()
 	DrainHttpRequests();
 }
 
+void FOnDemandHttpThread::GetHttpStats(FOnDemandHttpStats& Out) const
+{
+	return HttpClient->GetHttpStats(Out);
+}
+
 FOnDemandHttpThread::FRequestHandle FOnDemandHttpThread::IssueRequest(const FOnDemandChunkInfo& ChunkInfo, const FIoOffsetAndLength& ReadRange, int32 Priority, FCompletionCallback&& CompletionCallback, EHttpRequestType Type)
 {
 	return IssueRequest(FOnDemandChunkInfo(ChunkInfo), ReadRange, Priority, MoveTemp(CompletionCallback), Type);
diff --git a/Engine/Source/Runtime/Experimental/IoStore/OnDemand/Private/OnDemandHttpThread.h b/Engine/Source/Runtime/Experimental/IoStore/OnDemand/Private/OnDemandHttpThread.h
--- a/Engine/Source/Runtime/Experimental/IoStore/OnDemand/Private/OnDemandHttpThread.h
+++ b/Engine/Source/Runtime/Experimental/IoStore/OnDemand/Private/OnDemandHttpThread.h
@@ -26,6 +26,7 @@ namespace UE::IoStore
 class FMultiEndpointHttpClient;
 struct FHttpRequest;
 struct FOnDemandChunkInfo;
+class FOnDemandHttpStats;
 
 extern int32 GIasHttpHealthCheckWaitTime;
 extern int32 GOnDemandBackendThreadPriorityIndex;
@@ -72,6 +73,8 @@ public:
 	virtual ~FOnDemandHttpThread();
 
 	static FOnDemandHttpThread& Get();
+
+	void GetHttpStats(FOnDemandHttpStats& Out) const;
 	
 	/**
 	 * Issue a http request to read from a chunk.
diff --git a/Engine/Source/Runtime/Experimental/IoStore/OnDemand/Private/OnDemandIoStore.cpp b/Engine/Source/Runtime/Experimental/IoStore/OnDemand/Private/OnDemandIoStore.cpp
--- a/Engine/Source/Runtime/Experimental/IoStore/OnDemand/Private/OnDemandIoStore.cpp
+++ b/Engine/Source/Runtime/Experimental/IoStore/OnDemand/Private/OnDemandIoStore.cpp
@@ -2241,6 +2241,18 @@ void FOnDemandIoStore::SetStreamingOptions(EOnDemandStreamingOptions Options)
 	}
 }
 
+void FOnDemandIoStore::GetHttpStats(FOnDemandHttpStats& Out) const
+{
+	if (FHttpIoDispatcher::IsInitialized())
+	{
+		FHttpIoDispatcher::GetHttpStats(Out);
+	}
+	else
+	{
+		HttpClient->GetHttpStats(Out);
+	}
+}
+
 void FOnDemandIoStore::ReportAnalytics(TArray<FAnalyticsEventAttribute>& OutAnalyticsArray) const
 {
 	if (Installer.IsValid())
diff --git a/Engine/Source/Runtime/Experimental/IoStore/OnDemand/Private/OnDemandIoStore.h b/Engine/Source/Runtime/Experimental/IoStore/OnDemand/Private/OnDemandIoStore.h
--- a/Engine/Source/Runtime/Experimental/IoStore/OnDemand/Private/OnDemandIoStore.h
+++ b/Engine/Source/Runtime/Experimental/IoStore/OnDemand/Private/OnDemandIoStore.h
@@ -377,6 +377,7 @@ public:
 	virtual void						DumpMountedContainersToLog() const override;
 	virtual bool						IsOnDemandStreamingEnabled() const override;
 	virtual void						SetStreamingOptions(EOnDemandStreamingOptions Options) override;
+	virtual void						GetHttpStats(FOnDemandHttpStats& Out) const override;
 	virtual void						ReportAnalytics(TArray<FAnalyticsEventAttribute>& OutAnalyticsArray) const override;
 	virtual TUniquePtr<IAnalyticsRecording> StartAnalyticsRecording() const override;
 	virtual void						OnImmediateAnalytic(FOnDemandImmediateAnalyticHandler EventHandler) override;
diff --git a/Engine/Source/Runtime/Experimental/IoStore/OnDemandCore/Internal/IO/IoStoreOnDemandInternals.h b/Engine/Source/Runtime/Experimental/IoStore/OnDemandCore/Internal/IO/IoStoreOnDemandInternals.h
--- a/Engine/Source/Runtime/Experimental/IoStore/OnDemandCore/Internal/IO/IoStoreOnDemandInternals.h
+++ b/Engine/Source/Runtime/Experimental/IoStore/OnDemandCore/Internal/IO/IoStoreOnDemandInternals.h
@@ -71,6 +71,17 @@ public:
 	std::atomic<EStatus>	Status{EStatus::Pending};
 };
 
+////////////////////////////////////////////////////////////////////////////////
+class IOnDemandInternalHttpStats
+{
+public:
+	enum { SAMPLE_COUNT = FOnDemandHttpStats::SAMPLE_COUNT };
+	virtual ~IOnDemandInternalHttpStats() {}
+	virtual uint32 GetRecvKiBps() const = 0;
+	virtual void GetRecvKiBps(uint32 (&Out)[SAMPLE_COUNT]) const = 0;
+	virtual uint32 GetTotalRecvKiB() const = 0;
+};
+
 } // namespace UE::IoStore
 
 ////////////////////////////////////////////////////////////////////////////////////////////////////////////////////////////////////////////////////////////////
diff --git a/Engine/Source/Runtime/Experimental/IoStore/OnDemandCore/Private/HttpIoDispatcher.cpp b/Engine/Source/Runtime/Experimental/IoStore/OnDemandCore/Private/HttpIoDispatcher.cpp
--- a/Engine/Source/Runtime/Experimental/IoStore/OnDemandCore/Private/HttpIoDispatcher.cpp
+++ b/Engine/Source/Runtime/Experimental/IoStore/OnDemandCore/Private/HttpIoDispatcher.cpp
@@ -251,6 +251,11 @@ FIoStatus FHttpIoDispatcher::Shutdown()
 	return FIoStatus::Ok;
 }
 
+void FHttpIoDispatcher::GetHttpStats(FOnDemandHttpStats& Out)
+{
+	return GHttpIoDispatcher->GetHttpStats(Out);
+}
+
 FIoStatus FHttpIoDispatcher::RegisterHostGroup(const FName& HostGroup, TConstArrayView<FAnsiString> HostNames, FAnsiStringView TestUrl)
 {
 	if (GHttpIoDispatcher.IsValid())
diff --git a/Engine/Source/Runtime/Experimental/IoStore/OnDemandCore/Private/OnDemandHttpStats.cpp b/Engine/Source/Runtime/Experimental/IoStore/OnDemandCore/Private/OnDemandHttpStats.cpp
new file mode 100644
--- /dev/null
+++ b/Engine/Source/Runtime/Experimental/IoStore/OnDemandCore/Private/OnDemandHttpStats.cpp
@@ -0,0 +1,34 @@
+// Copyright Epic Games, Inc. All Rights Reserved.
+
+#include <IO/IoStoreOnDemand.h>
+#include <IO/IoStoreOnDemandInternals.h>
+
+namespace UE
+{
+
+namespace IoStore
+{
+
+FOnDemandHttpStats::FOnDemandHttpStats() = default;
+FOnDemandHttpStats::~FOnDemandHttpStats() = default;
+
+uint32 FOnDemandHttpStats::GetRecvKiBps() const
+{
+	return Internal.IsValid() ? Internal->GetRecvKiBps() : 0;
+}
+
+void FOnDemandHttpStats::GetRecvKiBps(uint32 (&Out)[SAMPLE_COUNT]) const
+{
+	if (Internal.IsValid())
+	{
+		Internal->GetRecvKiBps(Out);
+	}
+}
+
+uint32 FOnDemandHttpStats::GetTotalRecvKiB() const
+{
+	return Internal.IsValid() ? Internal->GetTotalRecvKiB() : 0;
+}
+
+} // namespace IoStore
+} // namespace UE
diff --git a/Engine/Source/Runtime/Experimental/IoStore/OnDemandCore/Public/IO/HttpIoDispatcher.h b/Engine/Source/Runtime/Experimental/IoStore/OnDemandCore/Public/IO/HttpIoDispatcher.h
--- a/Engine/Source/Runtime/Experimental/IoStore/OnDemandCore/Public/IO/HttpIoDispatcher.h
+++ b/Engine/Source/Runtime/Experimental/IoStore/OnDemandCore/Public/IO/HttpIoDispatcher.h
@@ -20,6 +20,9 @@ UE_API DECLARE_LOG_CATEGORY_EXTERN(LogHttpIoDispatcher, Log, All);
 namespace UE
 {
 
+namespace IoStore { class FOnDemandHttpStats; }
+using FOnDemandHttpStats = IoStore::FOnDemandHttpStats;
+
 /** Internal request handle type. */
 using FIoHttpRequestHandle = UPTRINT;
 
@@ -357,6 +360,7 @@ public:
 	UE_API static FIoStatus			Initialize(TSharedPtr<class IHttpIoDispatcher> Dispatcher);
 	UE_API static FIoStatus			Shutdown();
 
+	UE_API static void				GetHttpStats(FOnDemandHttpStats& Out);
 	UE_API static FIoStatus			RegisterHostGroup(const FName& HostGroup, TConstArrayView<FAnsiString> HostNames, FAnsiStringView TestUrl);
 	UE_API static FIoStatus			RegisterHostGroup(const FName& HostGroup, FAnsiStringView HostName, FAnsiStringView TestUrl);
 	UE_API static bool				IsHostGroupRegistered(const FName& HostGroup);
@@ -393,6 +397,7 @@ public:
 	virtual							~IHttpIoDispatcher() = default;
 	virtual void					Shutdown() = 0;
 	
+	virtual void					GetHttpStats(FOnDemandHttpStats& Out) const = 0;
 	virtual FIoStatus				RegisterHostGroup(const FName& HostGroup, TConstArrayView<FAnsiString> HostNames, FAnsiStringView TestUrl) = 0;
 	virtual bool					IsHostGroupRegistered(const FName& HostGroup) = 0;
 	virtual bool					IsHostGroupOk(const FName& HostGroup) = 0;
diff --git a/Engine/Source/Runtime/Experimental/IoStore/OnDemandCore/Public/IO/IoStoreOnDemand.h b/Engine/Source/Runtime/Experimental/IoStore/OnDemandCore/Public/IO/IoStoreOnDemand.h
--- a/Engine/Source/Runtime/Experimental/IoStore/OnDemandCore/Public/IO/IoStoreOnDemand.h
+++ b/Engine/Source/Runtime/Experimental/IoStore/OnDemandCore/Public/IO/IoStoreOnDemand.h
@@ -527,6 +527,22 @@ struct FOnDemandVerifyCacheResult
 /** Verify cache completion callback. */
 using FOnDemandVerifyCacheCompleted = TUniqueFunction<void(FOnDemandVerifyCacheResult)>;
 
+/** Container for http stats */
+class FOnDemandHttpStats
+{
+public:
+	enum { SAMPLE_COUNT = 8 };
+	UE_API FOnDemandHttpStats();
+	UE_API ~FOnDemandHttpStats();
+	UE_API uint32 GetRecvKiBps() const;
+	UE_API void GetRecvKiBps(uint32 (&Out)[SAMPLE_COUNT]) const;
+	UE_API uint32 GetTotalRecvKiB() const;
+
+private:
+	friend class FMultiEndpointHttpClient;
+	TUniquePtr<class IOnDemandInternalHttpStats> Internal;
+};
+
 /** Interface for recording analytics over a given time period */
 class IAnalyticsRecording
 {
@@ -634,6 +650,8 @@ public:
 	virtual bool IsOnDemandStreamingEnabled() const = 0;
 	/** Set streaming options.*/
 	virtual void SetStreamingOptions(EOnDemandStreamingOptions Options) = 0;
+	/** Returns statistics about the network */
+	virtual void GetHttpStats(FOnDemandHttpStats& Out) const = 0;
 	/**
 	 * Reports the statistics for the current OnDemandBackend. It will record any events that have occurred from the last time that this was called.
 	 * This is a legacy method, prefer using StartAnalyticsRecording instead.
