Geometry collection : add support for setting linear and angular drag at runtime

diff --git a/Engine/Source/Runtime/Experimental/Chaos/Private/PhysicsProxy/GeometryCollectionPhysicsProxy.cpp b/Engine/Source/Runtime/Experimental/Chaos/Private/PhysicsProxy/GeometryCollectionPhysicsProxy.cpp
--- a/Engine/Source/Runtime/Experimental/Chaos/Private/PhysicsProxy/GeometryCollectionPhysicsProxy.cpp
+++ b/Engine/Source/Runtime/Experimental/Chaos/Private/PhysicsProxy/GeometryCollectionPhysicsProxy.cpp
@@ -3770,6 +3770,50 @@ void FGeometryCollectionPhysicsProxy::SetPhysicsMaterial_Internal(const Chaos::F
 	}
 }
 
+void FGeometryCollectionPhysicsProxy::SetLinearDamping_External(float InDamping)
+{
+	// update physics thread
+	ExecuteOnPhysicsThread(*this, [this, InDamping]()
+		{
+			SetLinearDamping_Internal(InDamping);
+		});
+}
+
+void FGeometryCollectionPhysicsProxy::SetLinearDamping_Internal(float InDamping)
+{
+	Parameters.LinearDamping = InDamping;
+
+	for (Chaos::FPBDRigidClusteredParticleHandle* Handle : SolverParticleHandles)
+	{
+		if (Handle)
+		{
+			Handle->SetLinearEtherDrag(Parameters.LinearDamping);
+		}
+	}
+}
+
+void FGeometryCollectionPhysicsProxy::SetAngularDamping_External(float InDamping)
+{
+	// update physics thread
+	ExecuteOnPhysicsThread(*this, [this, InDamping]()
+		{
+			SetAngularDamping_Internal(InDamping);
+		});
+}
+
+void FGeometryCollectionPhysicsProxy::SetAngularDamping_Internal(float InDamping)
+{
+	Parameters.AngularDamping = InDamping;
+
+	for (Chaos::FPBDRigidClusteredParticleHandle* Handle : SolverParticleHandles)
+	{
+		if (Handle)
+		{
+			Handle->SetAngularEtherDrag(Parameters.AngularDamping);
+		}
+	}
+}
+
 void FGeometryCollectionPhysicsProxy::PushStateOnGameThread(Chaos::FPBDRigidsSolver* InSolver)
 {
 	// CONTEXT: ANYTHREAD but spawned form GAMETHREAD in a parallelFor
diff --git a/Engine/Source/Runtime/Experimental/Chaos/Public/PhysicsProxy/GeometryCollectionPhysicsProxy.h b/Engine/Source/Runtime/Experimental/Chaos/Public/PhysicsProxy/GeometryCollectionPhysicsProxy.h
--- a/Engine/Source/Runtime/Experimental/Chaos/Public/PhysicsProxy/GeometryCollectionPhysicsProxy.h
+++ b/Engine/Source/Runtime/Experimental/Chaos/Public/PhysicsProxy/GeometryCollectionPhysicsProxy.h
@@ -487,6 +487,8 @@ public:
 	CHAOS_API void ApplyBreakingAngularVelocity_External(FGeometryCollectionItemIndex ItemIndex, const FVector& AngularVelocity);
 	CHAOS_API void ApplyLinearVelocity_External(FGeometryCollectionItemIndex ItemIndex, const FVector& LinearVelocity);
 	CHAOS_API void ApplyAngularVelocity_External(FGeometryCollectionItemIndex ItemIndex, const FVector& AngularVelocity);
+	CHAOS_API void SetLinearDamping_External(float InDamping);
+	CHAOS_API void SetAngularDamping_External(float InDamping);
 
 	CHAOS_API void SetProxyDirty_External();
 
@@ -651,6 +653,8 @@ protected:
 	CHAOS_API void SetGravityGroupIndex_Internal(int32 GravityGroupIndex);
 	CHAOS_API void SetOneWayInteractionLevel_Internal(int32 InOneWayInteractionLevel);
 	CHAOS_API void SetPhysicsMaterial_Internal(const Chaos::FMaterialHandle& MaterialHandle);
+	CHAOS_API void SetLinearDamping_Internal(float InDamping);
+	CHAOS_API void SetAngularDamping_Internal(float InDamping);
 
 private:
 
diff --git a/Engine/Source/Runtime/Experimental/GeometryCollectionEngine/Private/GeometryCollection/GeometryCollectionComponent.cpp b/Engine/Source/Runtime/Experimental/GeometryCollectionEngine/Private/GeometryCollection/GeometryCollectionComponent.cpp
--- a/Engine/Source/Runtime/Experimental/GeometryCollectionEngine/Private/GeometryCollection/GeometryCollectionComponent.cpp
+++ b/Engine/Source/Runtime/Experimental/GeometryCollectionEngine/Private/GeometryCollection/GeometryCollectionComponent.cpp
@@ -7328,6 +7328,24 @@ void UGeometryCollectionComponent::SetPhysMaterialOverride(UPhysicalMaterial* Ne
 	}
 }
 
+void UGeometryCollectionComponent::SetLinearDamping(float InDamping)
+{
+	Super::SetLinearDamping(InDamping);
+	if (PhysicsProxy)
+	{
+		PhysicsProxy->SetLinearDamping_External(InDamping);
+	}
+}
+
+void UGeometryCollectionComponent::SetAngularDamping(float InDamping)
+{
+	Super::SetAngularDamping(InDamping);
+	if (PhysicsProxy)
+	{
+		PhysicsProxy->SetAngularDamping_External(InDamping);
+	}
+}
+
 void UGeometryCollectionComponent::FlushNetDormancyIfNeeded() const
 {
 	if (GetDesiredNetAwakeningMode() != ENetAwakeningMode::FlushNetDormancy)
diff --git a/Engine/Source/Runtime/Experimental/GeometryCollectionEngine/Public/GeometryCollection/GeometryCollectionComponent.h b/Engine/Source/Runtime/Experimental/GeometryCollectionEngine/Public/GeometryCollection/GeometryCollectionComponent.h
--- a/Engine/Source/Runtime/Experimental/GeometryCollectionEngine/Public/GeometryCollection/GeometryCollectionComponent.h
+++ b/Engine/Source/Runtime/Experimental/GeometryCollectionEngine/Public/GeometryCollection/GeometryCollectionComponent.h
@@ -665,6 +665,9 @@ public:
 	GEOMETRYCOLLECTIONENGINE_API virtual void AddTorqueInRadians(FVector Torque, FName BoneName = NAME_None, bool bAccelChange = false) override;
 	GEOMETRYCOLLECTIONENGINE_API virtual void PostLoad() override;
 	GEOMETRYCOLLECTIONENGINE_API virtual void SetPhysMaterialOverride(UPhysicalMaterial* NewPhysMaterial) override;
+	GEOMETRYCOLLECTIONENGINE_API virtual void SetLinearDamping(float InDamping) override;
+	GEOMETRYCOLLECTIONENGINE_API virtual void SetAngularDamping(float InDamping) override;
+
 protected:
 	GEOMETRYCOLLECTIONENGINE_API virtual void OnComponentCollisionSettingsChanged(bool bUpdateOverlaps=true) override;
 	GEOMETRYCOLLECTIONENGINE_API virtual bool CanBeUsedInPhysicsReplication(const FName BoneName = NAME_None) const override;
