Fix potential issue if all urls are filtered from a hostgroup.

#rnx

- Return an error if a host group ends up filtering out all of it's formerly valid urls to prevent an empty hostgroup from being created.

diff --git a/Engine/Source/Runtime/Experimental/IoStore/OnDemandCore/Private/OnDemandHostGroup.cpp b/Engine/Source/Runtime/Experimental/IoStore/OnDemandCore/Private/OnDemandHostGroup.cpp
--- a/Engine/Source/Runtime/Experimental/IoStore/OnDemandCore/Private/OnDemandHostGroup.cpp
+++ b/Engine/Source/Runtime/Experimental/IoStore/OnDemandCore/Private/OnDemandHostGroup.cpp
@@ -248,6 +248,12 @@ TIoStatusOr<FOnDemandHostGroup> FOnDemandHostGroup::Create(TConstArrayView<FStri
 
 		Impl->HostUrls.Emplace(MoveTemp(AnsiUrl));
 	}
+
+	if (Impl->HostUrls.IsEmpty())
+	{
+		return FIoStatus(EIoErrorCode::InvalidParameter, TEXT("Alls URLs were filtered out"));
+	}
+
 	Impl->PrimaryIndex = 0;
 
 	return FOnDemandHostGroup(MoveTemp(Impl));
