Fix crash when cycling hosts on an empty hostgroup.

#rnx

- If we have no urls in the hostgroup we can just return the default and error values.

diff --git a/Engine/Source/Runtime/Experimental/IoStore/OnDemandCore/Private/OnDemandHostGroup.cpp b/Engine/Source/Runtime/Experimental/IoStore/OnDemandCore/Private/OnDemandHostGroup.cpp
--- a/Engine/Source/Runtime/Experimental/IoStore/OnDemandCore/Private/OnDemandHostGroup.cpp
+++ b/Engine/Source/Runtime/Experimental/IoStore/OnDemandCore/Private/OnDemandHostGroup.cpp
@@ -116,8 +116,16 @@ FAnsiStringView FOnDemandHostGroup::Host(int32 Index) const
 
 FAnsiStringView FOnDemandHostGroup::CycleHost(int32& InOutIndex) const
 {
-	InOutIndex = (InOutIndex + 1) % Impl->HostUrls.Num();
-	return Host(InOutIndex);
+	if (!Impl->HostUrls.IsEmpty())
+	{
+		InOutIndex = (InOutIndex + 1) % Impl->HostUrls.Num();
+		return Host(InOutIndex);
+	}
+	else
+	{
+		InOutIndex = INDEX_NONE;
+		return FAnsiStringView();
+	}
 }
 
 void FOnDemandHostGroup::SetPrimaryHost(int32 Index)
