Chaos, Bubble Resim: Core bubble resim logic: when a resim happens, we grow AABBs (the resim "bubbles") around particles needing to resim. Any non static particles overlapping those bubbles are marked as full resim, all others are frozen. Freezing techniques involves making them static or kinematic, and recording their pre freezing state. After resm, they are reset to their original state. The current logic to make the bubbles and detect overlaps is completely flawed (it uses the spatial acceleration at the wrong frame, it doesn't include whole island overlaps, etc.) but this is just a minimal implementation to put in place the basic scaffolding, we will improve the details as we go. To save perf, we also skip caching dirty particle data for frozen particles.
Simplification of CheckQuaternionThreshold that was doing unnecessary math.

diff --git a/Engine/Source/Runtime/Engine/Private/PhysicsEngine/NetworkPhysicsComponent.cpp b/Engine/Source/Runtime/Engine/Private/PhysicsEngine/NetworkPhysicsComponent.cpp
--- a/Engine/Source/Runtime/Engine/Private/PhysicsEngine/NetworkPhysicsComponent.cpp
+++ b/Engine/Source/Runtime/Engine/Private/PhysicsEngine/NetworkPhysicsComponent.cpp
@@ -1992,6 +1992,8 @@ FAsyncNetworkPhysicsComponentOutput& FAsyncNetworkPhysicsComponent::GetAsyncOutp
 
 void FAsyncNetworkPhysicsComponent::OnPreProcessInputs_Internal(const int32 PhysicsStep)
 {
+	QUICK_SCOPE_CYCLE_COUNTER(AsyncNetworkPhysicsComponent_OnPreProcessInputs_Internal);
+
 	ConsumeAsyncInput(PhysicsStep);
 
 	const FNetworkPhysicsSettingsNetworkPhysicsComponent& ComponentSettings = GetComponentSettings();
@@ -2227,6 +2229,8 @@ void FAsyncNetworkPhysicsComponent::OnPreProcessInputs_Internal(const int32 Phys
 
 void FAsyncNetworkPhysicsComponent::OnPostProcessInputs_Internal(const int32 PhysicsStep)
 {
+	QUICK_SCOPE_CYCLE_COUNTER(AsyncNetworkPhysicsComponent_OnPostProcessInputs_Internal);
+
 	const FNetworkPhysicsSettingsNetworkPhysicsComponent& ComponentSettings = GetComponentSettings();
 	const bool bIsServer = IsServer();
 
diff --git a/Engine/Source/Runtime/Experimental/Chaos/Private/PBDRigidsSolver.cpp b/Engine/Source/Runtime/Experimental/Chaos/Private/PBDRigidsSolver.cpp
--- a/Engine/Source/Runtime/Experimental/Chaos/Private/PBDRigidsSolver.cpp
+++ b/Engine/Source/Runtime/Experimental/Chaos/Private/PBDRigidsSolver.cpp
@@ -471,7 +471,10 @@ namespace Chaos
 				PRAGMA_ENABLE_DEPRECATION_WARNINGS
 			}
 
-			MSolver->ApplyCallbacks_Internal();
+			{
+				QUICK_SCOPE_CYCLE_COUNTER(Chaos_ApplyCallbacks_Internal);
+				MSolver->ApplyCallbacks_Internal();
+			}
 			
 			{
 				SCOPE_CYCLE_COUNTER(STAT_UpdateParams);
@@ -497,6 +500,7 @@ namespace Chaos
 
 			if(FRewindData* RewindData = MSolver->GetRewindData())
 			{
+				QUICK_SCOPE_CYCLE_COUNTER(Chaos_RewindData_AdvanceFrame);
 				RewindData->AdvanceFrame(MDeltaTime,[Evolution = MSolver->GetEvolution()]()
 				{
 					return Evolution->CreateExternalResimCache();
@@ -1883,146 +1887,155 @@ namespace Chaos
 		MarshallingManager.SetConsumerData_Internal(&PushData);
 
 		//update callbacks
-		SimCallbackObjects.Reserve(SimCallbackObjects.Num() + PushData.SimCallbackObjectsToAdd.Num());
-		for(ISimCallbackObject* SimCallbackObject : PushData.SimCallbackObjectsToAdd)
 		{
-			if (HasSimulationCallback(SimCallbackObject))
+			QUICK_SCOPE_CYCLE_COUNTER(Chaos_ProcessPushedData_AddCallbacks);
+			SimCallbackObjects.Reserve(SimCallbackObjects.Num() + PushData.SimCallbackObjectsToAdd.Num());
+			for(ISimCallbackObject* SimCallbackObject : PushData.SimCallbackObjectsToAdd)
 			{
-				SimCallbackObjects.Add(SimCallbackObject);
-			}
+				if (HasSimulationCallback(SimCallbackObject))
+				{
+					SimCallbackObjects.Add(SimCallbackObject);
+				}
 
-			if (SimCallbackObject->HasOption(ESimCallbackOptions::MidPhaseModification))
-			{
-				MidPhaseModifiers.Add(SimCallbackObject);
-			}
+				if (SimCallbackObject->HasOption(ESimCallbackOptions::MidPhaseModification))
+				{
+					MidPhaseModifiers.Add(SimCallbackObject);
+				}
 
-			if (SimCallbackObject->HasOption(ESimCallbackOptions::CCDModification))
-			{
-				CCDModifiers.Add(SimCallbackObject);
-			}
+				if (SimCallbackObject->HasOption(ESimCallbackOptions::CCDModification))
+				{
+					CCDModifiers.Add(SimCallbackObject);
+				}
 
-			if (SimCallbackObject->HasOption(ESimCallbackOptions::StrainModification))
-			{
-				StrainModifiers.Add(SimCallbackObject);
-			}
+				if (SimCallbackObject->HasOption(ESimCallbackOptions::StrainModification))
+				{
+					StrainModifiers.Add(SimCallbackObject);
+				}
 
-			if (SimCallbackObject->HasOption(ESimCallbackOptions::ContactModification))
-			{
-				ContactModifiers.Add(SimCallbackObject);
-			}
+				if (SimCallbackObject->HasOption(ESimCallbackOptions::ContactModification))
+				{
+					ContactModifiers.Add(SimCallbackObject);
+				}
 
-			if (SimCallbackObject->HasOption(ESimCallbackOptions::ParticleRegister))
-			{
-				RegistrationWatchers.Add(SimCallbackObject);
-			}
+				if (SimCallbackObject->HasOption(ESimCallbackOptions::ParticleRegister))
+				{
+					RegistrationWatchers.Add(SimCallbackObject);
+				}
 
-			if (SimCallbackObject->HasOption(ESimCallbackOptions::ParticleUnregister))
-			{
-				UnregistrationWatchers.Add(SimCallbackObject);
-			}
+				if (SimCallbackObject->HasOption(ESimCallbackOptions::ParticleUnregister))
+				{
+					UnregistrationWatchers.Add(SimCallbackObject);
+				}
 
-			if (SimCallbackObject->HasOption(ESimCallbackOptions::Rewind))
-			{
-				if (MRewindCallback)
+				if (SimCallbackObject->HasOption(ESimCallbackOptions::Rewind))
 				{
-					MRewindCallback->RegisterRewindableSimCallback_Internal(SimCallbackObject);
+					if (MRewindCallback)
+					{
+						MRewindCallback->RegisterRewindableSimCallback_Internal(SimCallbackObject);
+					}
 				}
-			}
 
-			if (SimCallbackObject->HasOption(ESimCallbackOptions::PhysicsObjectUnregister))
-			{
-				PhysicsObjectUnregistrationWatchers.Add(SimCallbackObject);
+				if (SimCallbackObject->HasOption(ESimCallbackOptions::PhysicsObjectUnregister))
+				{
+					PhysicsObjectUnregistrationWatchers.Add(SimCallbackObject);
+				}
 			}
 		}
 
 		//save any pending data for this particular interval
-		for (const FSimCallbackInputAndObject& InputAndCallbackObj : PushData.SimCallbackInputs)
 		{
-			InputAndCallbackObj.CallbackObject->SetCurrentInput_Internal(InputAndCallbackObj.Input);
+			QUICK_SCOPE_CYCLE_COUNTER(Chaos_SetCurrentInput_Internal);
+			for (const FSimCallbackInputAndObject& InputAndCallbackObj : PushData.SimCallbackInputs)
+			{
+				InputAndCallbackObj.CallbackObject->SetCurrentInput_Internal(InputAndCallbackObj.Input);
+			}
 		}
 
 		//remove any callbacks that are unregistered
-		for (ISimCallbackObject* RemovedCallbackObject : PushData.SimCallbackObjectsToRemove)
 		{
-			RemovedCallbackObject->bPendingDelete = true;
-
-			if (RemovedCallbackObject->HasOption(ESimCallbackOptions::Rewind))
+			QUICK_SCOPE_CYCLE_COUNTER(Chaos_ProcessPushedData_RemoveCallbacks);
+			for (ISimCallbackObject* RemovedCallbackObject : PushData.SimCallbackObjectsToRemove)
 			{
-				if (MRewindCallback)
+				RemovedCallbackObject->bPendingDelete = true;
+
+				if (RemovedCallbackObject->HasOption(ESimCallbackOptions::Rewind))
 				{
-					MRewindCallback->UnregisterRewindableSimCallback_Internal(RemovedCallbackObject);
+					if (MRewindCallback)
+					{
+						MRewindCallback->UnregisterRewindableSimCallback_Internal(RemovedCallbackObject);
+					}
 				}
 			}
-		}
 
-		for (int32 Idx = SimCallbackObjects.Num() - 1; Idx >= 0; --Idx)
-		{
-			if (SimCallbackObjects[Idx]->bPendingDelete)
+			for (int32 Idx = SimCallbackObjects.Num() - 1; Idx >= 0; --Idx)
 			{
-				SimCallbackObjects.RemoveAtSwap(Idx, EAllowShrinking::No);
+				if (SimCallbackObjects[Idx]->bPendingDelete)
+				{
+					SimCallbackObjects.RemoveAtSwap(Idx, EAllowShrinking::No);
+				}
 			}
-		}
 
-		for (int32 Idx = MidPhaseModifiers.Num() - 1; Idx >= 0; --Idx)
-		{
-			if (MidPhaseModifiers[Idx]->bPendingDelete)
+			for (int32 Idx = MidPhaseModifiers.Num() - 1; Idx >= 0; --Idx)
 			{
-				//will also be in SimCallbackObjects so we'll delete it in that loop
-				MidPhaseModifiers.RemoveAtSwap(Idx, EAllowShrinking::No);
+				if (MidPhaseModifiers[Idx]->bPendingDelete)
+				{
+					//will also be in SimCallbackObjects so we'll delete it in that loop
+					MidPhaseModifiers.RemoveAtSwap(Idx, EAllowShrinking::No);
+				}
 			}
-		}
 
-		for (int32 Idx = CCDModifiers.Num() - 1; Idx >= 0; --Idx)
-		{
-			if (CCDModifiers[Idx]->bPendingDelete)
+			for (int32 Idx = CCDModifiers.Num() - 1; Idx >= 0; --Idx)
 			{
-				//will also be in SimCallbackObjects so we'll delete it in that loop
-				CCDModifiers.RemoveAtSwap(Idx, EAllowShrinking::No);
+				if (CCDModifiers[Idx]->bPendingDelete)
+				{
+					//will also be in SimCallbackObjects so we'll delete it in that loop
+					CCDModifiers.RemoveAtSwap(Idx, EAllowShrinking::No);
+				}
 			}
-		}
 
-		for (int32 Idx = StrainModifiers.Num() - 1; Idx >= 0; --Idx)
-		{
-			if (StrainModifiers[Idx]->bPendingDelete)
+			for (int32 Idx = StrainModifiers.Num() - 1; Idx >= 0; --Idx)
 			{
-				//will also be in SimCallbackObjects so we'll delete it in that loop
-				StrainModifiers.RemoveAtSwap(Idx, EAllowShrinking::No);
+				if (StrainModifiers[Idx]->bPendingDelete)
+				{
+					//will also be in SimCallbackObjects so we'll delete it in that loop
+					StrainModifiers.RemoveAtSwap(Idx, EAllowShrinking::No);
+				}
 			}
-		}
 
-		for (int32 Idx = ContactModifiers.Num() - 1; Idx >= 0; --Idx)
-		{
-			if (ContactModifiers[Idx]->bPendingDelete)
+			for (int32 Idx = ContactModifiers.Num() - 1; Idx >= 0; --Idx)
 			{
-				//will also be in SimCallbackObjects so we'll delete it in that loop
-				ContactModifiers.RemoveAtSwap(Idx, EAllowShrinking::No);
+				if (ContactModifiers[Idx]->bPendingDelete)
+				{
+					//will also be in SimCallbackObjects so we'll delete it in that loop
+					ContactModifiers.RemoveAtSwap(Idx, EAllowShrinking::No);
+				}
 			}
-		}
 
-		for (int32 Idx = RegistrationWatchers.Num() - 1; Idx >= 0; --Idx)
-		{
-			if (RegistrationWatchers[Idx]->bPendingDelete)
+			for (int32 Idx = RegistrationWatchers.Num() - 1; Idx >= 0; --Idx)
 			{
-				//will also be in SimCallbackObjects so we'll delete it in that loop
-				RegistrationWatchers.RemoveAtSwap(Idx, EAllowShrinking::No);
+				if (RegistrationWatchers[Idx]->bPendingDelete)
+				{
+					//will also be in SimCallbackObjects so we'll delete it in that loop
+					RegistrationWatchers.RemoveAtSwap(Idx, EAllowShrinking::No);
+				}
 			}
-		}
 
-		for (int32 Idx = UnregistrationWatchers.Num() - 1; Idx >= 0; --Idx)
-		{
-			if (UnregistrationWatchers[Idx]->bPendingDelete)
+			for (int32 Idx = UnregistrationWatchers.Num() - 1; Idx >= 0; --Idx)
 			{
-				//will also be in SimCallbackObjects so we'll delete it in that loop
-				UnregistrationWatchers.RemoveAtSwap(Idx, EAllowShrinking::No);
+				if (UnregistrationWatchers[Idx]->bPendingDelete)
+				{
+					//will also be in SimCallbackObjects so we'll delete it in that loop
+					UnregistrationWatchers.RemoveAtSwap(Idx, EAllowShrinking::No);
+				}
 			}
-		}
 
-		for (int32 Idx = PhysicsObjectUnregistrationWatchers.Num() - 1; Idx >= 0; --Idx)
-		{
-			if (PhysicsObjectUnregistrationWatchers[Idx]->bPendingDelete)
+			for (int32 Idx = PhysicsObjectUnregistrationWatchers.Num() - 1; Idx >= 0; --Idx)
 			{
-				//will also be in SimCallbackObjects so we'll delete it in that loop
-				PhysicsObjectUnregistrationWatchers.RemoveAtSwap(Idx, EAllowShrinking::No);
+				if (PhysicsObjectUnregistrationWatchers[Idx]->bPendingDelete)
+				{
+					//will also be in SimCallbackObjects so we'll delete it in that loop
+					PhysicsObjectUnregistrationWatchers.RemoveAtSwap(Idx, EAllowShrinking::No);
+				}
 			}
 		}
 
@@ -2055,35 +2068,39 @@ namespace Chaos
 			}
 		}
 
-		for (FSimCallbackCommandObject*& SimCallbackObject : PushData.SimCommands)
 		{
-			if (!SimCallbackObject)
+			QUICK_SCOPE_CYCLE_COUNTER(Chaos_ProcessPushedData_SimCommands);
+			for (FSimCallbackCommandObject*& SimCallbackObject : PushData.SimCommands)
 			{
-				// A scheduled command delayed for a later frame will leave a nullptr in the array
-				continue;
-			}
+				if (!SimCallbackObject)
+				{
+					// A scheduled command delayed for a later frame will leave a nullptr in the array
+					continue;
+				}
 
-			if (SimCallbackObject->ShouldExecute(CurrentFrame, IsResimming()))
-			{
-				SimCallbackObject->SetSimAndDeltaTime_Internal(GetSolverTime(), MLastDt);
-				SimCallbackObject->PreSimulate_Internal();
+				if (SimCallbackObject->ShouldExecute(CurrentFrame, IsResimming()))
+				{
+					SimCallbackObject->SetSimAndDeltaTime_Internal(GetSolverTime(), MLastDt);
+					SimCallbackObject->PreSimulate_Internal();
+				}
+				else if (SimCallbackObject->ShouldDelay(CurrentFrame, IsResimming()))
+				{
+					// Cache command to run on a future frame and null the pointer in this PushData
+					MarshallingManager.EnqueueCommandScheduled_Internal(SimCallbackObject);
+					SimCallbackObject = nullptr; // Stop tracking the command in this PushData
+				}
 			}
-			else if (SimCallbackObject->ShouldDelay(CurrentFrame, IsResimming()))
+		
+			for (ISimCallbackObject* SimCallbackObject : PushData.SimCallbackObjectsToAdd)
 			{
-				// Cache command to run on a future frame and null the pointer in this PushData
-				MarshallingManager.EnqueueCommandScheduled_Internal(SimCallbackObject);
-				SimCallbackObject = nullptr; // Stop tracking the command in this PushData
+				SimCallbackObject->PostInitialize_Internal();
 			}
-		}
-		
-		for (ISimCallbackObject* SimCallbackObject : PushData.SimCallbackObjectsToAdd)
-		{
-			SimCallbackObject->PostInitialize_Internal();
-		}
 
-		if (MRewindCallback && !IsShuttingDown())
-		{
-			MRewindCallback->ProcessInputs_Internal(GetCurrentFrame(), PushData.SimCallbackInputs);
+			if (MRewindCallback && !IsShuttingDown())
+			{
+				QUICK_SCOPE_CYCLE_COUNTER(Chaos_RewindCallback_ProcessInputs_Internal);
+				MRewindCallback->ProcessInputs_Internal(GetCurrentFrame(), PushData.SimCallbackInputs);
+			}
 		}
 	}
 
@@ -2128,10 +2145,11 @@ namespace Chaos
 
 			if ((ResimStep <= LastStep) && NumResimSteps <= MarshallingManager.GetNumHistory_Internal())
 			{
-				FResimDebugInfo DebugInfo;
 				QUICK_SCOPE_CYCLE_COUNTER(ChaosRewindAndResim);
-				if (MRewindData->RewindToFrame(ResimStep))
+				if (MRewindData->CanRewindToFrame(ResimStep))
 				{
+					MRewindData->RewindToFrame(ResimStep);
+					
 #if DEBUG_REWIND_DATA
 					UE_LOG(LogChaos, Warning, TEXT("CLIENT | PT | ConditionalApplyRewind_Internal | PERFORMING RESIMULATION | Resim From Frame = %d | Num Steps = %d | To Current Frame: %d"), ResimStep, NumResimSteps, CurrentFrame);
 #endif
@@ -2143,69 +2161,97 @@ namespace Chaos
 					TArray<FPushPhysicsData*> RecordedPushData = MarshallingManager.StealHistory_Internal(NumResimSteps);
 					bool bFirst = true;
 
+					FResimDebugInfo DebugInfo;
 					FDurationTimer ResimTimer(DebugInfo.ResimTime);
-					// Do rollback as necessary
-					for (int32 Step = ResimStep; Step <= LastStep; ++Step)
 					{
-						if ((LastStep - Step) < RecordedPushData.Num())
+						QUICK_SCOPE_CYCLE_COUNTER(ChaosRewindAndResim_Resim);
+						// Do rollback as necessary
+						for (int32 Step = ResimStep; Step <= LastStep; ++Step)
 						{
-							if (!bFirst && !bUseCollisionResimCache)
+							if (!bFirst)
 							{
-								MRewindData->StepNonResimParticles(Step);
+								FPhysicsSolverBase& Solver = *this;
+								CVD_SCOPE_TRACE_SOLVER_STEP(CVDDC_Default, TEXT("Resim Step Start"));
+								CVD_TRACE_PARTICLES_SOA(GetEvolution()->GetParticles(), &GetEvolution()->GetRigidClustering());
 							}
 
-							if (PhysicsReplicationCVars::ResimulationCVars::bApplyTargetsWhileResimulating || bFirst)
+							if ((LastStep - Step) < RecordedPushData.Num())
 							{
-								// Update all the particles having received a target from the server
-								MRewindData->ApplyTargets(Step, bFirst);
-							}
+								if (!bFirst && !bUseCollisionResimCache)
+								{
+									MRewindData->StepNonResimParticles(Step);
+								}
 
-							FPushPhysicsData* PushData = RecordedPushData[LastStep - Step];	//push data is sorted as latest first
-							if (bFirst)
-							{
-								MTime = PushData->StartTime;	//not sure if sub-steps have proper StartTime so just do this once and let solver evolve remaining time
-								GetEvolution()->ResetCollisions();
-							}
-								
-							GetEvolution()->SetReset(bFirst);
+								if (PhysicsReplicationCVars::ResimulationCVars::bApplyTargetsWhileResimulating || bFirst)
+								{
+									// Update all the particles having received a target from the server
+									MRewindData->ApplyTargets(Step, bFirst);
+								}
+
+								FPushPhysicsData* PushData = RecordedPushData[LastStep - Step];	//push data is sorted as latest first
+								if (bFirst)
+								{
+									SCOPED_NAMED_EVENT(ChaosRewindAndResim_ResetCollisions, FColor::Red);
+									MTime = PushData->StartTime;	//not sure if sub-steps have proper StartTime so just do this once and let solver evolve remaining time
+									GetEvolution()->ResetCollisions();
+								}
+
+								GetEvolution()->SetReset(bFirst);
 #if !(UE_BUILD_SHIPPING || UE_BUILD_TEST)
-							if(bEnableNetworkPredictionDebug)
-							{
-								UE_LOG(LogChaos, Log, TEXT("		-> Re-simulating Frame = %d "), Step);
-								for(auto& Handle : GetParticles().GetNonDisabledDynamicView())
+								if (bEnableNetworkPredictionDebug)
 								{
-									if (GetEvolution()->GetIslandManager().GetParticleIsland(Handle.Handle()))
+									UE_LOG(LogChaos, Log, TEXT("		-> Re-simulating Frame = %d "), Step);
+									for (auto& Handle : GetParticles().GetNonDisabledDynamicView())
 									{
-										if (!Handle.IsSleeping())
-										{
-											UE_LOG(LogChaos, Log, TEXT("Particle Dynamic At Position = %s | Velocity = %s | Quaternion = %s | Omega = %s | Resim Frame = %d | Sync State = %d | Needs Resim = %d | Unique Idx = %d"),
-												*Handle.GetX().ToString(), *Handle.GetV().ToString(), *Handle.GetR().ToString(), *Handle.GetW().ToString(),
-												GetEvolution()->GetIslandManager().GetParticleIsland(Handle.Handle())->GetResimFrame(), (uint8)Handle.SyncState(),
-												GetEvolution()->GetIslandManager().GetParticleIsland(Handle.Handle())->NeedsResim(), Handle.UniqueIdx().Idx);
-										}
-										else
+										if (GetEvolution()->GetIslandManager().GetParticleIsland(Handle.Handle()))
 										{
-											UE_LOG(LogChaos, Log, TEXT("Particle Sleeping At Position = %s | Velocity = %s | Quaternion = %s | Omega = %s | Resim Frame = %d | Sync State = %d | Needs Resim = %d | Unique Idx = %d"),
-												*Handle.GetX().ToString(), *Handle.GetV().ToString(), *Handle.GetR().ToString(), *Handle.GetW().ToString(),
-												GetEvolution()->GetIslandManager().GetParticleIsland(Handle.Handle())->GetResimFrame(), (uint8)Handle.SyncState(),
-												GetEvolution()->GetIslandManager().GetParticleIsland(Handle.Handle())->NeedsResim(), Handle.UniqueIdx().Idx);
+											if (!Handle.IsSleeping())
+											{
+												UE_LOG(LogChaos, Log, TEXT("Particle Dynamic At Position = %s | Velocity = %s | Quaternion = %s | Omega = %s | Resim Frame = %d | Sync State = %d | Needs Resim = %d | Unique Idx = %d"),
+													*Handle.GetX().ToString(), *Handle.GetV().ToString(), *Handle.GetR().ToString(), *Handle.GetW().ToString(),
+													GetEvolution()->GetIslandManager().GetParticleIsland(Handle.Handle())->GetResimFrame(), (uint8)Handle.SyncState(),
+													GetEvolution()->GetIslandManager().GetParticleIsland(Handle.Handle())->NeedsResim(), Handle.UniqueIdx().Idx);
+											}
+											else
+											{
+												UE_LOG(LogChaos, Log, TEXT("Particle Sleeping At Position = %s | Velocity = %s | Quaternion = %s | Omega = %s | Resim Frame = %d | Sync State = %d | Needs Resim = %d | Unique Idx = %d"),
+													*Handle.GetX().ToString(), *Handle.GetV().ToString(), *Handle.GetR().ToString(), *Handle.GetW().ToString(),
+													GetEvolution()->GetIslandManager().GetParticleIsland(Handle.Handle())->GetResimFrame(), (uint8)Handle.SyncState(),
+													GetEvolution()->GetIslandManager().GetParticleIsland(Handle.Handle())->NeedsResim(), Handle.UniqueIdx().Idx);
+											}
 										}
 									}
 								}
-							}
 #endif
-							MRewindCallback->PreResimStep_Internal(Step, bFirst);
+								if (bFirst)
+								{
+									FPhysicsSolverBase& Solver = *this;
+									CVD_SCOPE_TRACE_SOLVER_STEP(CVDDC_Default, TEXT("Resim Step Start"));
+									CVD_TRACE_PARTICLES_SOA(GetEvolution()->GetParticles(), &GetEvolution()->GetRigidClustering());
+								}
 
-							// Run the advance but omit the game thread callbacks as we're executing during the physics tick
-							FSolverTasksPTOnly ImmediateTask(*this, PushData);
-							//ensure(bSolverHasFrozenGameThreadCallbacks == false);	//We don't support this for resim as it's very expensive and difficult to schedule
-							ImmediateTask.AdvanceSolver();
-							MRewindCallback->PostResimStep_Internal(Step);
-							
-							bFirst = false;
+								{
+									QUICK_SCOPE_CYCLE_COUNTER(ChaosRewindAndResim_RewindCallback_PreResimStep_Internal);
+									MRewindCallback->PreResimStep_Internal(Step, bFirst);
+								}
+
+								// Run the advance but omit the game thread callbacks as we're executing during the physics tick
+								FSolverTasksPTOnly ImmediateTask(*this, PushData);
+								//ensure(bSolverHasFrozenGameThreadCallbacks == false);	//We don't support this for resim as it's very expensive and difficult to schedule
+								ImmediateTask.AdvanceSolver();
+								{
+									QUICK_SCOPE_CYCLE_COUNTER(ChaosRewindAndResim_RewindCallback_PostResimStep_Internal);
+									MRewindCallback->PostResimStep_Internal(Step);
+								}
+
+								bFirst = false;
+							}
 						}
 					}
 
+					// Clean up after resim. This is where particles are unfrozen if they were frozen due to bubble resim.
+					MRewindData->EndResim();
+
 					MarshallingManager.SetConsumerData_Internal(ConsumerData);
 					GetEvolution()->GetIslandManager().ResetParticleResimFrame();
 
@@ -2218,7 +2264,7 @@ namespace Chaos
 #if DEBUG_REWIND_DATA
 				else
 				{
-					UE_LOG(LogChaos, Log, TEXT("CLIENT | PT | ConditionalApplyRewind_Internal | Resimulation failed, FRewindData::RewindToFrame returned false | Current Frame = %d | Num Steps = %d | Resim Frame = %d | Last Frame = %d | Rewind History Size = %d"), CurrentFrame, NumResimSteps, ResimStep, LastStep, MarshallingManager.GetNumHistory_Internal());
+					UE_LOG(LogChaos, Log, TEXT("CLIENT | PT | ConditionalApplyRewind_Internal | Resimulation failed, FRewindData::CanRewindToFrame returned false | Current Frame = %d | Num Steps = %d | Resim Frame = %d | Last Frame = %d | Rewind History Size = %d"), CurrentFrame, NumResimSteps, ResimStep, LastStep, MarshallingManager.GetNumHistory_Internal());
 				}
 #endif
 			}
diff --git a/Engine/Source/Runtime/Experimental/Chaos/Private/RewindData.cpp b/Engine/Source/Runtime/Experimental/Chaos/Private/RewindData.cpp
--- a/Engine/Source/Runtime/Experimental/Chaos/Private/RewindData.cpp
+++ b/Engine/Source/Runtime/Experimental/Chaos/Private/RewindData.cpp
@@ -6,14 +6,58 @@
 #include "PhysicsProxy/ClusterUnionPhysicsProxy.h"
 #include "PhysicsProxy/GeometryCollectionPhysicsProxy.h"
 #include "HAL/IConsoleManager.h"
-#include "Chaos/PBDJointConstraints.h"
+#include "Chaos/Island/IslandManager.h"
+#include "Chaos/ResimConsoleVariables.h"
+#include "ChaosVisualDebugger/ChaosVisualDebuggerTrace.h"
 #include "Misc/FileHelper.h"
 #include "Misc/Paths.h"
 #include "Misc/StringOutputDevice.h"
-#include "Chaos/Island/IslandManager.h"
 
 namespace Chaos
 {
+	namespace ResimConsoleVars
+	{
+		bool bEnableRewindLogs = false;
+		static FAutoConsoleVariableRef CVarEnableRewindLogs(
+			TEXT("p.Chaos.EnableRewindLogs"),
+			bEnableRewindLogs,
+			TEXT("Set to true to log important rewind logs"));
+
+		int32 BubbleResimFreezeMode = 1;
+		static FAutoConsoleVariableRef CVarBubbleResimFreezeMode(
+			TEXT("p.Chaos.BubbleResim.FreezeMode"),
+			BubbleResimFreezeMode,
+			TEXT("Mode used to freeze particles outside of resim bubble. Mode 0: Make particles static. Mode 1: Make particles kinematic."));
+
+		bool BubbleResimEnableCheapFreeze = false;
+		static FAutoConsoleVariableRef CVarBubbleResimEnableCheapFreeze(
+			TEXT("p.Chaos.BubbleResim.EnableCheapFreeze"),
+			BubbleResimEnableCheapFreeze,
+			TEXT("Use SetObjectStateLowLevel instead of going through the evolution to modify the object state, to freeze/unfreeze a particle"));
+
+		float BubbleResimExtraSweptDistance = 10.0f;
+		static FAutoConsoleVariableRef CVarBubbleResimExtraSweptDistance(
+			TEXT("p.Chaos.BubbleResim.ExtraSweptDistanceForCorrectedParticle"),
+			BubbleResimExtraSweptDistance,
+			TEXT("Extra swept distance when forming resim bubbles"));
+
+		float BubbleResimExtraThickness = 10.0f;
+		static FAutoConsoleVariableRef CVarBubbleResimExtraThickness(
+			TEXT("p.Chaos.BubbleResim.ExtraThicknessForCorrectedParticle"),
+			BubbleResimExtraThickness,
+			TEXT("Extra AABB thickness when forming resim bubbles"));
+
+		int32 BubbleResimGrowthMaxNumTries = 3;
+		static FAutoConsoleVariableRef CVarBubbleResimGrowthMaxNumTries(
+			TEXT("p.Chaos.BubbleResim.MaxNumTriesForResimBubbleGrowth"),
+			BubbleResimGrowthMaxNumTries,
+			TEXT("When forming bubbles, we grow the bubble by the AABBs of particles it overlaps. This is the maximum number of such growths we allow."));
+	}
+	
+	namespace ResimDefines
+	{
+		FName BubbleResimDebugDrawTag("BubbleResimDebugDraws");
+	}
 
 FVec3 FGeometryParticleState::ZeroVector = FVec3(0);
 
@@ -222,6 +266,7 @@ void FRewindData::RewindStates(const int32 RewindFrame, const bool bResetSolver)
 
 void FRewindData::ApplyTargets(const int32 Frame, const bool bResetSimulation)
 {
+	QUICK_SCOPE_CYCLE_COUNTER(ChaosRewindAndResim_ApplyTargets);
 	CVD_SCOPE_TRACE_SOLVER_STEP(CVDDC_Default, TEXT("Rewind Apply Targets"));
 
 	PRAGMA_DISABLE_DEPRECATION_WARNINGS
@@ -246,18 +291,24 @@ void FRewindData::ApplyTargets(const int32 Frame, const bool bResetSimulation)
 		FGeometryParticleHandle* PTParticle = DirtyParticleInfo.GetObjectPtr();
 		FGeometryParticleStateBase& History = DirtyParticleInfo.GetHistory();
 
+		// In bubble resim we skip applying targets to any particle marked as frozen during resim
+		if (ResimConsoleVars::bIsBubbleResimEnabled && PTParticle->ResimType() == EResimType::FrozenDuringResim)
+		{
+			continue;
+		}
+
 		const bool bResimAsFollower = DirtyParticleInfo.bResimAsFollower;
 
-		RewindHelper(PTParticle, bResimAsFollower, History.TargetPositions, [](auto Particle, const auto& Data) 
-		{
-			Particle->SetXR(Data); 
-		});
+		RewindHelper(PTParticle, bResimAsFollower, History.TargetPositions, [](auto Particle, const auto& Data)
+			{
+				Particle->SetXR(Data);
+			});
 		RewindHelper(PTParticle->CastToKinematicParticle(), bResimAsFollower, History.TargetVelocities, [](auto Particle, const auto& Data)
-		{
-			Particle->SetV(Data.V());
-			Particle->SetW(Data.W());
-		});
-		RewindHelper(PTParticle->CastToRigidParticle(), bResimAsFollower, History.TargetStates, [this](auto Particle, const auto& Data) 
+			{
+				Particle->SetV(Data.V());
+				Particle->SetW(Data.W());
+			});
+		RewindHelper(PTParticle->CastToRigidParticle(), bResimAsFollower, History.TargetStates, [this](auto Particle, const auto& Data)
 		{ 
 			if (Particle == nullptr || Solver->GetEvolution() == nullptr)
 			{
@@ -288,15 +339,18 @@ void FRewindData::ApplyTargets(const int32 Frame, const bool bResetSimulation)
 			Solver->GetEvolution()->SetParticleObjectState(Particle, Data.ObjectState());
 
 			// Todo: EResimType should be set by a resimulation system and ApplyTargets() should only process particles marked for resim
-			switch (Data.ObjectState())
+			if (!ResimConsoleVars::bIsBubbleResimEnabled)
 			{
-			case EObjectStateType::Dynamic:
-			case EObjectStateType::Sleeping:
-				Particle->SetResimType(EResimType::FullResim);
-				break;
-			default:
-				Particle->SetResimType(EResimType::ResimAsFollower);
-				break;
+				switch (Data.ObjectState())
+				{
+				case EObjectStateType::Dynamic:
+				case EObjectStateType::Sleeping:
+					Particle->SetResimType(EResimType::FullResim);
+					break;
+				default:
+					Particle->SetResimType(EResimType::ResimAsFollower);
+					break;
+				}
 			}
 		});
 
@@ -349,7 +403,7 @@ const int32 FRewindData::CompareTargetsToLastFrame()
 						ShouldTriggerResim |= FRewindData::CheckVectorThreshold(TargetState->GetX(), PastState->GetX(), FPhysicsSolverBase::GetResimulationErrorPositionThreshold()); // TODO: Take per actor settings into consideration via NetworkPhysicsSettingsComponent
 					}
 
-					if (bCompareR)
+					if (bCompareR && !ShouldTriggerResim)
 					{
 						ShouldTriggerResim |= FRewindData::CheckQuaternionThreshold(TargetState->GetR(), PastState->GetR(), FPhysicsSolverBase::GetResimulationErrorRotationThreshold()); // TODO: Take per actor settings into consideration via NetworkPhysicsSettingsComponent
 					}
@@ -369,7 +423,7 @@ const int32 FRewindData::CompareTargetsToLastFrame()
 						ShouldTriggerResim |= FRewindData::CheckVectorThreshold(TargetState->GetV(), PastState->GetV(), FPhysicsSolverBase::GetResimulationErrorLinearVelocityThreshold()); // TODO: Take per actor settings into consideration via NetworkPhysicsSettingsComponent
 					}
 
-					if (bCompareW)
+					if (bCompareW && !ShouldTriggerResim)
 					{
 						ShouldTriggerResim |= FRewindData::CheckVectorThreshold(TargetState->GetW(), PastState->GetW(), FPhysicsSolverBase::GetResimulationErrorAngularVelocityThreshold()); // TODO: Take per actor settings into consideration via NetworkPhysicsSettingsComponent
 					}
@@ -394,37 +448,288 @@ bool FRewindData::CheckVectorThreshold(FVec3 A, FVec3 B, float Threshold)
 
 bool FRewindData::CheckQuaternionThreshold(FQuat A, FQuat B, float ThresholdDegrees)
 {
-	// Get the rotational delta between A and B 
-	const FQuat RotDelta = A * B.Inverse();
+	const double ShortestArcAngle = FMath::Abs(FMath::Clamp(A | B, -1.0f, 1.0f));
+	const double AngleRadians = 2.0f * FMath::Acos(ShortestArcAngle);
+	const double AngleDegrees = FMath::RadiansToDegrees(AngleRadians);
 
-	// Convert delta to angle and axis
-	float Angle;
-	FVector Axis;
-	RotDelta.ToAxisAndAngle(Axis, Angle);
-	Angle = FMath::RadiansToDegrees(FMath::UnwindRadians(Angle));
-	Angle = FMath::Abs(Angle);
+	return AngleDegrees >= ThresholdDegrees;
+}
 
-	return Angle >= ThresholdDegrees;
+FAABB3 FRewindData::GetParticleSweptAABB(const FGeometryParticleHandle* ParticleHandle, FRewindData::FDirtyParticleInfo* DirtyParticleInfo, const FFrameAndPhase& AtFrameAndPhase, float HorizonDuration, float ExtraSweptDistance, float ExtraThickness) const
+{
+	FAABB3 ResultAABB = FAABB3::ZeroAABB();
+	const FGeometryParticleHandle* GeometryParticle = DirtyParticleInfo ? DirtyParticleInfo->GetObjectPtr() : ParticleHandle;
+	if (GeometryParticle)
+	{
+		FVector ParticlePosition = FVector::ZeroVector;
+		FQuat ParticleRotation = FQuat::Identity;
+		FVector ParticleVelocity = FVector::ZeroVector;
+		bool HistoryFound = false;
+		bool VelocityHistoryFound = false;
+		if (DirtyParticleInfo)
+		{
+			FGeometryParticleStateBase& History = DirtyParticleInfo->GetHistory();
+			const FParticlePositionRotation* PositionRotation = History.TargetPositions.Read(AtFrameAndPhase, PropertiesPool);
+			PositionRotation = PositionRotation ? PositionRotation : History.ParticlePositionRotation.Read(AtFrameAndPhase, PropertiesPool);
+			if (PositionRotation)
+			{
+				ParticlePosition = PositionRotation->GetX();
+				ParticleRotation = PositionRotation->GetR();
+				HistoryFound = true;
+			}
+			const FParticleVelocities* Velocities = History.TargetVelocities.Read(AtFrameAndPhase, PropertiesPool);
+			Velocities = Velocities ? Velocities : History.Velocities.Read(AtFrameAndPhase, PropertiesPool);
+			if (Velocities)
+			{
+				ParticleVelocity = Velocities->GetV();
+				VelocityHistoryFound = true;
+			}
+		}
+		if (!HistoryFound)
+		{
+			ParticlePosition = GeometryParticle->GetX();
+			ParticleRotation = GeometryParticle->GetR();
+		}
+		if (!VelocityHistoryFound)
+		{
+			if (const FKinematicGeometryParticleHandle* KinematicParticle = GeometryParticle->CastToKinematicParticle())
+			{
+				ParticleVelocity = KinematicParticle->GetV();
+			}
+		}
+		ResultAABB = GeometryParticle->LocalBounds().TransformedAABB(FTransform(ParticleRotation, ParticlePosition));
+		ResultAABB.Thicken(ExtraThickness);
+		FVector ParticleVelocityNormal = ParticleVelocity;
+		if (ParticleVelocityNormal.Normalize())
+		{
+			ResultAABB.GrowByVector(HorizonDuration * ParticleVelocity + ExtraSweptDistance * ParticleVelocityNormal);
+		}
+	}
+	return ResultAABB;
 }
 
+struct FNonStaticUniqueOverlapVisitor : public ISpatialVisitor<FAccelerationStructureHandle, FReal>
+{
+	using FVisitorFunctionType = TFunction<bool(const FGeometryParticleHandle*)>;
 
-bool bForceSkipResim = false;
-FAutoConsoleVariableRef CVarForceSkipResim(TEXT("p.Resim.ForceSkipResim"), bForceSkipResim, TEXT("Don't perform resimulations while this is true"));
+	FNonStaticUniqueOverlapVisitor(FVisitorFunctionType InParticleVisitor) : ParticleVisitor(InParticleVisitor) {}
+	virtual bool Overlap(const TSpatialVisitorData<FAccelerationStructureHandle>& Instance) override
+	{
+		if (FGeometryParticleHandle* Particle = Instance.Payload.GetGeometryParticleHandle_PhysicsThread())
+		{
+			bool WasParticleAlreadyInSet = false;
+			if (Particle->ObjectState() !=  EObjectStateType::Static)
+			{
+				OverlappingParticles.Add(Particle, &WasParticleAlreadyInSet);
+				if (!WasParticleAlreadyInSet)
+				{
+					return ParticleVisitor(Particle);
+				}
+			}
+		}
+		return true; // Keep going
+	}
+	virtual bool Raycast(const TSpatialVisitorData<FAccelerationStructureHandle>& Instance, FQueryFastData& CurData) override { return false; }
+	virtual bool Sweep(const TSpatialVisitorData<FAccelerationStructureHandle>& Instance, FQueryFastData& CurData) override { return false; }
 
-CHAOS_API bool bResimAllowRewindToResimulatedFrames = true;
-FAutoConsoleVariableRef CVarResimAllowRewindToResimulatedFrames(TEXT("p.Resim.AllowRewindToResimulatedFrames"), bResimAllowRewindToResimulatedFrames, TEXT("Allow rewinding back to a frame that was previously part of a resimulation. If a resimulation is performed between frame 100-110, allow a new resim from 105-115 if needed, else next resim will be able to start from frame 111."));
+	TSet<FGeometryParticleHandle*> OverlappingParticles;
+	FVisitorFunctionType ParticleVisitor;
+};
 
-CHAOS_API float ResimAllowRewindToResimulatedFramesOverlapPercent = 1.0f;
-FAutoConsoleVariableRef CVarResimAllowRewindToResimulatedFramesOverlapPercent(TEXT("p.Resim.AllowRewindToResimulatedFrames.OverlapPercent"), ResimAllowRewindToResimulatedFramesOverlapPercent, TEXT("Value in percent as a multiplier, 0.3 = 30%. When p.Resim.AllowRewindToResimulatedFrames is true, this sets how many frames of overlap we allow when rewinding back past a previous resimulation. If a resimulation is performed between frame 100-110 (10 frames), 30% overlap would allow 3 frames to get resimulated again, i.e. we would be allowed to rewind back to frame 107 again if we need to."));
+void FRewindData::FindParticlesOverlappingResimBubbles(const FFrameAndPhase& RewindFrameAndPhase, float TotalResimDuration, TSet<FGeometryParticleHandle*>& OutParticlesOverlappingResimBubbles)
+{
+	if (!ResimConsoleVars::bIsBubbleResimEnabled)
+	{
+		return;
+	}
+	
+	FPBDRigidsEvolution* Evolution = Solver->GetEvolution();
+
+	// Bubbles are swept particle AABBs
+	TArray<Chaos::FAABB3> ResimBubbles;
+
+	// First seed resim bubbles by building swept AABBs around particles needing resim
+	{
+		QUICK_SCOPE_CYCLE_COUNTER(ChaosRewindAndResim_SeedResimBubbles);
+		for (FDirtyParticleInfo& DirtyParticleInfo : DirtyParticles)
+		{
+			if (DirtyParticleInfo.bNeedsResim)
+			{
+				FGeometryParticleHandle* ParticleNeedingResim = DirtyParticleInfo.GetObjectPtr();
+				FAABB3 ResimAABB = GetParticleSweptAABB(DirtyParticleInfo.GetObjectPtr(), &DirtyParticleInfo, RewindFrameAndPhase,
+					TotalResimDuration, Chaos::ResimConsoleVars::BubbleResimExtraSweptDistance, Chaos::ResimConsoleVars::BubbleResimExtraThickness);
+				//static float ResimAABBThickening = 1.0f;
+				//ResimAABB.Thicken(ResimAABBThickening);
+				ResimBubbles.Add(ResimAABB);
+				// Important: it's better to add the particles needing resim for sure to ParticlesOverlappingBubble,
+				// because they might not be found to overlap the bubble based on the spatial acceleration structure.
+				// The spatial acceleration we have is for the frame we rewind from, not the frame we rewind to
+				OutParticlesOverlappingResimBubbles.Add(ParticleNeedingResim);
+				CVD_TRACE_DEBUG_DRAW_BOX(FBox(ResimAABB.Min(), ResimAABB.Max()), ResimDefines::BubbleResimDebugDrawTag, FColor::Red);
+			}
+		}
+	}
 
-bool FRewindData::RewindToFrame(int32 Frame)
+	// Then query the spatial acceleration structure to determine which particles overlap resim bubbles, and grow the bubble accordingly
+	// We repeat this until we don't overlap anything new, or we have reached the maximum number of tries
+	if (ISpatialAccelerationCollection<FAccelerationStructureHandle, FReal, 3>*SpatialAcceleration = Evolution->GetSpatialAcceleration())
+	{
+		QUICK_SCOPE_CYCLE_COUNTER(ChaosRewindAndResim_GrowResimBubbles);
+		for (FAABB3& ResimAABB : ResimBubbles)
+		{
+			int32 NumTries = 0;
+			int32 MaxNumTries = Chaos::ResimConsoleVars::BubbleResimGrowthMaxNumTries;
+			bool FoundNewOverlaps = true;
+			auto LambdaParticleVisitor = [&](const FGeometryParticleHandle* Particle)->bool
+				{
+					FDirtyParticleInfo* DirtyParticleInfo = DirtyParticles.Find(Particle);
+					FAABB3 ParticleResimAABB = GetParticleSweptAABB(Particle, DirtyParticleInfo, RewindFrameAndPhase,
+						TotalResimDuration, Chaos::ResimConsoleVars::BubbleResimExtraSweptDistance, Chaos::ResimConsoleVars::BubbleResimExtraThickness);
+					float Alpha = float(NumTries) / float(MaxNumTries <= 1 ? 1 : MaxNumTries - 1);
+					FLinearColor OverlapDrawColor = FMath::Lerp(FLinearColor(FColor::Magenta), FLinearColor(FColor::Purple), Alpha);
+					CVD_TRACE_DEBUG_DRAW_BOX(FBox(ParticleResimAABB.Min(), ParticleResimAABB.Max()), ResimDefines::BubbleResimDebugDrawTag, OverlapDrawColor.ToFColor(true));
+					ResimAABB.GrowToInclude(ParticleResimAABB);
+					FoundNewOverlaps = true;
+					return true; // Keep going
+				};
+			FNonStaticUniqueOverlapVisitor OverlapVisitor(LambdaParticleVisitor);
+			while (FoundNewOverlaps && NumTries < MaxNumTries)
+			{
+				FoundNewOverlaps = false;
+				SpatialAcceleration->Overlap(ResimAABB, OverlapVisitor);
+				float Alpha = float(NumTries) / float(MaxNumTries <= 1 ? 1 : MaxNumTries - 1);
+				FLinearColor BubbleDrawColor = FMath::Lerp(FLinearColor::Yellow, FLinearColor::Red, Alpha);
+				NumTries++;
+				CVD_TRACE_DEBUG_DRAW_BOX(FBox(ResimAABB.Min(), ResimAABB.Max()), ResimDefines::BubbleResimDebugDrawTag, BubbleDrawColor.ToFColor(true));
+			}
+			OutParticlesOverlappingResimBubbles.Append(OverlapVisitor.OverlappingParticles);
+		}
+	}
+}
+
+void FRewindData::FreezeParticlesOutsideResimBubbles()
 {
-	QUICK_SCOPE_CYCLE_COUNTER(RewindToFrame);
+	if (!ResimConsoleVars::bIsBubbleResimEnabled)
+	{
+		return;
+	}
 
-	CVD_SCOPE_TRACE_SOLVER_STEP(CVDDC_Default, TEXT("Rewind To Frame"));
+	QUICK_SCOPE_CYCLE_COUNTER(ChaosRewindAndResim_FreezeNonResimmingParticles);
 
-	EnsureIsInPhysicsThreadContext();
+	FPBDRigidsEvolution* Evolution = Solver->GetEvolution();
+
+	TArray<FPBDRigidParticleHandle*, TInlineAllocator<64>> ParticlesToFreezeInResim;
+	ParticlesToFreezeInResim.Reserve(64);
+
+	for (auto& Particle : Evolution->GetParticles().GetActiveDynamicMovingKinematicParticlesView())
+	{
+		if (FGeometryParticleHandle* ParticleHandle = Particle.Handle())
+		{
+			if (FPBDRigidParticleHandle* RigidParticleHandle = ParticleHandle->CastToRigidParticle();
+				RigidParticleHandle
+				&& (RigidParticleHandle->ResimType() == EResimType::FrozenDuringResim)
+				&& (RigidParticleHandle->ObjectState() != EObjectStateType::Static))
+			{
+				ParticlesToFreezeInResim.Add(RigidParticleHandle);
+			}
+		}
+	}
+
+	FrozenParticles.Reserve(ParticlesToFreezeInResim.Num());
+	for (int32 ParticleToFreezeIndex = 0; ParticleToFreezeIndex < ParticlesToFreezeInResim.Num(); ++ParticleToFreezeIndex)
+	{
+		FPBDRigidParticleHandle* RigidParticleHandle = ParticlesToFreezeInResim[ParticleToFreezeIndex];
+		if (IPhysicsProxyBase* PhysicsProxy = RigidParticleHandle->PhysicsProxy())
+		{
+			TPair<FParticleSafeRef, EObjectStateType> FrozenParticleState;
+			FrozenParticleState.Key.Proxy = static_cast<FSingleParticlePhysicsProxy*>(PhysicsProxy);
+			FrozenParticleState.Key.Timestamp = PhysicsProxy->GetSyncTimestamp();
+			FrozenParticleState.Value = RigidParticleHandle->ObjectState();
+			FrozenParticles.Add(MoveTemp(FrozenParticleState));
+		}
+		EObjectStateType FrozenObjectState = (ResimConsoleVars::BubbleResimFreezeMode == 0) ? EObjectStateType::Static : EObjectStateType::Kinematic;
+		if (Chaos::ResimConsoleVars::BubbleResimEnableCheapFreeze)
+		{
+			RigidParticleHandle->SetObjectStateLowLevel(FrozenObjectState);
+		}
+		else
+		{
+			// This is too slow, e.g. it calls UpdateViews every time, but it might be more correct
+			// TODO: extract the minimum of what's needed from SetParticleObjectState
+			Evolution->SetParticleObjectState(RigidParticleHandle, FrozenObjectState); 
+		}
+	}
+	if (Chaos::ResimConsoleVars::BubbleResimEnableCheapFreeze)
+	{
+		QUICK_SCOPE_CYCLE_COUNTER(ChaosRewindAndResim_UpdateDirtyViews);
+		Evolution->GetParticles().UpdateDirtyViews();
+	}
+}
 
+void FRewindData::UnfreezeParticlesOutsideResimBubbles()
+{
+	if (!ResimConsoleVars::bIsBubbleResimEnabled)
+	{
+		return;
+	}
+	
+	FPBDRigidsEvolution* Evolution = Solver->GetEvolution();
+	// Revert to pre-resim object state
+	{
+		QUICK_SCOPE_CYCLE_COUNTER(ChaosRewindAndResim_RevertObjectState);
+		for (TPair<FParticleSafeRef, EObjectStateType>& FrozenParticle : FrozenParticles)
+		{
+			IPhysicsProxyBase* Proxy = FrozenParticle.Key.Proxy;
+			TSharedPtr<FProxyTimestampBase, ESPMode::ThreadSafe> Timestamp = FrozenParticle.Key.Timestamp;
+			EObjectStateType PreResimObjectState = FrozenParticle.Value;
+			if (Timestamp.IsValid() && !Timestamp->bDeleted && Proxy && Proxy->GetSyncTimestamp() == Timestamp && Proxy->GetType() == EPhysicsProxyType::SingleParticleProxy)
+			{
+				FSingleParticlePhysicsProxy* SingleParticlePhysicsProxy = static_cast<FSingleParticlePhysicsProxy*>(Proxy);
+				if (FPBDRigidParticleHandle* RigidParticleHandle = SingleParticlePhysicsProxy->GetHandle_LowLevel()->CastToRigidParticle())
+				{
+					if (RigidParticleHandle->ObjectState() != PreResimObjectState)
+					{
+						if (Chaos::ResimConsoleVars::BubbleResimEnableCheapFreeze)
+						{
+							RigidParticleHandle->SetObjectStateLowLevel(PreResimObjectState);
+							QUICK_SCOPE_CYCLE_COUNTER(ChaosRewindAndResim_AddToIslandManager);
+							Evolution->GetIslandManager().AddParticle(RigidParticleHandle);
+						}
+						else
+						{
+							// This is too expensive as it calls UpdateViews for each particle
+							Evolution->SetParticleObjectState(RigidParticleHandle, PreResimObjectState);
+						}
+					}
+					// Trace modified particles. Unfortunately there is no way currently to only trace their object state so we trace the whole particle.
+					CVD_TRACE_PARTICLE(RigidParticleHandle);
+				}
+			}
+		}
+	}
+	{
+		QUICK_SCOPE_CYCLE_COUNTER(ChaosRewindAndResim_MarkAllFrozenInResimByDefault);
+		// Mark all particles as frozen in resim (default value)
+		// Improve this by only going through particles that were unfrozen for resim
+		for (auto& Particle : Evolution->GetParticles().GetActiveDynamicMovingKinematicParticlesView())
+		{
+			Particle.SetResimType(EResimType::FrozenDuringResim);
+		}
+	}
+	if (Chaos::ResimConsoleVars::BubbleResimEnableCheapFreeze)
+	{
+		QUICK_SCOPE_CYCLE_COUNTER(ChaosRewindAndResim_UpdateDirtyViews);
+		Evolution->GetParticles().UpdateDirtyViews();
+	}
+	// Empty frozen particles so we can start from scratch on the next rewind/resims
+	FrozenParticles.Empty();
+}
+
+bool bForceSkipResim = false;
+FAutoConsoleVariableRef CVarForceSkipResim(TEXT("p.Resim.ForceSkipResim"), bForceSkipResim, TEXT("Don't perform resimulations while this is true"));
+
+bool FRewindData::CanRewindToFrame(int32 RewindFrame) const
+{
 	if (bForceSkipResim)
 	{
 		return false;
@@ -432,30 +737,56 @@ bool FRewindData::RewindToFrame(int32 Frame)
 
 	//Can't go too far back
 	const int32 EarliestFrame = GetEarliestFrame_Internal();
-	if (Frame < EarliestFrame)
+	if (RewindFrame < EarliestFrame)
 	{
-#if DEBUG_REWIND_DATA
-		UE_LOG(LogChaos, Log, TEXT("CLIENT | PT | RewindToFrame | Failed due to rewind frame earlier than available history | Rewind Frame: %d | Earliest Frame: %d"), Frame, EarliestFrame);
-#endif
+		if (Chaos::ResimConsoleVars::bEnableRewindLogs)
+		{
+			UE_LOG(LogChaos, Log, TEXT("CLIENT | PT | RewindToFrame | Failed due to rewind frame earlier than available history | Rewind Frame: %d | Earliest Frame: %d"), RewindFrame, EarliestFrame);
+		}
 		return false;
 	}
 
 	//If we need to save and we are right on the edge of the buffer, we can't go back to earliest frame
-	if (Frame == EarliestFrame && bNeedsSave && FramesSaved == Managers.Capacity())
+	if (RewindFrame == EarliestFrame && bNeedsSave && FramesSaved == Managers.Capacity())
 	{
-#if DEBUG_REWIND_DATA
-		UE_LOG(LogChaos, Log, TEXT("CLIENT | PT | RewindToFrame | Failed due to rewinding to last available frame and bNeedsSave is set to true"));
-#endif
+		if (Chaos::ResimConsoleVars::bEnableRewindLogs)
+		{
+			UE_LOG(LogChaos, Log, TEXT("CLIENT | PT | RewindToFrame | Failed due to rewinding to last available frame and bNeedsSave is set to true"));
+		}
 		return false;
 	}
 
+	return true;
+}
+
+CHAOS_API bool bResimAllowRewindToResimulatedFrames = true;
+FAutoConsoleVariableRef CVarResimAllowRewindToResimulatedFrames(TEXT("p.Resim.AllowRewindToResimulatedFrames"), bResimAllowRewindToResimulatedFrames, TEXT("Allow rewinding back to a frame that was previously part of a resimulation. If a resimulation is performed between frame 100-110, allow a new resim from 105-115 if needed, else next resim will be able to start from frame 111."));
+
+CHAOS_API float ResimAllowRewindToResimulatedFramesOverlapPercent = 1.0f;
+FAutoConsoleVariableRef CVarResimAllowRewindToResimulatedFramesOverlapPercent(TEXT("p.Resim.AllowRewindToResimulatedFrames.OverlapPercent"), ResimAllowRewindToResimulatedFramesOverlapPercent, TEXT("Value in percent as a multiplier, 0.3 = 30%. When p.Resim.AllowRewindToResimulatedFrames is true, this sets how many frames of overlap we allow when rewinding back past a previous resimulation. If a resimulation is performed between frame 100-110 (10 frames), 30% overlap would allow 3 frames to get resimulated again, i.e. we would be allowed to rewind back to frame 107 again if we need to."));
+
+void FRewindData::RewindToFrame(int32 RewindFrame)
+{
+	QUICK_SCOPE_CYCLE_COUNTER(RewindToFrame);
+
+	CVD_SCOPE_TRACE_SOLVER_STEP(CVDDC_Default, TEXT("Rewind To Frame"));
+
+	EnsureIsInPhysicsThreadContext();
+
+	if (Chaos::ResimConsoleVars::bEnableRewindLogs)
+	{
+		UE_LOG(LogChaos, Log, TEXT("CLIENT | PT | FRewindData::RewindToFrame | Rewinding from Frame %d to Frame %d"), CurFrame, RewindFrame);
+	}
+
+	FPBDRigidsEvolution* Evolution = Solver->GetEvolution();
+
 	//If property changed between Frame and CurFrame, record the latest value and rewind to old
-	FFrameAndPhase RewindFrameAndPhase{ Frame, FFrameAndPhase::PrePushData };
+	FFrameAndPhase RewindFrameAndPhase{ RewindFrame, FFrameAndPhase::PrePushData };
 	FFrameAndPhase CurFrameAndPhase{ CurFrame, FFrameAndPhase::PrePushData };
 
 	if (bResimAllowRewindToResimulatedFrames)
 	{
-		float NumFrames = static_cast<float>(CurFrame - Frame);
+		float NumFrames = static_cast<float>(CurFrame - RewindFrame);
 		NumFrames *= ResimAllowRewindToResimulatedFramesOverlapPercent;
 		BlockResimFrame = CurFrame - FMath::CeilToInt(NumFrames);
 	}
@@ -464,120 +795,174 @@ bool FRewindData::RewindToFrame(int32 Frame)
 		BlockResimFrame = CurFrame;
 	}
 
-	ResimFrame = Frame;
-	CurFrame = Frame;
+	SetResimFrame(RewindFrame);
+	CurFrame = RewindFrame;
 	bNeedsSave = false;
 
+	// This is used for bubble resim
+	int32 NumResimFrames = CurFrameAndPhase.Frame - RewindFrameAndPhase.Frame;
+	float TotalResimDuration = static_cast<float>(NumResimFrames * (Solver ? Solver->GetAsyncDeltaTime() : 0.0));
+
+	// Bubble resim: form resim bubbles and determine particles overlapping them
+	// This does nothing if bubble resim is disabled
+	TSet<FGeometryParticleHandle*> ParticlesOverlappingBubble;
+	FindParticlesOverlappingResimBubbles(RewindFrameAndPhase, TotalResimDuration, /*out param*/ParticlesOverlappingBubble);
+
 	auto RewindHelper = [RewindFrameAndPhase, CurFrameAndPhase, this](auto Obj, bool bResimAsFollower, auto& Property, const auto& RewindFunc) -> bool
-	{
-		if (bResimAsFollower)
 		{
-			//If we're rewinding a particle that doesn't need to save head (resim as follower never checks for desync so we don't care about head)
-			if (auto Val = Property.Read(RewindFrameAndPhase, PropertiesPool))
+			if (bResimAsFollower)
 			{
-				RewindFunc(Obj, *Val);
+				//If we're rewinding a particle that doesn't need to save head (resim as follower never checks for desync so we don't care about head)
+				if (auto Val = Property.Read(RewindFrameAndPhase, PropertiesPool))
+				{
+					RewindFunc(Obj, *Val);
+					return true;
+				}
 			}
-		}
-		else
-		{
-			//If we're rewinding an object that needs to save head (during resim when we get back to latest frame and phase we need to check for desync)
-			if (!Property.IsClean(RewindFrameAndPhase))
+			else
 			{
-				if (!bRewindDataOptimization)
+				//If we're rewinding an object that needs to save head (during resim when we get back to latest frame and phase we need to check for desync)
+				if (!Property.IsClean(RewindFrameAndPhase))
 				{
-					// When not using optimized RewindData cache the current state in Phase::PrePushData on rewind.
-					CopyDataFromObject(Property.WriteAccessMonotonic(CurFrameAndPhase, PropertiesPool), *Obj);
+					if (!bRewindDataOptimization)
+					{
+						// When not using optimized RewindData cache the current state in Phase::PrePushData on rewind.
+						CopyDataFromObject(Property.WriteAccessMonotonic(CurFrameAndPhase, PropertiesPool), *Obj);
+					}
+					RewindFunc(Obj, *Property.Read(RewindFrameAndPhase, PropertiesPool));
+					return true;
 				}
-				RewindFunc(Obj, *Property.Read(RewindFrameAndPhase, PropertiesPool));
-
-				return true;
 			}
+			return false;
+		};
+
+	// Determine resim type
+	if (ResimConsoleVars::bIsBubbleResimEnabled)
+	{
+		// Bubble resim: Everything in the resim bubble fully resimulates, nothing resims as follower (as least for now)
+		QUICK_SCOPE_CYCLE_COUNTER(ChaosRewindAndResim_MarkAllParticlesInBubbleFullResim);
+		for (FDirtyParticleInfo& DirtyParticleInfo : DirtyParticles)
+		{
+			DirtyParticleInfo.bResimAsFollower = false;
+		}
+		for (FGeometryParticleHandle* GeometryParticle : ParticlesOverlappingBubble)
+		{
+			GeometryParticle->SetResimType(EResimType::FullResim);
 		}
+	}
+	else
+	{
+		for (FDirtyParticleInfo& DirtyParticleInfo : DirtyParticles)
+		{
+			FGeometryParticleHandle* GeometryParticle = DirtyParticleInfo.GetObjectPtr();
+			DirtyParticleInfo.bResimAsFollower = GeometryParticle->ObjectState() != EObjectStateType::Dynamic && GeometryParticle->ObjectState() != EObjectStateType::Sleeping;
+		}
+	}
 
-		return false;
-	};
-			
-	for(FDirtyParticleInfo& DirtyParticleInfo : DirtyParticles)
 	{
-		FGeometryParticleHandle* GeometryParticle = DirtyParticleInfo.GetObjectPtr();
-		FKinematicGeometryParticleHandle* KinematicParticle = GeometryParticle->CastToKinematicParticle();
-		FPBDRigidParticleHandle* RigidParticle = GeometryParticle->CastToRigidParticle();
+		QUICK_SCOPE_CYCLE_COUNTER(ChaosRewindAndResim_RewindParticles);
+		for(FDirtyParticleInfo& DirtyParticleInfo : DirtyParticles)
+		{
+			FGeometryParticleHandle* GeometryParticle = DirtyParticleInfo.GetObjectPtr();
+			FKinematicGeometryParticleHandle* KinematicParticle = GeometryParticle->CastToKinematicParticle();
+			FPBDRigidParticleHandle* RigidParticle = GeometryParticle->CastToRigidParticle();
 
-		//rewind is about to start, all particles should be in sync at this point
-		ensure(GeometryParticle->SyncState() == ESyncState::InSync);
+			//rewind is about to start, all particles should be in sync at this point
+			ensure(GeometryParticle->SyncState() == ESyncState::InSync);
 		
-		FGeometryParticleStateBase& History = DirtyParticleInfo.GetHistory(); //non-const in case we need to record what's at head for a rewind (CurFrame has already been increased to the next frame)
+			FGeometryParticleStateBase& History = DirtyParticleInfo.GetHistory(); //non-const in case we need to record what's at head for a rewind (CurFrame has already been increased to the next frame)
 
-		PRAGMA_DISABLE_DEPRECATION_WARNINGS
-		History.CachePreCorrectionState(*GeometryParticle); // Deprecated UE 5.6
-		PRAGMA_ENABLE_DEPRECATION_WARNINGS
+			PRAGMA_DISABLE_DEPRECATION_WARNINGS
+			History.CachePreCorrectionState(*GeometryParticle); // Deprecated UE 5.6
+			PRAGMA_ENABLE_DEPRECATION_WARNINGS
 
-		CachePreResimState(*GeometryParticle);
+			if (!ResimConsoleVars::bIsBubbleResimEnabled || GeometryParticle->ResimType() != EResimType::FrozenDuringResim)
+			{
+				CachePreResimState(*GeometryParticle);
+			}
 
-		// Todo: This should be set by bubble resimulation so that Dynamic and Sleeping particles outside of relevancy of resim doesn't actually resimulate as dynamic particles.
-		DirtyParticleInfo.bResimAsFollower = GeometryParticle->ObjectState() != EObjectStateType::Dynamic && GeometryParticle->ObjectState() != EObjectStateType::Sleeping;
-		const bool bResimAsFollower = DirtyParticleInfo.bResimAsFollower;
+			const bool bResimAsFollower = DirtyParticleInfo.bResimAsFollower;
 
-		if (KinematicParticle)
-		{
-			// Always clear the kinematic target
-			static const FKinematicTarget DefaultKinematicTarget = FKinematicTarget();
-			KinematicParticle->SetKinematicTarget(DefaultKinematicTarget);
-		}
+			if (KinematicParticle)
+			{
+				// Always clear the kinematic target
+				static const FKinematicTarget DefaultKinematicTarget = FKinematicTarget();
+				KinematicParticle->SetKinematicTarget(DefaultKinematicTarget);
+			}
 
-		bool bAnyChange = RewindHelper(GeometryParticle, bResimAsFollower, History.ParticlePositionRotation, [](auto Particle, const auto& Data) {Particle->SetXR(Data); });
-		bAnyChange |= RewindHelper(KinematicParticle, bResimAsFollower, History.Velocities, [](auto Particle, const auto& Data) { Particle->SetV(Data.V()); Particle->SetW(Data.W()); });
-		bAnyChange |= RewindHelper(GeometryParticle, bResimAsFollower, History.NonFrequentData, [this](auto Particle, const auto& Data)
-		{
-			Solver->GetEvolution()->InvalidateParticle(Particle); // Clear collision/constraints before updating NonFrequentData
-			Particle->SetNonFrequentData(Data); 
-		});
-		bAnyChange |= RewindHelper(KinematicParticle, bResimAsFollower, History.KinematicTarget, [this](auto Particle, const auto& Data)
-		{
-			if (Data.GetMode() != EKinematicTargetMode::None)
+			bool bAnyChange = false;
+
+			// With bubble resim enabled, we only rewind particles in the bubble		
+			bool ShouldRewind = !ResimConsoleVars::bIsBubbleResimEnabled || GeometryParticle->ResimType() != EResimType::FrozenDuringResim;
+			if (ShouldRewind)
 			{
-				Particle->SetKinematicTarget(Data);
-				Solver->GetEvolution()->GetParticles().MarkMovingKinematic(Particle->Handle());
+				bAnyChange = RewindHelper(GeometryParticle, bResimAsFollower, History.ParticlePositionRotation, [&](auto Particle, const auto& Data) {Particle->SetXR(Data); });
+				bAnyChange |= RewindHelper(KinematicParticle, bResimAsFollower, History.Velocities, [](auto Particle, const auto& Data) { Particle->SetV(Data.V()); Particle->SetW(Data.W()); });
+				bAnyChange |= RewindHelper(GeometryParticle, bResimAsFollower, History.NonFrequentData, [Evolution](auto Particle, const auto& Data)
+					{
+						Evolution->InvalidateParticle(Particle); // Clear collision/constraints before updating NonFrequentData
+						Particle->SetNonFrequentData(Data);
+					});
+				bAnyChange |= RewindHelper(KinematicParticle, bResimAsFollower, History.KinematicTarget, [Evolution](auto Particle, const auto& Data)
+					{
+						if (Data.GetMode() != EKinematicTargetMode::None)
+						{
+							Particle->SetKinematicTarget(Data);
+							Evolution->GetParticles().MarkMovingKinematic(Particle->Handle());
+						}
+					});
+				bAnyChange |= RewindHelper(RigidParticle, bResimAsFollower, History.Dynamics, [](auto Particle, const auto& Data) {Particle->SetDynamics(Data); });
+				bAnyChange |= RewindHelper(RigidParticle, bResimAsFollower, History.DynamicsMisc, [this](auto Particle, const auto& Data) {Solver->SetParticleDynamicMisc(Particle, Data); });
+				bAnyChange |= RewindHelper(RigidParticle, bResimAsFollower, History.MassProps, [](auto Particle, const auto& Data) {Particle->SetMassProps(Data); });
 			}
-		});
-		bAnyChange |= RewindHelper(RigidParticle, bResimAsFollower, History.Dynamics, [](auto Particle, const auto& Data) {Particle->SetDynamics(Data); });
-		bAnyChange |= RewindHelper(RigidParticle, bResimAsFollower, History.DynamicsMisc, [this](auto Particle, const auto& Data) {Solver->SetParticleDynamicMisc(Particle, Data); });
-		bAnyChange |= RewindHelper(RigidParticle, bResimAsFollower, History.MassProps, [](auto Particle, const auto& Data) {Particle->SetMassProps(Data); });
 
-		// Todo: This should be set by bubble resimulation so that Dynamic and Sleeping particles outside of relevancy of resim doesn't actually resimulate as dynamic particles.
-		// Set this after rewinding, since ResimType gets overwritten if NonFrequentData is cached
-		GeometryParticle->SetResimType(bResimAsFollower ? EResimType::ResimAsFollower : EResimType::FullResim);
+			// Set this after rewinding, since ResimType gets overwritten if NonFrequentData is cached
+			if (!ResimConsoleVars::bIsBubbleResimEnabled)
+			{
+				GeometryParticle->SetResimType(bResimAsFollower ? EResimType::ResimAsFollower : EResimType::FullResim);
+			}
 
-		if (!bResimAsFollower)
-		{
-			if (bAnyChange)
+			if (!bResimAsFollower && bAnyChange)
 			{
 				//particle actually changes not just created/streamed so need to update its state
 
 				//Data changes so send back to GT for interpolation. TODO: improve this in case data ends up being identical in resim
-				Solver->GetEvolution()->GetParticles().MarkTransientDirtyParticle(DirtyParticleInfo.GetObjectPtr());
+				Evolution->GetParticles().MarkTransientDirtyParticle(DirtyParticleInfo.GetObjectPtr());
 
 				DirtyParticleInfo.DirtyDynamics = INDEX_NONE;	//make sure to undo this as we want to record it again during resim
 
 				//for now just mark anything that changed as enabled during resim. TODO: use bubble
 				DirtyParticleInfo.GetObjectPtr()->SetEnabledDuringResim(true);
 			}
-		}
 
-		if (DirtyParticleInfo.InitializedOnStep > Frame)
-		{
-			//hasn't initialized yet, so disable
-			//must do this after rewind because SetDynamicsMisc will re-enable
-			//(the disable is a temp way to ignore objects not spawned yet, they weren't really disabled which is why it gets re-enabled)
-			Solver->GetEvolution()->DisableParticle(DirtyParticleInfo.GetObjectPtr());
-		}
+			if (DirtyParticleInfo.InitializedOnStep > RewindFrame)
+			{
+				//hasn't initialized yet, so disable
+				//must do this after rewind because SetDynamicsMisc will re-enable
+				//(the disable is a temp way to ignore objects not spawned yet, they weren't really disabled which is why it gets re-enabled)
+				Evolution->DisableParticle(DirtyParticleInfo.GetObjectPtr());
+			}
 
-		if (bAnyChange)
-		{
-			CVD_TRACE_PARTICLE(GeometryParticle);
+			// If we modified the particle, trace it to CVD
+			// Bubble Resim: CVD tracing for bubble resim is done at the end for this function, after we freeze particles
+			if (bAnyChange && !ResimConsoleVars::bIsBubbleResimEnabled)
+			{
+				CVD_TRACE_PARTICLE(GeometryParticle);
+			}
 		}
 	}
 
+	// Bubble resim only: "Freeze" particles outside resim bubbles
+	FreezeParticlesOutsideResimBubbles();
+
+	// With bubble resim, anything that's not static ends up either being frozen (outside the resim bubbles) or marked as FullResim (inside the bubble)
+	// So we are better off tracing everything. We're only tracing the ActiveDynamicMovingKinematicParticlesView to try and keep the size down
+	// TODO: Implement particle tracing to trace only object state and resim type but on all particles to make sure we've not forgotten some updates
+	if (ResimConsoleVars::bIsBubbleResimEnabled)
+	{
+		CVD_TRACE_PARTICLES_VIEW(Evolution->GetParticles().GetActiveDynamicMovingKinematicParticlesView());
+	}
+
 #if !UE_BUILD_SHIPPING
 	// For now, just ensure that the joints are InSync
 	for(FDirtyJointInfo& DirtyJointInfo : DirtyJoints)
@@ -587,12 +972,22 @@ bool FRewindData::RewindToFrame(int32 Frame)
 		ensure(Joint->SyncState() == ESyncState::InSync);
 	}
 #endif
+}
 
-	return true;
+void FRewindData::EndResim()
+{
+	CVD_SCOPE_TRACE_SOLVER_STEP(CVDDC_Default, TEXT("End Resim"));
+	QUICK_SCOPE_CYCLE_COUNTER(ChaosRewindAndResim_EndResim);
+
+	// Bubble resim only: "Unfreeze" particles that were frozen because they were outside resim bubbles
+	UnfreezeParticlesOutsideResimBubbles();
 }
 
 void FRewindData::StepNonResimParticles(const int32 Frame)
 {
+	QUICK_SCOPE_CYCLE_COUNTER(ChaosRewindAndResim_StepNonResimParticles);
+	CVD_SCOPE_TRACE_SOLVER_STEP(CVDDC_Default, TEXT("Step Resim As Followers"));
+
 	const FFrameAndPhase FrameAndPhase{ Frame, FFrameAndPhase::PrePushData };
 	auto RewindHelper = [FrameAndPhase, this](auto Obj, auto& Property, const auto& RewindFunc) -> bool
 	{
@@ -624,7 +1019,7 @@ void FRewindData::StepNonResimParticles(const int32 Frame)
 			KinematicParticle->SetKinematicTarget(DefaultKinematicTarget);
 		}
 
-		//  Set Postion and Rotation
+		//  Set Position and Rotation
 		bHasChanged = RewindHelper(GeometryParticle, History.ParticlePositionRotation, [](auto Particle, const auto& Data) -> bool
 		{
 			if (Particle->GetX() != Data.GetX() || Particle->GetR() != Data.GetR())
@@ -731,6 +1126,11 @@ void FRewindData::StepNonResimParticles(const int32 Frame)
 					return ECollisionVisitorResult::Continue;
 				});
 		}
+
+		if (bHasChanged && RigidParticle)
+		{
+			CVD_TRACE_PARTICLE(RigidParticle);
+		}
 	}
 }
 
@@ -760,21 +1160,31 @@ void FRewindData::FinishFrame()
 	{
 		FFrameAndPhase FutureFrame{ CurFrame + 1, FFrameAndPhase::PrePushData };
 
-		auto FinishHelper = [this, FutureFrame](auto& DirtyObjs)
+		auto DesyncIfNecessaryHelper = [this, FutureFrame](auto& DirtyObjs)
 		{
 			for (auto& Info : DirtyObjs)
 			{
-				if (Info.bResimAsFollower)
+				if (!Info.bResimAsFollower)
 				{
-					//resim as follower means always in sync and no cleanup needed
-					continue;
+					auto& Handle = *Info.GetObjectPtr();
+					if (Handle.ResimType() == EResimType::FullResim)
+					{
+						//solver doesn't affect dynamics, so no reason to test if they desynced from original sim
+						//question: should we skip all other properties? dynamics is a commonly changed one but might be worth skipping everything solver skips
+						DesyncIfNecessary</*bSkipDynamics=*/true>(Info, FutureFrame);
+					}
 				}
+			}
+		};
 
-				auto& Handle = *Info.GetObjectPtr();
-
-				if (Handle.ResimType() == EResimType::FullResim)
+		auto FinalResimHelper = [this, FutureFrame](auto& DirtyObjs)
+		{
+			for (auto& Info : DirtyObjs)
+			{
+				if (!Info.bResimAsFollower)
 				{
-					if (IsFinalResim())
+					auto& Handle = *Info.GetObjectPtr();
+					if (Handle.ResimType() == EResimType::FullResim)
 					{
 						//Last resim so mark as in sync
 						Handle.SetSyncState(ESyncState::InSync);
@@ -783,18 +1193,32 @@ void FRewindData::FinishFrame()
 						//Anything saved on upcoming frame (was done during rewind) can be removed since we are now at head
 						Info.ClearPhaseAndFuture(FutureFrame);
 					}
-					else if (!bRewindDataOptimization)
-					{
-						//solver doesn't affect dynamics, so no reason to test if they desynced from original sim
-						//question: should we skip all other properties? dynamics is a commonly changed one but might be worth skipping everything solver skips
-						DesyncIfNecessary</*bSkipDynamics=*/true>(Info, FutureFrame);
-					}
+				}
+				// Trying if this gets rid of the ensure on subsequent resims where some particles were left as hard desync, I don't know why
+				// TODO: This needs urgent attention (retest, does it still happen?) because it may indicate something is not done correctly in bubble resim
+				if (ResimConsoleVars::bIsBubbleResimEnabled)
+				{
+					auto& Handle = *Info.GetObjectPtr();
+					Handle.SetSyncState(ESyncState::InSync);
 				}
 			}
 		};
 
-		FinishHelper(DirtyParticles);
-		FinishHelper(DirtyJoints);
+		bool bIsFinalResim = IsFinalResim();
+
+		if (!bIsFinalResim)
+		{
+			if (!bRewindDataOptimization)
+			{
+				DesyncIfNecessaryHelper(DirtyParticles);
+				DesyncIfNecessaryHelper(DirtyJoints);
+			}
+		}
+		else
+		{
+			FinalResimHelper(DirtyParticles);
+			FinalResimHelper(DirtyJoints);
+		}
 	}
 	
 	++CurFrame;
@@ -1412,10 +1836,10 @@ int32 FRewindData::FindValidResimFrame(const int32 RequestedFrame)
 
 	if (RequestedFrame <= BlockResimFrame)
 	{
-#if DEBUG_REWIND_DATA
-		UE_LOG(LogChaos, Log, TEXT("CLIENT | PT | FindValidResimFrame | Resim is blocked | BlockResimFrame: %d | RequestedFrame: %d"), BlockResimFrame, RequestedFrame);
-#endif
-
+		if (Chaos::ResimConsoleVars::bEnableRewindLogs)
+		{
+			UE_LOG(LogChaos, Log, TEXT("CLIENT | PT | FindValidResimFrame | Resim is blocked | BlockResimFrame: %d | RequestedFrame: %d"), BlockResimFrame, RequestedFrame);
+		}
 		return ValidFrame;
 	}
 
@@ -1784,6 +2208,12 @@ void FRewindData::CacheKinematicTarget(TPBDRigidParticleHandle<FReal, 3>& Handle
 
 void FRewindData::CacheDirtyParticleData(FGeometryParticleHandle* Particle, const FFrameAndPhase::EParticleHistoryPhase& CurrentPhase, const bool& bDirty)
 {
+	// Bubble Resim
+	if (ResimConsoleVars::bIsBubbleResimEnabled && IsResim() && Particle->ResimType() == EResimType::FrozenDuringResim)
+	{
+		return;
+	}
+
 	FDirtyParticleInfo& DirtyParticleInfo = FindOrAddDirtyObj(*Particle);
 	if (bDirty)
 	{
@@ -1974,6 +2404,8 @@ void FRewindData::CacheDirtyJointData(FPBDJointConstraintHandle* Joint, const FF
 
 void FRewindData::CacheCurrentDirtyData(const FFrameAndPhase::EParticleHistoryPhase& CurrentPhase)
 {
+	QUICK_SCOPE_CYCLE_COUNTER(Chaos_CacheCurrentDirtyData);
+
 	if (!bRewindDataOptimization)
 	{
 		// The non optimized flow cache data in various other callbacks which will conflict with this
@@ -2059,7 +2491,7 @@ void FRewindData::SetTargetStateAtFrame(FGeometryParticleHandle& Handle, const i
 void FRewindData::RequestResimulation(int32 Frame, Chaos::FGeometryParticleHandle* Particle)
 {
 	// Update ResimFrame but don't allow to set a newer frame than already set
-	ResimFrame = (ResimFrame == INDEX_NONE) ? Frame : FMath::Min(ResimFrame, Frame);
+	SetResimFrame((ResimFrame == INDEX_NONE) ? Frame : FMath::Min(ResimFrame, Frame));
 
 	if (Particle && ensure(Solver))
 	{
diff --git a/Engine/Source/Runtime/Experimental/Chaos/Public/RewindData.h b/Engine/Source/Runtime/Experimental/Chaos/Public/RewindData.h
--- a/Engine/Source/Runtime/Experimental/Chaos/Public/RewindData.h
+++ b/Engine/Source/Runtime/Experimental/Chaos/Public/RewindData.h
@@ -4,9 +4,11 @@
 #include "Chaos/ParticleHandle.h"
 #include "Chaos/Framework/PhysicsSolverBase.h"
 #include "Serialization/BufferArchive.h"
+#include "Containers/Set.h"
 #include "Containers/CircularBuffer.h"
 #include "Chaos/ResimCacheBase.h"
 #include "Chaos/PBDJointConstraints.h"
+#include "Chaos/AABB.h"
 
 #ifndef VALIDATE_REWIND_DATA
 #define VALIDATE_REWIND_DATA 0
@@ -20,6 +22,9 @@
 #define DEBUG_NETWORK_PHYSICS 0
 #endif
 
+class IPhysicsProxyBase;
+struct FProxyTimestampBase;
+
 namespace Chaos
 {
 	namespace Private
@@ -1893,6 +1898,9 @@ public:
 	const int32 GetResimFrame() const { return ResimFrame; }
 	void SetResimFrame(int32 Frame) { ResimFrame = Frame; }
 
+	/** Called at the end of resim */
+	void EndResim();
+
 	/** Check if a frame number is within the current rewind history */
 	bool IsFrameWithinRewindHistory(int32 Frame) { return Frame < CurrentFrame() && Frame >= GetEarliestFrame_Internal(); }
 
@@ -2105,13 +2113,16 @@ private:
 		return TObjState(State, Handle, PropertiesPool, { Frame, Phase });
 	}
 
+	/** Returns whether rewinding to RewindFrame is allowed and possible */
+	bool CanRewindToFrame(int32 RewindFrame) const;
+
 	/** Apply the cached history state for the given frame cached particles and joints */
-	bool RewindToFrame(int32 RewindFrame);
+	void RewindToFrame(int32 RewindFrame);
 
 	/** Apply targets positions and velocities while resimulating */
 	void ApplyTargets(const int32 Frame, const bool bResetSimulation);
 
-	/** Apply resim data for objects not simulating during resimlation */
+	/** Apply resim data for objects not simulating during resimulation */
 	void StepNonResimParticles(const int32 Frame);
 
 	template <typename TDirtyInfo>
@@ -2152,6 +2163,26 @@ private:
 	TArray<const Private::FPBDIsland*> IslandValidationIslands;
 	TArray<const FGeometryParticleHandle*> IslandValidationIslandParticles;
 
+	// Bubble resim
+	// Freeze particles outside resim bubbles before rewind/resim
+	void FreezeParticlesOutsideResimBubbles();
+	// Unfreeze particles that were frozen, this happens at the end of a resim (after all the resim frames have occurred)
+	void UnfreezeParticlesOutsideResimBubbles();
+	// Util function to get a particle's swept AABB at a given frame and phase of the history
+	FAABB3 GetParticleSweptAABB(const FGeometryParticleHandle* ParticleHandle, FDirtyParticleInfo* DirtyParticleInfo, const FFrameAndPhase& AtFrameAndPhase, float HorizonDuration, float ExtraSweptDistance, float ExtraThickness) const;
+	// Adds all particles needing resim and particle likely to interact with them during resim to OutParticlesOverlappingResimBubbles
+	// This is accomplished by creating resim "bubbles" around particles needing resim and growing them as they overlap other particles
+	void FindParticlesOverlappingResimBubbles(const FFrameAndPhase& RewindFrameAndPhase, float TotalResimDuration, TSet<FGeometryParticleHandle*>& OutParticlesOverlappingResimBubbles);
+	// This structure references a particle in a way that is safe to use several frames after having been established
+	struct FParticleSafeRef
+	{
+		IPhysicsProxyBase* Proxy;
+		TSharedPtr<FProxyTimestampBase, ESPMode::ThreadSafe> Timestamp;
+	};
+	// Used to keep track of objects we "freeze" during resim (by making them kinematic and marking them FrozenDuringResim)
+	// and their object state prior to the freezing so it can be set back to its original value
+	TArray<TPair<FParticleSafeRef, EObjectStateType>> FrozenParticles;
+
 	template <typename TObj>
 	bool IsResimAndInSync(const TObj& Handle) const { return IsResim() && Handle.SyncState() == ESyncState::InSync; }
 
