Build warning fix, Solver was used before being checked for nullptr

diff --git a/Engine/Source/Runtime/Experimental/Chaos/Private/RewindData.cpp b/Engine/Source/Runtime/Experimental/Chaos/Private/RewindData.cpp
--- a/Engine/Source/Runtime/Experimental/Chaos/Private/RewindData.cpp
+++ b/Engine/Source/Runtime/Experimental/Chaos/Private/RewindData.cpp
@@ -448,6 +448,7 @@ bool FRewindData::CheckVectorThreshold(FVec3 A, FVec3 B, float Threshold)
 
 bool FRewindData::CheckQuaternionThreshold(FQuat A, FQuat B, float ThresholdDegrees)
 {
+	// We assume A and B are normalized
 	const double ShortestArcAngle = FMath::Abs(FMath::Clamp(A | B, -1.0f, 1.0f));
 	const double AngleRadians = 2.0f * FMath::Acos(ShortestArcAngle);
 	const double AngleDegrees = FMath::RadiansToDegrees(AngleRadians);
@@ -767,6 +768,12 @@ FAutoConsoleVariableRef CVarResimAllowRewindToResimulatedFramesOverlapPercent(TE
 
 void FRewindData::RewindToFrame(int32 RewindFrame)
 {
+	FPBDRigidsEvolution* Evolution = Solver ? Solver->GetEvolution() : nullptr;
+	if (!Evolution)
+	{
+		return;
+	}
+
 	QUICK_SCOPE_CYCLE_COUNTER(RewindToFrame);
 
 	CVD_SCOPE_TRACE_SOLVER_STEP(CVDDC_Default, TEXT("Rewind To Frame"));
@@ -778,8 +785,6 @@ void FRewindData::RewindToFrame(int32 RewindFrame)
 		UE_LOG(LogChaos, Log, TEXT("CLIENT | PT | FRewindData::RewindToFrame | Rewinding from Frame %d to Frame %d"), CurFrame, RewindFrame);
 	}
 
-	FPBDRigidsEvolution* Evolution = Solver->GetEvolution();
-
 	//If property changed between Frame and CurFrame, record the latest value and rewind to old
 	FFrameAndPhase RewindFrameAndPhase{ RewindFrame, FFrameAndPhase::PrePushData };
 	FFrameAndPhase CurFrameAndPhase{ CurFrame, FFrameAndPhase::PrePushData };
