[ChaosVD] Adding a buffer size limit to the relay writer

Now if we try to buffer more data than the configured limit, the relay will trigger an overflow event and the recording will automatically stop.

diff --git a/Engine/Source/Runtime/Experimental/ChaosVisualDebugger/Private/ChaosVDRuntimeModule.cpp b/Engine/Source/Runtime/Experimental/ChaosVisualDebugger/Private/ChaosVDRuntimeModule.cpp
--- a/Engine/Source/Runtime/Experimental/ChaosVisualDebugger/Private/ChaosVDRuntimeModule.cpp
+++ b/Engine/Source/Runtime/Experimental/ChaosVisualDebugger/Private/ChaosVDRuntimeModule.cpp
@@ -20,6 +20,7 @@
 #endif
 
 #include "ChaosVDRecordingDetails.h"
+#include "Math/GuardedInt.h"
 #include "Misc/App.h"
 #include "Modules/ModuleManager.h"
 #include "ProfilingDebugging/TraceAuxiliary.h"
@@ -70,6 +71,11 @@ static FAutoConsoleVariable CVarChaosVDMaxBytesPerRelayChunk(
 	65536,
 	TEXT("How many bytes per relay data message we should send"));
 
+static FAutoConsoleVariable CVarChaosVDMaxPendingBytesInRelayWriter(
+	TEXT("p.Chaos.VD.MaxPendingBytesInRelayWriter"),
+	536870912,
+	TEXT("How many bytes the relay writer can hold before overflowing, triggering a close. Defaults to ~512MB"));
+
 FChaosVDRecordingStateChangedDelegate FChaosVDRuntimeModule::RecordingStartedDelegate = FChaosVDRecordingStateChangedDelegate();
 FChaosVDRecordingStateChangedDelegate FChaosVDRuntimeModule::PostRecordingStartedDelegate = FChaosVDRecordingStateChangedDelegate();
 FChaosVDRecordingStateChangedDelegate FChaosVDRuntimeModule::RecordingStopDelegate = FChaosVDRecordingStateChangedDelegate();
@@ -155,12 +161,20 @@ int32 FChaosVDRuntimeModule::GenerateUniqueID()
 
 int64 FChaosVDRuntimeModule::GetBufferedDataBytesSize() const
 {
+	FGuardedInt64 RelaySendBufferSize;
+
 	if (RelayWriter)
 	{
-		return RelayWriter->GetQueuedBytesNum();
+		RelaySendBufferSize += RelayWriter->GetQueuedBytesNum();
+	}
+
+	if (ExternalRelayBufferSizeProviderDelegate.IsBound())
+	{
+		RelaySendBufferSize += ExternalRelayBufferSizeProviderDelegate.Execute();
 	}
 
-	return 0;
+	constexpr int64 DefaultValue = std::numeric_limits<int64>::max();
+	return RelaySendBufferSize.Get(DefaultValue); 
 }
 
 FString FChaosVDRuntimeModule::GetLastRecordingFileNamePath() const
@@ -210,6 +224,21 @@ void FChaosVDRuntimeModule::RegisterExternalRelayExecutor(const FChaosVDRelayTra
 	ExternalRelayExecutorDelegate = InExternalRelayExecutor;
 }
 
+void FChaosVDRuntimeModule::RegisterExternalRelayOverflowHandler(const FSimpleDelegate& InExternalRelayOverflowHandler)
+{
+	if (!ensureMsgf(!ExternalRelayOverflowHandlerDelegate.IsBound(), TEXT("An external relay overflow handler is already registered!")))
+	{
+		return;
+	}
+
+	ExternalRelayOverflowHandlerDelegate = InExternalRelayOverflowHandler;
+}
+
+void FChaosVDRuntimeModule::UnregisterExternalRelayOverflowHandler()
+{
+	ExternalRelayOverflowHandlerDelegate.Unbind();
+}
+
 void FChaosVDRuntimeModule::UnregisterCurrentExternalRelayExecutor()
 {
 	ExternalRelayExecutorDelegate.Unbind();
@@ -225,6 +254,34 @@ void FChaosVDRuntimeModule::UnregisterExternalTraceStatusProvider()
 	ExternalTraceStatusDelegate.Unbind();
 }
 
+void FChaosVDRuntimeModule::RegisterExternalRelayBufferSizeProvider(const FChaosVDExternalRelayBufferQueueSizeDelegate& InExternalRelayBufferSizeCallback)
+{
+	if (!ensureMsgf(!ExternalRelayBufferSizeProviderDelegate.IsBound(), TEXT("An external relay InExternalRelayBufferSizeCallback is already registered!")))
+	{
+		return;
+	}
+
+	ExternalRelayBufferSizeProviderDelegate = InExternalRelayBufferSizeCallback;
+}
+
+void FChaosVDRuntimeModule::UnregisterExternalRelayBufferSizeProvider()
+{
+	ExternalRelayBufferSizeProviderDelegate.Unbind();
+}
+
+void FChaosVDRuntimeModule::HandleRelayOverflow()
+{
+	if (ExternalRelayOverflowHandlerDelegate.IsBound())
+	{
+		UE_LOG(LogChaosVDRuntime, Error, TEXT("[%hs] Relay writer has overflowed. Executing external overflow handler... "), __func__);
+		ExternalRelayOverflowHandlerDelegate.Execute();
+		return;
+	}
+
+	UE_LOG(LogChaosVDRuntime, Error, TEXT("[%hs] Relay writer has overflowed. Stopping Recording... "), __func__);
+	StopRecording();
+}
+
 void FChaosVDRuntimeModule::RelayTraceData()
 {
 	if (ensure(ExternalRelayExecutorDelegate.IsBound() && RelayWriter.IsValid()))
@@ -241,6 +298,7 @@ void FChaosVDRuntimeModule::SetupRelayDataPumpDelegates()
 	}
 
 	RelayWriter->OnNewDataAvailable().BindRaw(this, &FChaosVDRuntimeModule::RelayTraceData);
+	RelayWriter->OnDataOverflow().BindRaw(this, &FChaosVDRuntimeModule::HandleRelayOverflow);
 }
 
 void FChaosVDRuntimeModule::ClearRelayDataPumpDelegates()
@@ -248,6 +306,7 @@ void FChaosVDRuntimeModule::ClearRelayDataPumpDelegates()
 	if (RelayWriter)
 	{
 		RelayWriter->OnNewDataAvailable().Unbind();
+		RelayWriter->OnDataOverflow().Unbind();
 	}
 }
 
@@ -386,7 +445,12 @@ void FChaosVDRuntimeModule::StartRecording(const FChaosVDStartRecordingCommandMe
 				}
 
 				RelayWriter = MakeUnique<Chaos::VD::FRelayTraceWriter>();
-				RelayWriter->SetMaxBytesPerBunch(CVarChaosVDMaxBytesPerRelayChunk->GetInt());
+
+				constexpr int32 MinBunchSize = 4096;
+				RelayWriter->SetMaxBytesPerBunch(FMath::Clamp(CVarChaosVDMaxBytesPerRelayChunk->GetInt(), MinBunchSize, std::numeric_limits<int32>::max()));
+
+				constexpr int32 MinBufferSizeLimit = 2097152;
+				RelayWriter->SetMaxAllowedPendingBytes(FMath::Clamp(CVarChaosVDMaxPendingBytesInRelayWriter->GetInt(), MinBufferSizeLimit, std::numeric_limits<int32>::max()));
 				
 				bIsRecording = FTraceAuxiliary::Relay(reinterpret_cast<UPTRINT>(RelayWriter.Get()), Chaos::VD::FRelayTraceWriter::WriteHelper, Chaos::VD::FRelayTraceWriter::CloseHelper, nullptr, &TracingOptions);
 
diff --git a/Engine/Source/Runtime/Experimental/ChaosVisualDebugger/Private/Trace/RelayTraceWriter.cpp b/Engine/Source/Runtime/Experimental/ChaosVisualDebugger/Private/Trace/RelayTraceWriter.cpp
--- a/Engine/Source/Runtime/Experimental/ChaosVisualDebugger/Private/Trace/RelayTraceWriter.cpp
+++ b/Engine/Source/Runtime/Experimental/ChaosVisualDebugger/Private/Trace/RelayTraceWriter.cpp
@@ -4,6 +4,7 @@
 
 #if WITH_CHAOS_VISUAL_DEBUGGER
 #include "ChaosVDRuntimeModule.h"
+#include "Math/GuardedInt.h"
 #include "Serialization/BufferArchive.h"
 #include "Serialization/MemoryReader.h"
 
@@ -37,6 +38,16 @@ namespace Chaos::VD
 		MaxPendingBytesPerBunch = MaxBytes;
 	}
 
+	void FRelayTraceWriter::SetMaxAllowedPendingBytes(uint64 MaxBytes)
+	{
+		if (MaxBytes == 0)
+		{
+			UE_LOG(LogChaosVDRuntime, Error, TEXT("[%hs] Only values greater than 0 are supported"), __func__);
+			return;
+		}
+	
+		MaxAllowedPendingBytes = MaxBytes;
+	}
 
 	TOptional<TUniquePtr<FBufferArchive>> FRelayTraceWriter::DequeuePendingBunch()
 	{
@@ -87,6 +98,13 @@ namespace Chaos::VD
 			return false;
 		}
 
+		if (!CanWriteData(Size))
+		{
+			DataOverflowDelegate.ExecuteIfBound();
+			Close();
+			return false;
+		}
+
 		if (!CurrentPendingBunch)
 		{
 			CurrentPendingBunch = CreateBunch();
@@ -133,6 +151,23 @@ namespace Chaos::VD
 		return !CurrentPendingBunch->IsError();
 	}
 
+	bool FRelayTraceWriter::CanWriteData(uint32 NewDataSize) const
+	{
+		TGuardedInt<uint64> ExpectedQueuedBytes;
+
+		ExpectedQueuedBytes+= QueuedBytesNum;
+
+		ExpectedQueuedBytes += NewDataSize;
+
+		if (CurrentPendingBunch)
+		{
+			ExpectedQueuedBytes+= CurrentPendingBunch->Num();
+		}
+	
+		constexpr uint64 DefaultValue = std::numeric_limits<uint64>::max();
+		return ExpectedQueuedBytes.Get(DefaultValue) <= MaxAllowedPendingBytes;
+	}
+
 	int64 FRelayTraceWriter::GetQueuedBytesNum() const
 	{
 		return QueuedBytesNum;
diff --git a/Engine/Source/Runtime/Experimental/ChaosVisualDebugger/Public/ChaosVDRuntimeModule.h b/Engine/Source/Runtime/Experimental/ChaosVisualDebugger/Public/ChaosVDRuntimeModule.h
--- a/Engine/Source/Runtime/Experimental/ChaosVisualDebugger/Public/ChaosVDRuntimeModule.h
+++ b/Engine/Source/Runtime/Experimental/ChaosVisualDebugger/Public/ChaosVDRuntimeModule.h
@@ -35,6 +35,7 @@ DECLARE_MULTICAST_DELEGATE_OneParam(FChaosVDCaptureRequestDelegate, EChaosVDFull
 DECLARE_MULTICAST_DELEGATE_OneParam(FChaosVDRecordingStartFailedDelegate, const FText&)
 DECLARE_DELEGATE_OneParam(FChaosVDRelayTraceDataDelegate, Chaos::VD::FRelayTraceWriter&)
 DECLARE_DELEGATE_RetVal(FChaosVDTraceDetails, FChaosVDExternalTraceStatusDelegate)
+DECLARE_DELEGATE_RetVal(uint64, FChaosVDExternalRelayBufferQueueSizeDelegate)
 
 DECLARE_LOG_CATEGORY_EXTERN(LogChaosVDRuntime, Log, All)
 
@@ -136,6 +137,7 @@ public:
 		return AccumulatedRecordingTime;
 	}
 
+	/** Returns the accumulated current send buffer size used by the relay transport system - This should be used for logging or UI feedback only */
 	UE_API int64 GetBufferedDataBytesSize() const;
 
 	/** [Deprecated] Returns the full path of the active recording file*/
@@ -160,6 +162,17 @@ public:
 	 */
 	UE_API void RegisterExternalRelayExecutor(const FChaosVDRelayTraceDataDelegate& InExternalRelayExecutor);
 
+	/**
+	 * Registers an external handler that will be in charge handling what happens when the Relay writer overflows
+	 * @param InExternalRelayOverflowHandler Delegate that will be executed when the Relay writer overflows
+	 */
+	UE_API void RegisterExternalRelayOverflowHandler(const FSimpleDelegate& InExternalRelayOverflowHandler);
+
+	/**
+	 * Unbinds the delegate to an external relay overflow handler
+	 */
+	UE_API void UnregisterExternalRelayOverflowHandler();
+
 	/**
 	 * Unbinds the delegate to an external relay executor
 	 */
@@ -176,13 +189,32 @@ public:
 	 */
 	UE_API void UnregisterExternalTraceStatusProvider();
 
+	/**
+	 *  Registers an external handler that will provide information about any external send buffer the current relay transport implementation might be using, if any.
+	 *  This is used for UI metrics. To handle buffer overflows, use RegisterExternalRelayOverflowHandler 
+	 * @param InExternalRelayBufferSizeCallback 
+	 */
+	UE_API void RegisterExternalRelayBufferSizeProvider(const FChaosVDExternalRelayBufferQueueSizeDelegate& InExternalRelayBufferSizeCallback);
+
+	/**
+	 * Unbinds the delegate to the current external Relay buffer size provider
+	 */
+	UE_API void UnregisterExternalRelayBufferSizeProvider();
+
 	FSimpleMulticastDelegate& OnTraceConnectionDetailsUpdated()
 	{
 		return TraceConnectionDetailsUpdatedDelegate;
 	}
 
+	Chaos::VD::FRelayTraceWriter* GetRelayWriterInstance() const
+	{
+		return RelayWriter.Get();
+	}
+
 private:
 
+	void HandleRelayOverflow();
+
 	void RelayTraceData();
 
 	void SetupRelayDataPumpDelegates();
@@ -236,6 +268,8 @@ private:
 
 	TUniquePtr<Chaos::VD::FRelayTraceWriter> RelayWriter;
 	FChaosVDRelayTraceDataDelegate ExternalRelayExecutorDelegate;
+	FSimpleDelegate ExternalRelayOverflowHandlerDelegate;
+	FChaosVDExternalRelayBufferQueueSizeDelegate ExternalRelayBufferSizeProviderDelegate;
 
 	static UE_API FTransactionallySafeRWLock DelegatesRWLock;
 
diff --git a/Engine/Source/Runtime/Experimental/ChaosVisualDebugger/Public/Trace/RelayTraceWriter.h b/Engine/Source/Runtime/Experimental/ChaosVisualDebugger/Public/Trace/RelayTraceWriter.h
--- a/Engine/Source/Runtime/Experimental/ChaosVisualDebugger/Public/Trace/RelayTraceWriter.h
+++ b/Engine/Source/Runtime/Experimental/ChaosVisualDebugger/Public/Trace/RelayTraceWriter.h
@@ -32,6 +32,13 @@ public:
 	 */
 	CHAOSVDRUNTIME_API void SetMaxBytesPerBunch(uint32 MaxBytes);
 
+	
+	/**
+	 * Sets max number of bytes this writer can buffer before overflowing, triggering a close
+	 * @param MaxBytes number of bytes to set as limit
+	 */
+	CHAOSVDRUNTIME_API void SetMaxAllowedPendingBytes(uint64 MaxBytes);
+
 	static bool WriteHelper(UPTRINT WriterHandle, const void* Data, uint32 Size)
 	{
 		if (!WriterHandle)
@@ -94,21 +101,31 @@ public:
 		return NewDataAvailableDelegate;
 	}
 
+	FSimpleDelegate& OnDataOverflow()
+	{
+		return DataOverflowDelegate;
+	}
+
 	private:
 
 	FNewDataAvailableDelegate NewDataAvailableDelegate;
+	FNewDataAvailableDelegate DataOverflowDelegate;
 
 	TUniquePtr<FBufferArchive> CreateBunch();
 
 	bool Write_Internal(const void* Data, uint32 Size);
 
+	bool CanWriteData(uint32 NewDataSize) const;
+
 	TSpscQueue<TUniquePtr<FBufferArchive>> PendingBunchesToRelay;
 	TUniquePtr<FBufferArchive> CurrentPendingBunch;
 	std::atomic<int32> PendingBunchesCount = 0;
 	std::atomic<bool> bClosed = false;
 	uint32 MaxPendingBytesPerBunch = 0;
 	std::atomic<int64> QueuedBytesNum = 0;
+	uint64 MaxAllowedPendingBytes = 1073741824;
 };
+
 }
 #endif // WITH_CHAOS_VISUAL_DEBUGGER
 
