[Chaos]
* Fixed non-unity compile error

diff --git a/Engine/Source/Runtime/Experimental/Chaos/Public/Chaos/EPA.h b/Engine/Source/Runtime/Experimental/Chaos/Public/Chaos/EPA.h
--- a/Engine/Source/Runtime/Experimental/Chaos/Public/Chaos/EPA.h
+++ b/Engine/Source/Runtime/Experimental/Chaos/Public/Chaos/EPA.h
@@ -8,6 +8,7 @@
 #include "ChaosCheck.h"
 #include "ChaosLog.h"
 #include "Templates/Function.h"
+#include <algorithm>
 
 namespace Chaos
 {
