Quantize time for more stability and scatter samples that span a quantized boundary.

#rnx

diff --git a/Engine/Source/Runtime/Experimental/IoStore/HttpClient/Private/Loop.inl b/Engine/Source/Runtime/Experimental/IoStore/HttpClient/Private/Loop.inl
--- a/Engine/Source/Runtime/Experimental/IoStore/HttpClient/Private/Loop.inl
+++ b/Engine/Source/Runtime/Experimental/IoStore/HttpClient/Private/Loop.inl
@@ -19,7 +19,6 @@ public:
 private:
 	enum
 	{
-		QuantizeBy			= 255,
 		BucketMs			= 1000,
 	};
 
@@ -51,15 +50,10 @@ void FRecvStats::Update(uint32 NowMs, uint32 RecvSizeKiB)
 	TotalRecvKiB.fetch_add(RecvSizeKiB, std::memory_order_relaxed);
 
 	// Quantize time for a modicum of stability.
-	static_assert(((QuantizeBy + 1) & QuantizeBy) == 0);
-	// NowMs = NowMs & ~QuantizeBy;
+	uint32 NowQuantMs = NowMs - (NowMs % BucketMs);
 
-	uint32 DeltaMs = NowMs - LastMs;
-	if (DeltaMs == 0)
-	{
-		return;
-	}
-	LastMs = NowMs;
+	uint32 DeltaMs = NowQuantMs - LastMs;
+	LastMs = NowQuantMs;
 
 	Accumulator.TimeMs += DeltaMs;
 	if (Accumulator.TimeMs < BucketMs)
@@ -69,8 +63,8 @@ void FRecvStats::Update(uint32 NowMs, uint32 RecvSizeKiB)
 	}
 
 	uint32 RecvKiB = Accumulator.RecvKiB;
-	uint32 TimeMs = NowMs - DeltaMs;
-	if (Accumulator.TimeMs >= BucketMs)
+	uint32 TimeMs = NowQuantMs - DeltaMs;
+	if (Accumulator.TimeMs > BucketMs)
 	{
 		Accumulator = {
 			.RecvKiB = RecvSizeKiB,
@@ -79,13 +73,14 @@ void FRecvStats::Update(uint32 NowMs, uint32 RecvSizeKiB)
 	}
 	else
 	{
-		uint32 Half = RecvSizeKiB; // (RecvSizeKiB + 1) >> 1;
+		uint32 Partial = ((NowMs - NowQuantMs) * RecvSizeKiB) / BucketMs;
+
+		RecvKiB += RecvSizeKiB - Partial,
 		Accumulator = {
-			.RecvKiB = RecvSizeKiB - Half,
+			.RecvKiB = Partial,
 			.TimeMs = BucketMs - Accumulator.TimeMs,
 		};
-		RecvKiB += Half;
-		TimeMs = NowMs - Accumulator.TimeMs;
+		TimeMs = NowQuantMs - Accumulator.TimeMs;
 	}
 
 	uint64 Value = (uint64(TimeMs) << 32u) + RecvKiB;
@@ -98,6 +93,11 @@ void FRecvStats::Update(uint32 NowMs, uint32 RecvSizeKiB)
 ////////////////////////////////////////////////////////////////////////////////
 void FRecvStats::Flush()
 {
+	if (Accumulator.RecvKiB == 0)
+	{
+		return;
+	}
+
 	uint64 Value = (uint64(LastMs) << 32u) + Accumulator.RecvKiB;
 	uint32 Index = Cursor % UE_ARRAY_COUNT(Buckets);
 	Buckets[Index].Value.store(Value, std::memory_order_release);
@@ -108,7 +108,7 @@ void FRecvStats::Sample(uint32 NowMs, FLoopPerf& Out) const
 {
 	static_assert((BucketMs * UE_ARRAY_COUNT(Buckets)) / 1000 == UE_ARRAY_COUNT(Out.RecvKiBps));
 
-	// NowMs = (NowMs + QuantizeBy) & ~QuantizeBy;
+	uint32 NowQuantMs = NowMs - (NowMs % BucketMs);
 
 	union FData
 	{
@@ -136,7 +136,7 @@ void FRecvStats::Sample(uint32 NowMs, FLoopPerf& Out) const
 	}
 
 	// Gather samples for the last N seconds.
-	for (uint32 i = 0, n = UE_ARRAY_COUNT(Out.RecvKiBps), TimeMs = NowMs; i < n; ++i)
+	for (uint32 i = 0, n = UE_ARRAY_COUNT(Out.RecvKiBps), TimeMs = NowQuantMs; i < n; ++i)
 	{
 		const FBucket::FData& Data = Datas[Index].Data;
 
@@ -658,7 +658,7 @@ uint32 FPeerGroup::CollectRecvSizeKiB()
 	}
 
 	uint64 Size = Peer.GetRecvSize();
-	Size += 512;
+	Size += 1024;
 	Size >>= 10;
 
 	uint32 Ret = uint32(Size) - LastRecvSizeKiB;
