Updated IAX debug text output

#rnx

diff --git a/Engine/Source/Runtime/Experimental/IoStore/OnDemand/Private/Statistics.cpp b/Engine/Source/Runtime/Experimental/IoStore/OnDemand/Private/Statistics.cpp
--- a/Engine/Source/Runtime/Experimental/IoStore/OnDemand/Private/Statistics.cpp
+++ b/Engine/Source/Runtime/Experimental/IoStore/OnDemand/Private/Statistics.cpp
@@ -819,17 +819,31 @@ void FOnDemandIoBackendStats::PrintOnScreenStatistics(FCoreDelegates::FSeverityM
 			};
 
 		{
+			static uint32 RateKiBps = 0;
+			if (IOnDemandIoStore* Store = TryGetOnDemandIoStore(); Store != nullptr)
+			{
+				double Time = FPlatformTime::Seconds();
+				if (static double LastSampleTime = Time; Time - LastSampleTime > 1.0)
+				{
+					LastSampleTime -= 1.0;
+					if (LastSampleTime > 1.0)
+					{
+						LastSampleTime = Time;
+					}
+
+					static FOnDemandHttpStats Stats;
+					Store->GetHttpStats(Stats);
+					RateKiBps = Stats.GetRecvKiBps();
+				}
+			}
+
+			uint32 RateLimit = GIasHttpRateLimitKiBPerSecond;
+			RateLimit = (RateLimit < 0) ? 0 : RateLimit;
+
 			TStringBuilder<256> Builder;
 			Builder << TEXT("IAX Http:");
 			Builder << TEXT(" v") << GIaxHttpVersion;
-			if (GIasHttpRateLimitKiBPerSecond > 0)
-			{
-				Builder << TEXT(" DataRate Cap: ") << GIasHttpRateLimitKiBPerSecond << TEXT("KiB/s");
-			}
-			else
-			{
-				Builder << TEXT(" DataRate Cap: None");
-			}
+			Builder << TEXT(" DataRate: ") << RateKiBps << TEXT(" / ") << RateLimit << TEXT(" KiB/s");
 
 			const int32 NumConcurrentRequests = GHttpNumConcurrentRequests;
 			Builder << TEXT(" ConcurrentRequests: ") << NumConcurrentRequests << TEXT("/") << GIasHttpConcurrentRequests;
