Minor clean up of IAX winsock initialization code.

#rnx

- Call WSACleanup rather than GetLastError if WSAGetLastError fails.
- Instead of storing the result of ::WSAStartup we now store a bool recording if WSA was started up correctly or not.
- Moved the error message parsing code to a common method.

diff --git a/Engine/Source/Runtime/Experimental/IoStore/OnDemand/Private/OnDemandModule.cpp b/Engine/Source/Runtime/Experimental/IoStore/OnDemand/Private/OnDemandModule.cpp
--- a/Engine/Source/Runtime/Experimental/IoStore/OnDemand/Private/OnDemandModule.cpp
+++ b/Engine/Source/Runtime/Experimental/IoStore/OnDemand/Private/OnDemandModule.cpp
@@ -15,48 +15,68 @@
 
 namespace UE::IoStore
 {
-namespace Private
-{
+
 ////////////////////////////////////////////////////////////////////////////////
 struct FPlatformSocketSystem
 {
-	#if PLATFORM_WINDOWS
+#if PLATFORM_WINDOWS
+	// Note that we use base types not UE types to match the platform API
+
 	void Startup()
 	{
-		WSADATA WsaData;
-		Result = WSAStartup(MAKEWORD(2, 2), &WsaData);
-		if (Result != 0)
+		WSADATA WSAData;
+		const int Result = WSAStartup(MAKEWORD(2, 2), &WSAData);
+		if (Result == 0)
+		{
+			bLoaded = true;
+		}
+		else
 		{
-			TCHAR SystemErrorMsg[MAX_SPRINTF] = { 0 };
-			FPlatformMisc::GetSystemErrorMessage(SystemErrorMsg, UE_ARRAY_COUNT(SystemErrorMsg), Result);
-			UE_LOG(LogIoStoreOnDemand, Error, TEXT("WSAStartup failed due to: %s (%d)"), SystemErrorMsg, Result);
+			// If WSAStartup fails it will return it's error code
+			LogError(TEXT("WSAStartup"), Result);
 		}
 	}
 
 	void Cleanup()
 	{
-		if (Result != 0)
+		if (!bLoaded)
 		{
 			return;
 		}
-		Result = WSACleanup();
+
+		const int Result = WSACleanup();
+
 		if (Result != 0)
 		{
-			const uint32 SystemError = FPlatformMisc::GetLastError();
-			TCHAR SystemErrorMsg[MAX_SPRINTF] = { 0 };
-			FPlatformMisc::GetSystemErrorMessage(SystemErrorMsg, UE_ARRAY_COUNT(SystemErrorMsg), SystemError);
-			UE_LOG(LogIoStoreOnDemand, Error, TEXT("WSACleanup failed due to: %s (%u)"), SystemErrorMsg, SystemError);
+			// If WSACleanup fails it returns SOCKET_ERROR and WSAGetLastError must be called instead
+			const int SystemError = WSAGetLastError();
+			LogError(TEXT("WSAGetLastError"), SystemError);
 		}
 	}
-	int Result = -1; 
-	#else
+
+private:
+
+	void LogError(const TCHAR* ErrorFunction, int ErrorCode)
+	{
+		TCHAR SystemErrorMsg[MAX_SPRINTF] = { 0 };
+		FPlatformMisc::GetSystemErrorMessage(SystemErrorMsg, UE_ARRAY_COUNT(SystemErrorMsg), ErrorCode);
+		UE_LOG(LogIoStoreOnDemand, Error, TEXT("%s failed due to: %s (%d)"), ErrorFunction, SystemErrorMsg, ErrorCode);
+	}
+
+	bool bLoaded = false;
+
+#else
+
 	void Startup()
-	{ }
+	{
+	}
 	void Cleanup()
-	{ }
-	#endif
+	{
+	}
+
+#endif // PLATFORM_WINDOWS
 };
-} // UE::IoStore::Private
+
 ////////////////////////////////////////////////////////////////////////////////
 class FOnDemandIoStoreFactory final
 	: public IOnDemandIoStoreFactory
@@ -86,7 +106,7 @@ public:
 	virtual void ShutdownModule() override;
 
 private:
-	Private::FPlatformSocketSystem PlatformSocketSystem;
+	FPlatformSocketSystem PlatformSocketSystem;
 	TUniquePtr<FOnDemandIoStoreFactory> Factory = MakeUnique<FOnDemandIoStoreFactory>();
 };
 IMPLEMENT_MODULE(UE::IoStore::FIoStoreOnDemandModule, IoStoreOnDemand);
