Add an assert to make sure the two arrays are the same size before iterating over them with the same index value.

#rnx

diff --git a/Engine/Source/Runtime/Experimental/IoStore/OnDemand/Private/OnDemandPackageStoreBackend.cpp b/Engine/Source/Runtime/Experimental/IoStore/OnDemand/Private/OnDemandPackageStoreBackend.cpp
--- a/Engine/Source/Runtime/Experimental/IoStore/OnDemand/Private/OnDemandPackageStoreBackend.cpp
+++ b/Engine/Source/Runtime/Experimental/IoStore/OnDemand/Private/OnDemandPackageStoreBackend.cpp
@@ -326,6 +326,7 @@ void FOnDemandPackageStoreBackend::UpdateReferencedPackages(const TArray<TBitArr
 	}
 	else
 	{
+		check(ReferencedChunkEntryIndices.Num() == Containers.Num());
 		for (int32 ContainerIdx = 0; FContainer& Container : Containers)
 		{
 			const TBitArray<>& Refs = ReferencedChunkEntryIndices[ContainerIdx];
