Reset size when a new handle is moved in.

#rnx

diff --git a/Engine/Source/Runtime/Experimental/IoStore/HttpClient/Private/Socket.inl b/Engine/Source/Runtime/Experimental/IoStore/HttpClient/Private/Socket.inl
--- a/Engine/Source/Runtime/Experimental/IoStore/HttpClient/Private/Socket.inl
+++ b/Engine/Source/Runtime/Experimental/IoStore/HttpClient/Private/Socket.inl
@@ -397,6 +397,7 @@ void FSocket::Move(FSocket&& Rhs)
 {
 	check(!IsValid() || !Rhs.IsValid()); // currently we only want to pass one around
 	Swap(Socket, Rhs.Socket);
+	RecvSize = 0;
 }
 
 ////////////////////////////////////////////////////////////////////////////////
