Maintain statistical accuracy

#rnx

#FailedValidations Preflight

diff --git a/Engine/Source/Runtime/Experimental/IoStore/HttpClient/Private/Loop.inl b/Engine/Source/Runtime/Experimental/IoStore/HttpClient/Private/Loop.inl
--- a/Engine/Source/Runtime/Experimental/IoStore/HttpClient/Private/Loop.inl
+++ b/Engine/Source/Runtime/Experimental/IoStore/HttpClient/Private/Loop.inl
@@ -638,7 +638,7 @@ private:
 	FActivityList		Send;
 	FActivityList		Recv;
 	FHttpPeer			Peer;
-	uint32				LastRecvSizeKiB = 0;
+	uint64				LastRecvSize = 0;
 	uint32				LastUseMs = 0;
 	uint8				IsKeepAlive = 0;
 	uint8				InflightNum = 0;
@@ -658,11 +658,9 @@ uint32 FPeerGroup::CollectRecvSizeKiB()
 	}
 
 	uint64 Size = Peer.GetRecvSize();
-	Size += 1024;
-	Size >>= 10;
+	uint32 Ret = uint32(Size - LastRecvSize) >> 10;
 
-	uint32 Ret = uint32(Size) - LastRecvSizeKiB;
-	LastRecvSizeKiB = uint32(Size);
+	LastRecvSize = Size;
 	return Ret;
 }
 
@@ -953,7 +951,7 @@ void FPeerGroup::TickSend(FTickState& State, FHost& Host, FPoller& Poller)
 		bNegotiating = true;
 		bWillBlock = Outcome.IsWaiting();
 		InflightNum = 0;
-		LastRecvSizeKiB = 0;
+		LastRecvSize = 0;
 
 		FHttpPeer::FParams Params = {
 			.Socket = MoveTemp(Socket),
