Incorrect variable used in comparison

#rnx

diff --git a/Engine/Source/Runtime/Experimental/IoStore/HttpClient/Private/Loop.inl b/Engine/Source/Runtime/Experimental/IoStore/HttpClient/Private/Loop.inl
--- a/Engine/Source/Runtime/Experimental/IoStore/HttpClient/Private/Loop.inl
+++ b/Engine/Source/Runtime/Experimental/IoStore/HttpClient/Private/Loop.inl
@@ -1292,7 +1292,7 @@ void FEventLoop::FImpl::Cancel(FTicket Ticket)
 void FEventLoop::FImpl::ReceiveWork()
 {
 	uint64 FreeSlotsLoad = FreeSlots.load(std::memory_order_relaxed);
-	if (FreeSlots == PrevFreeSlots)
+	if (FreeSlotsLoad == PrevFreeSlots)
 	{
 		return;
 	}
