IASTool download command now correctly parses it's string cmdline options.

#rnx

- The args are of type FStringView but were being converted to FString directly which caused them to lose the null terminator which breaks FString.
- Now we explicitly request the correct type and then construct FStrings from the views.
- Create a FString explicitly from TocPath rather than it's data pointer, it is faster and safer.

diff --git a/Engine/Source/Runtime/Experimental/IoStore/OnDemand/Private/Tool/Download.cpp b/Engine/Source/Runtime/Experimental/IoStore/OnDemand/Private/Tool/Download.cpp
--- a/Engine/Source/Runtime/Experimental/IoStore/OnDemand/Private/Tool/Download.cpp
+++ b/Engine/Source/Runtime/Experimental/IoStore/OnDemand/Private/Tool/Download.cpp
@@ -329,15 +329,16 @@ static FDownloadParams BuildDownloadParams(const FContext& Context)
 {
 	FDownloadParams Ret;
 
-	Ret.Directory				= Context.Get(TEXT("-Directory"),				Ret.Directory);
-	Ret.ServiceUrl				= Context.Get(TEXT("-ServiceUrl"),				Ret.ServiceUrl);
-	Ret.Bucket					= Context.Get(TEXT("-Bucket"),					Ret.Bucket);
-	Ret.Region					= Context.Get(TEXT("-Region"),					Ret.Region); 
-	Ret.AccessKey				= Context.Get(TEXT("-AccessKey"),				Ret.AccessKey);
-	Ret.SecretKey				= Context.Get(TEXT("-SecretKey"),				Ret.SecretKey);
-	Ret.SessionToken			= Context.Get(TEXT("-SessionToken"),			Ret.SessionToken);
-	Ret.CredentialsFile			= Context.Get(TEXT("-CredentialsFile"),			Ret.CredentialsFile);
-	Ret.CredentialsFileKeyName	= Context.Get(TEXT("-CredentialsFileKeyName"),	Ret.CredentialsFileKeyName);
+	Ret.Directory				= Context.Get<FStringView>(TEXT("-Directory"),				Ret.Directory);
+	Ret.ServiceUrl				= Context.Get<FStringView>(TEXT("-ServiceUrl"),				Ret.ServiceUrl);
+	Ret.Bucket					= Context.Get<FStringView>(TEXT("-Bucket"),					Ret.Bucket);
+	Ret.Region					= Context.Get<FStringView>(TEXT("-Region"),					Ret.Region);
+	Ret.AccessKey				= Context.Get<FStringView>(TEXT("-AccessKey"),				Ret.AccessKey);
+	Ret.SecretKey				= Context.Get<FStringView>(TEXT("-SecretKey"),				Ret.SecretKey);
+	Ret.SessionToken			= Context.Get<FStringView>(TEXT("-SessionToken"),			Ret.SessionToken);
+	Ret.CredentialsFile			= Context.Get<FStringView>(TEXT("-CredentialsFile"),		Ret.CredentialsFile);
+	Ret.CredentialsFileKeyName	= Context.Get<FStringView>(TEXT("-CredentialsFileKeyName"),	Ret.CredentialsFileKeyName);
+
 	Ret.MaxConcurrentDownloads	= Context.Get(TEXT("-MaxConcurrentDownloads"),	Ret.MaxConcurrentDownloads);
 
 	if (FIoStatus Validation = Ret.Validate(); !Validation.IsOk())
@@ -356,7 +357,7 @@ static int32 DownloadCommandEntry(const Tool::FContext& Context)
 
 	FDownloadParams Params = BuildDownloadParams(Context);
 
-	FIoStatus Status = DownloadContainerFiles(Params, TocPath.GetData());
+	FIoStatus Status = DownloadContainerFiles(Params, FString(TocPath));
 	if (!Status.IsOk())
 	{
 		FString Reason = Status.ToString();
