Only tear when when very idle.

#rnx

#FailedValidations Preflight

diff --git a/Engine/Source/Runtime/Experimental/IoStore/HttpClient/Private/Loop.inl b/Engine/Source/Runtime/Experimental/IoStore/HttpClient/Private/Loop.inl
--- a/Engine/Source/Runtime/Experimental/IoStore/HttpClient/Private/Loop.inl
+++ b/Engine/Source/Runtime/Experimental/IoStore/HttpClient/Private/Loop.inl
@@ -641,8 +641,8 @@ private:
 	uint64				LastRecvSize = 0;
 	uint32				LastUseMs = 0;
 	uint8				IsKeepAlive = 0;
-	uint8				InflightNum = 0;
-	uint8				InflightMax = 1;
+	int8				InflightNum = 0;
+	int8				InflightMax = 1;
 	bool				bNegotiating = false;
 	bool				bWaiting = false;
 
@@ -831,6 +831,8 @@ void FPeerGroup::RecvInternal(FTickState& State)
 
 	check(Recv.GetHead() == Activity);
 	Recv.PopPrepend(State.DoneList);
+
+	check(InflightNum > 0);
 	InflightNum--;
 
 	if (MaybeStartNewWork(State))
@@ -913,7 +915,7 @@ void FPeerGroup::TickSend(FTickState& State, FHost& Host, FPoller& Poller)
 	check(Pending != nullptr);
 
 	// Close idle sockets
-	if (Peer.IsValid() && LastUseMs + GIdleMs < State.NowMs)
+	if (Peer.IsValid() && InflightNum == 0 && LastUseMs + GIdleMs < State.NowMs)
 	{
 		LastUseMs = State.NowMs;
 		Peer = FHttpPeer();
