Prevent hidden errors when the wrong argument type is used.

#rnx

diff --git a/Engine/Source/Runtime/Experimental/IoStore/OnDemand/Private/Tool/Command.h b/Engine/Source/Runtime/Experimental/IoStore/OnDemand/Private/Tool/Command.h
--- a/Engine/Source/Runtime/Experimental/IoStore/OnDemand/Private/Tool/Command.h
+++ b/Engine/Source/Runtime/Experimental/IoStore/OnDemand/Private/Tool/Command.h
@@ -34,7 +34,16 @@ struct FArgumentSet
 };
 
 ////////////////////////////////////////////////////////////////////////////////
-template <typename Type> FArgument TArgument(FStringView Name, FStringView Desc);
+template <typename T>
+concept CArgTypes =
+	std::is_same_v<T, bool>
+	|| std::is_integral_v<T>
+	|| std::is_floating_point_v<T>
+	|| std::is_same_v<T, FStringView>;
+
+
+////////////////////////////////////////////////////////////////////////////////
+template <CArgTypes Type> FArgument TArgument(FStringView Name, FStringView Desc);
 
 #define TARGUMENT_DECL(t) \
 	template <> FArgument TArgument<t>(FStringView Name, FStringView Desc);
@@ -52,8 +61,8 @@ TARGUMENT_DECL(float)	TARGUMENT_DECL(double)
 class FContext
 {
 public:
-	template <typename Type> Type	Get(FStringView Name) const;
-	template <typename Type> Type	Get(FStringView Name, Type Default) const;
+	template <CArgTypes Type> Type	Get(FStringView Name) const;
+	template <CArgTypes Type> Type	Get(FStringView Name, Type Default) const;
 	[[noreturn]] void				Abort(const TCHAR* Reason) const;
 
 private:
@@ -74,7 +83,7 @@ private:
 };
 
 ////////////////////////////////////////////////////////////////////////////////
-template <typename Type>
+template <CArgTypes Type>
 Type FContext::Get(FStringView Name) const
 {
 	if (Name.GetData()[0] == '-')
@@ -92,7 +101,7 @@ Type FContext::Get(FStringView Name) const
 }
 
 ////////////////////////////////////////////////////////////////////////////////
-template <typename Type>
+template <CArgTypes Type>
 Type FContext::Get(FStringView Name, Type Default) const
 {
 	uint32 Index = FindArgument(Name);
