Chaos: Addresses an error in parallel bounding volume calculation in ComputeAllWorldSpaceBoundingBoxes.

diff --git a/Engine/Source/Runtime/Experimental/Chaos/Public/Chaos/BoundingVolumeUtilities.h b/Engine/Source/Runtime/Experimental/Chaos/Public/Chaos/BoundingVolumeUtilities.h
--- a/Engine/Source/Runtime/Experimental/Chaos/Public/Chaos/BoundingVolumeUtilities.h
+++ b/Engine/Source/Runtime/Experimental/Chaos/Public/Chaos/BoundingVolumeUtilities.h
@@ -393,7 +393,7 @@ typename TEnableIf<TModels_V<CParticleView, ParticleView>>::Type ComputeAllWorld
 			{
 				if (const auto PBDRigid = Particle.AsDynamic())
 				{
-					WorldSpaceBoxes.Last().ThickenSymmetrically(ComputeBoundsThickness(*PBDRigid, Dt, 0, 1));
+					WorldSpaceBoxes[Index].ThickenSymmetrically(ComputeBoundsThickness(*PBDRigid, Dt, 0, 1));
 				}
 			}
 		}
