Fix warning V547: Expression 'RateLimit < 0' is always false. Unsigned type value is never < 0.
#rnx

diff --git a/Engine/Source/Runtime/Experimental/IoStore/OnDemand/Private/Statistics.cpp b/Engine/Source/Runtime/Experimental/IoStore/OnDemand/Private/Statistics.cpp
--- a/Engine/Source/Runtime/Experimental/IoStore/OnDemand/Private/Statistics.cpp
+++ b/Engine/Source/Runtime/Experimental/IoStore/OnDemand/Private/Statistics.cpp
@@ -837,8 +837,7 @@ void FOnDemandIoBackendStats::PrintOnScreenStatistics(FCoreDelegates::FSeverityM
 				}
 			}
 
-			uint32 RateLimit = GIasHttpRateLimitKiBPerSecond;
-			RateLimit = (RateLimit < 0) ? 0 : RateLimit;
+			const uint32 RateLimit = (GIasHttpRateLimitKiBPerSecond < 0) ? 0 : GIasHttpRateLimitKiBPerSecond;
 
 			TStringBuilder<256> Builder;
 			Builder << TEXT("IAX Http:");
