[ChaosVD] Fix for an issue where the data overflow protection is triggered immediatly.

TGuardInt needs to be initialized with a valid value during construction, otherwise it stays flagged as invalid and any math operation on it will be ignored.

diff --git a/Engine/Source/Runtime/Experimental/ChaosVisualDebugger/Private/ChaosVDRuntimeModule.cpp b/Engine/Source/Runtime/Experimental/ChaosVisualDebugger/Private/ChaosVDRuntimeModule.cpp
--- a/Engine/Source/Runtime/Experimental/ChaosVisualDebugger/Private/ChaosVDRuntimeModule.cpp
+++ b/Engine/Source/Runtime/Experimental/ChaosVisualDebugger/Private/ChaosVDRuntimeModule.cpp
@@ -161,7 +161,7 @@ int32 FChaosVDRuntimeModule::GenerateUniqueID()
 
 int64 FChaosVDRuntimeModule::GetBufferedDataBytesSize() const
 {
-	FGuardedInt64 RelaySendBufferSize;
+	FGuardedInt64 RelaySendBufferSize(0);
 
 	if (RelayWriter)
 	{
diff --git a/Engine/Source/Runtime/Experimental/ChaosVisualDebugger/Private/Trace/RelayTraceWriter.cpp b/Engine/Source/Runtime/Experimental/ChaosVisualDebugger/Private/Trace/RelayTraceWriter.cpp
--- a/Engine/Source/Runtime/Experimental/ChaosVisualDebugger/Private/Trace/RelayTraceWriter.cpp
+++ b/Engine/Source/Runtime/Experimental/ChaosVisualDebugger/Private/Trace/RelayTraceWriter.cpp
@@ -153,7 +153,7 @@ namespace Chaos::VD
 
 	bool FRelayTraceWriter::CanWriteData(uint32 NewDataSize) const
 	{
-		TGuardedInt<uint64> ExpectedQueuedBytes;
+		TGuardedInt<uint64> ExpectedQueuedBytes(0);
 
 		ExpectedQueuedBytes+= QueuedBytesNum;
 
