Don't falsely act on the end flag when clamped buffer was exhausted

#rnx

diff --git a/Engine/Source/Runtime/Experimental/IoStore/HttpClient/Private/TransactionTwo.inl b/Engine/Source/Runtime/Experimental/IoStore/HttpClient/Private/TransactionTwo.inl
--- a/Engine/Source/Runtime/Experimental/IoStore/HttpClient/Private/TransactionTwo.inl
+++ b/Engine/Source/Runtime/Experimental/IoStore/HttpClient/Private/TransactionTwo.inl
@@ -514,7 +514,7 @@ FOutcome FDriverNg::RecvBody(FTlsPeer& Peer, FMutableMemoryView Dest)
 		FMutableMemoryView View = Dest.Left(RecvRemaining);
 		if (View.IsEmpty())
 		{
-			if (EndOfStream)
+			if (RecvRemaining <= 0 && EndOfStream)
 			{
 				RecvFrame = {};
 				return FOutcome::Ok(RecvSize);
