Frame logic should be included inside the loop as the loop fetches new frames too.

#rnx

diff --git a/Engine/Source/Runtime/Experimental/IoStore/HttpClient/Private/TransactionTwo.inl b/Engine/Source/Runtime/Experimental/IoStore/HttpClient/Private/TransactionTwo.inl
--- a/Engine/Source/Runtime/Experimental/IoStore/HttpClient/Private/TransactionTwo.inl
+++ b/Engine/Source/Runtime/Experimental/IoStore/HttpClient/Private/TransactionTwo.inl
@@ -467,20 +467,20 @@ FOutcome FDriverNg::Send(FTlsPeer& Peer)
 ////////////////////////////////////////////////////////////////////////////////
 FOutcome FDriverNg::RecvResponse(FTlsPeer& Peer, TArray<FHeaderNg>* Sink)
 {
-	if (RecvFrame.Type == FFrame::EType::RstStream)
-	{
-		return FOutcome::Error("Received stream reset");
-	}
-
-	// Assume if we're filling we are expecting something
-	check(RecvRemaining > 0);
-	check(RecvFrame.Type == FFrame::EType::Headers);
-
 	ResponseHeaders = Sink;
 	ON_SCOPE_EXIT { ResponseHeaders = nullptr; };
 
 	while (true)
 	{
+		if (RecvFrame.Type == FFrame::EType::RstStream)
+		{
+			return FOutcome::Error("Received stream reset");
+		}
+
+		// Assume if we're filling we are expecting something
+		check(RecvRemaining > 0);
+		check(RecvFrame.Type == FFrame::EType::Headers);
+
 		int32 EndOfHeaders = (RecvFrame.Flags & FFrame::Flag_EndHeaders);
 		uint32 ThisId = RecvFrame.GetId();
 
@@ -499,16 +499,16 @@ FOutcome FDriverNg::RecvResponse(FTlsPeer& Peer, TArray<FHeaderNg>* Sink)
 ////////////////////////////////////////////////////////////////////////////////
 FOutcome FDriverNg::RecvBody(FTlsPeer& Peer, FMutableMemoryView Dest)
 {
-	if (RecvFrame.Type == FFrame::EType::RstStream)
+	for (uint32 RecvSize = 0;;)
 	{
-		return FOutcome::Error("Received stream reset");
-	}
+		if (RecvFrame.Type == FFrame::EType::RstStream)
+		{
+			return FOutcome::Error("Received stream reset");
+		}
 
-	check(RecvRemaining >= 0);
-	check(RecvFrame.Type == FFrame::EType::Data);
+		check(RecvRemaining >= 0);
+		check(RecvFrame.Type == FFrame::EType::Data);
 
-	for (uint32 RecvSize = 0;;)
-	{
 		int32 EndOfStream = (RecvFrame.Flags & FFrame::Flag_EndStream);
 
 		FMutableMemoryView View = Dest.Left(RecvRemaining);
