Added IasTool package graph command to build, export and query cooked package dependencies.

Usage:
BuildGraph <Path/WildCardToUToc> -CryptoKeys=<PathToCryptoJson> -Out=<Path.[json|cb|csv]>
QueryGraph <Path.cb> -Traverse <PackageID> [-Incoming] [-SkipSoftRefs]
QueryGraph <Path.cb> -Deps <PackageID> [-Incoming]

#rnx

diff --git a/Engine/Source/Runtime/Experimental/IoStore/OnDemand/Private/Tool/Graph.cpp b/Engine/Source/Runtime/Experimental/IoStore/OnDemand/Private/Tool/Graph.cpp
new file mode 100644
--- /dev/null
+++ b/Engine/Source/Runtime/Experimental/IoStore/OnDemand/Private/Tool/Graph.cpp
@@ -0,0 +1,1130 @@
+// Copyright Epic Games, Inc. All Rights Reserved.
+
+#if defined(UE_WITH_IAS_TOOL)
+
+#include "Graph.h"
+#include "Command.h"
+
+#include "HAL/FileManager.h"
+#include "HAL/Platform.h"
+#include "IO/IoContainerHeader.h"
+#include "IO/IoStore.h"
+#include "IO/IoStoreOnDemand.h"
+#include "IO/PackageId.h"
+#include "Misc/CString.h"
+#include "Misc/KeyChainUtilities.h"
+#include "Misc/Paths.h"
+#include "Misc/Parse.h"
+#include "Serialization/CompactBinary.h"
+#include "Serialization/CompactBinaryWriter.h"
+#include "Serialization/CompactBinarySerialization.h"
+#include "Serialization/LargeMemoryReader.h"
+
+namespace UE::IoStore::Tool
+{
+
+////////////////////////////////////////////////////////////////////////////////
+FIoStatus ExportGephiCsv(const FPackageGraph& Graph, const FString& OutFilename)
+{
+	UE_LOG(LogIoStoreOnDemand, Display, TEXT("Exporting CSV with %d nodes and %d edges"), Graph.GetNodeCount(), Graph.GetEdgeCount());
+	IFileManager& Ifm = IFileManager::Get();
+	{
+		TUniquePtr<FArchive> Ar(Ifm.CreateFileWriter(*OutFilename));
+		if (!Ar.IsValid())
+		{
+			return FIoStatus(EIoErrorCode::WriteError);
+		}
+
+		UE_LOG(LogIoStoreOnDemand, Display, TEXT("Saving '%s'"), *OutFilename);
+		Ar->Logf(TEXT("Id, Label"));
+		for (int32 Idx = 0; Idx < Graph.GetNodeCount(); ++Idx)
+		{
+			FPackageGraph::FNode			Node{Idx};
+			const FPackageGraph::FNodeData&	Data = Graph.GetNodeData(Node);
+			Ar->Logf(TEXT("%d, %s"), Idx, *LexToString(Data.PackageId));
+		}
+	}
+	{
+		FString Filename		= FPaths::GetBaseFilename(OutFilename);
+		FString Path			= FPaths::GetPath(OutFilename);
+		FString EdgesFilename	= Path / Filename + TEXT("-edges.csv");
+
+		TUniquePtr<FArchive> Ar(Ifm.CreateFileWriter(*EdgesFilename));
+		if (!Ar.IsValid())
+		{
+			return FIoStatus(EIoErrorCode::WriteError);
+		}
+
+		UE_LOG(LogIoStoreOnDemand, Display, TEXT("Saving '%s'"), *EdgesFilename);
+		Ar->Logf(TEXT("Source, Target, Type, Weight, Id"));
+		for (int32 Idx = 0; Idx < Graph.GetEdgeCount(); ++Idx)
+		{
+			const FPackageGraph::FEdgeData&	Edge	= Graph.GetEdgeData(FPackageGraph::FEdge { .Index = Idx });
+			const float						Weight	= 1.0f;
+			const TCHAR*					Type	= TEXT("Directed");
+			Ar->Logf(TEXT("%d, %d, %s, %.2f, %d"), Edge.From.Index, Edge.To.Index, Type, Weight, Idx);
+		}
+	}
+
+	return FIoStatus::Ok;
+}
+
+////////////////////////////////////////////////////////////////////////////////
+void FPackageGraph::Reserve(int32 ContainerCount, int32 PackageCount)
+{
+	ContainerNodeData.Reserve(ContainerCount);
+	NodeData.Reserve(PackageCount);
+	EdgeData.Reserve(PackageCount);
+	PackageIdToNode.Reserve(PackageCount);
+}
+
+FPackageGraph::FContainerNode FPackageGraph::AddContainerNode(FString&& Name, const FIoContainerId& Id)
+{
+	const int32 Index = ContainerNodeData.Num();
+	ContainerNodeData.Add(FContainerData
+	{
+		.Name	= MoveTemp(Name),
+		.Id		= Id
+	});
+
+	return FContainerNode { .Index = Index };
+}
+
+FPackageGraph::FNode FPackageGraph::AddNode(
+	FPackageGraph::FContainerNode ContainerNode,
+	const FPackageId& PackageId,
+	const FString* Filename,
+	uint64 ChunkSize)
+{
+	check(ContainerNode.IsValid());
+	if (FNode Node = PackageIdToNode.FindRef(PackageId); Node.IsValid())
+	{
+		return Node;
+	}
+
+	const int32 Index	= NodeData.Num();
+	FNodeData& Data		= NodeData.AddDefaulted_GetRef();
+	Data.PackageId		= PackageId;
+	Data.Size			= ChunkSize;
+	Data.ContainerNode	= ContainerNode;
+
+	if (Filename != nullptr)
+	{
+		Data.Filename	= *Filename;
+	}
+
+	FNode NewNode{Index};
+	PackageIdToNode.Add(PackageId, NewNode);
+	ContainerNodeData[ContainerNode.Index].PackageNodes.Add(NewNode);
+
+	return NewNode;
+}
+
+void FPackageGraph::AddEdge(FNode From, FNode To, bool bSoft)
+{
+	check(From.IsValid() && To.IsValid());
+
+	const FEdge Edge { .Index = EdgeData.Num() };
+	EdgeData.Add(FEdgeData { .From = From, .To = To, .bSoft = bSoft });
+
+	FNodeData& FromNodeData = NodeData[From.Index];
+	FNodeData& ToNodeData	= NodeData[To.Index];
+
+	FromNodeData.Connections.Out.Add(Edge);
+	ToNodeData.Connections.In.Add(Edge);
+}
+
+void FPackageGraph::AddMissingEdge(FNode From, const FPackageId To, bool bSoft)
+{
+	check(From.IsValid()); 
+
+	const FMissingEdge MissingEdge { .Index = MissingEdgeData.Num() };
+	MissingEdgeData.Add(FMissingEdgeData { .From = From, .To = To, .bSoft = bSoft});
+
+	FNodeData& FromNodeData = NodeData[From.Index];
+}
+
+void FPackageGraph::AddRedirect(const FPackageId& From, const FPackageId& To)
+{
+	Redirects.Add(FRedirect
+	{
+		.From	= From,
+		.To		= To
+	});
+}
+
+void FPackageGraph::Traverse(FNode StartNode, bool bIncoming, bool bSkipSoftRefs, FVisitor& Visitor)
+{
+	TSet<FNode>		Visited;
+	TArray<FNode>	Stack;
+
+	check(StartNode.IsValid());
+
+	Stack.Push(StartNode);
+	while (Stack.IsEmpty() == false)
+	{
+		const FNode Node = Stack.Pop();
+
+		bool bAlreadyInSet = false;
+		Visited.Add(Node, &bAlreadyInSet);
+		if (bAlreadyInSet)
+		{
+			continue;
+		}
+
+		const FNodeData& Data = GetNodeData(Node);
+		if (Visitor.Visit(*this, Node, Data) == false)
+		{
+			return;
+		}
+
+		const TSet<FEdge>& Edges = bIncoming ? Data.Connections.In : Data.Connections.Out;
+		for (FEdge Edge : Edges)
+		{
+			const FEdgeData& ToTraverseEdgeData = EdgeData[Edge.Index];
+			if (ToTraverseEdgeData.bSoft && bSkipSoftRefs)
+			{
+				continue;
+			}
+
+			const FNode ToTraverse = bIncoming ? ToTraverseEdgeData.From : ToTraverseEdgeData.To;
+			if (Visited.Contains(ToTraverse) == false)
+			{
+				Stack.Push(ToTraverse);
+			}
+		}
+	}
+}
+
+void FPackageGraph::Build(
+	TConstArrayView<FString> ContainerFilenames,
+	const FKeyChain& KeyChain,
+	bool bIncludeSoftRefs,
+	FPackageGraph& Graph)
+{
+	UE_LOG(LogIoStoreOnDemand, Display, TEXT("Building package graph..."));
+
+	TMap<FGuid, FAES::FAESKey> EncryptionKeys;
+	for (const TPair<FGuid, FNamedAESKey>& KeyPair: KeyChain.GetEncryptionKeys())
+	{
+		EncryptionKeys.Add(KeyPair.Key, KeyPair.Value.Key);
+	}
+
+	struct FLoadedContainer
+	{
+		FString						Filename;
+		FString						Name;
+		FIoStoreReader				Reader;
+		FIoContainerHeader			Header;
+		TMap<FIoChunkId, FString>	ChunkFileNamesMap;
+	};
+	TArray<TUniquePtr<FLoadedContainer>> Containers;
+	Containers.Reserve(ContainerFilenames.Num());
+	int32 TotalPackages = 0;
+
+	// Load all containers
+	for (const FString& Filename : ContainerFilenames)
+	{
+		// Only .utoc file(s) 
+		if (FPaths::GetExtension(Filename, false) != TEXT("utoc"))
+		{
+			continue;
+		}
+
+		// Skip .o.utoc
+		if (FPaths::GetExtension(FPaths::GetBaseFilename(Filename), false) == TEXT("o"))
+		{
+			continue;
+		}
+
+		TUniquePtr<FLoadedContainer>	Container = MakeUnique<FLoadedContainer>();
+		FString							Path = Filename;
+
+		FIoStatus Status = Container->Reader.Initialize(*FPaths::ChangeExtension(Path, TEXT("")), EncryptionKeys);
+		if (!Status.IsOk())
+		{
+			UE_LOG(LogIoStoreOnDemand, Error, TEXT("Failed to open container '%s' for reading due to: '%s'"), *Path, *Status.ToString());
+			continue;
+		}
+
+		const FIoChunkId		HeaderChunkId = CreateContainerHeaderChunkId(Container->Reader.GetContainerId());
+		TIoStatusOr<FIoBuffer>	Chunk = Container->Reader.Read(HeaderChunkId, FIoReadOptions());
+
+		if (Chunk.IsOk() == false)
+		{
+			// Container header isn't required 
+			continue;
+		}
+
+		FIoBuffer			Buffer = Chunk.ValueOrDie();
+		FLargeMemoryReader	Ar(Buffer.GetData(), Buffer.GetSize());
+
+		Ar << Container->Header;
+		if (Ar.IsError() || Ar.IsCriticalError())
+		{
+			UE_LOG(LogIoStoreOnDemand, Error, TEXT("Failed to serialize container header, filename '%s'"), *Filename);
+			continue;
+		}
+
+		if (Container->Header.SoftPackageReferencesSerialInfo.Size > 0)
+		{
+			Ar.Seek(Container->Header.SoftPackageReferencesSerialInfo.Offset);
+			Ar << Container->Header.SoftPackageReferences;
+		}
+
+		if (Ar.IsError() || Ar.IsCriticalError())
+		{
+			UE_LOG(LogIoStoreOnDemand, Error, TEXT("Failed to serialize container soft references, filename '%s'"), *Filename);
+			continue;
+		}
+
+		for (const FIoContainerHeaderPackageRedirect& Redirect : Container->Header.PackageRedirects)
+		{
+			Graph.AddRedirect(Redirect.SourcePackageId, Redirect.TargetPackageId); 
+		}
+
+		Container->Reader.GetDirectoryIndexReader().IterateDirectoryIndex(FIoDirectoryIndexHandle::RootDirectory(), TEXT(""),
+			[&Container](FStringView Filename, uint32 TocEntryIndex) -> bool
+			{
+				TIoStatusOr<FIoStoreTocChunkInfo> ChunkInfo = Container->Reader.GetChunkInfo(TocEntryIndex);
+				if (ChunkInfo.IsOk())
+				{
+					Container->ChunkFileNamesMap.Add(ChunkInfo.ValueOrDie().Id, FString(Filename));
+				}
+				return true;
+			});
+
+
+		TotalPackages += Container->Header.PackageIds.Num();
+		Container->Name = FPaths::GetBaseFilename(Filename);
+		Containers.Add(MoveTemp(Container));
+	}
+
+	Graph.Reserve(Containers.Num(), TotalPackages);
+
+	// Create graph nodes 
+	for (const TUniquePtr<FLoadedContainer>& Container : Containers)
+	{
+		const FString&		Filename = Container->Filename;
+		FIoStoreReader&		Reader = Container->Reader;
+		FIoContainerHeader&	Header = Container->Header;
+
+		FString Name = Container->Name;
+		const FPackageGraph::FContainerNode ContainerNode = Graph.AddContainerNode(MoveTemp(Name), Reader.GetContainerId());
+		const int32							NodeCount = Graph.GetNodeCount();
+
+		for (int32 Idx = 0; Idx < Header.PackageIds.Num(); ++Idx)
+		{
+			const FPackageId	PackageId = Header.PackageIds[Idx];
+			const FIoChunkId	ChunkId = CreatePackageDataChunkId(PackageId);
+			const FString*		ChunkFilename = Container->ChunkFileNamesMap.Find(ChunkId);
+
+			TIoStatusOr<FIoStoreTocChunkInfo> ChunkInfo = Reader.GetChunkInfo(ChunkId);
+			if (ChunkInfo.IsOk() == false)
+			{
+				UE_LOG(LogIoStoreOnDemand, Error, TEXT("Failed to find chunk info for package '%s'"), *LexToString(PackageId));
+				continue;
+			}
+
+			Graph.AddNode(ContainerNode, PackageId, ChunkFilename, Align(ChunkInfo.ValueOrDie().CompressedSize, 16));
+		}
+
+		UE_LOG(LogIoStoreOnDemand, Display, TEXT("Created %d node(s) from '%s'"), Graph.GetNodeCount() - NodeCount, *Container->Name);
+	}
+
+	// Create graph edges
+	for (const TUniquePtr<FLoadedContainer>& Container : Containers)
+	{
+		const FString&		Filename = Container->Filename;
+		FIoStoreReader&		Reader = Container->Reader;
+		FIoContainerHeader&	Header = Container->Header;
+		const int32			EdgeCount = Graph.GetEdgeCount();
+
+		TConstArrayView<FFilePackageStoreEntry> PackageStoreEntries = MakeArrayView(
+			reinterpret_cast<const FFilePackageStoreEntry*>(Header.StoreEntries.GetData()), Header.PackageIds.Num());
+
+		TConstArrayView<FFilePackageStoreEntrySoftReferences> AllSoftReferences;
+		if (Header.SoftPackageReferences.bContainsSoftPackageReferences)
+		{
+			AllSoftReferences = MakeArrayView<const FFilePackageStoreEntrySoftReferences>(
+				reinterpret_cast<const FFilePackageStoreEntrySoftReferences*>(Header.SoftPackageReferences.PackageIndices.GetData()),
+				Header.PackageIds.Num());
+		}
+
+		for (int32 Idx = 0; Idx < Header.PackageIds.Num(); ++Idx)
+		{
+			const FPackageId PackageId			= Header.PackageIds[Idx];
+			const FFilePackageStoreEntry& Entry = PackageStoreEntries[Idx];
+
+			// Hard reference(s)
+			FPackageGraph::FNode FromNode = Graph.GetNode(PackageId);
+			check(FromNode.IsValid());
+			for (const FPackageId& ImportedPackageId : Entry.ImportedPackages)
+			{
+				const bool					bSoftEdge = false;
+				const FPackageGraph::FNode	ToNode = Graph.GetNode(ImportedPackageId);
+				if (ToNode.IsValid())
+				{
+					//TODO: Dedup edges
+					Graph.AddEdge(FromNode, ToNode, bSoftEdge);
+				}
+				else
+				{
+					Graph.AddMissingEdge(FromNode, ImportedPackageId, bSoftEdge); 
+				}
+			}
+
+			// Soft reference(s)
+			if (bIncludeSoftRefs && AllSoftReferences.IsEmpty() == false)
+			{
+				const FFilePackageStoreEntrySoftReferences& SoftRefs = AllSoftReferences[Idx];
+				if (SoftRefs.Indices.Num() > 0)
+				{
+					for (int32 SoftRefIdx : SoftRefs.Indices)
+					{
+						const bool					bSoftEdge = true;
+						const FPackageId			SoftRef = Header.SoftPackageReferences.PackageIds[SoftRefIdx];
+						const FPackageGraph::FNode	ToNode = Graph.GetNode(SoftRef);
+						if (ToNode.IsValid())
+						{
+							Graph.AddEdge(FromNode, ToNode, bSoftEdge);
+						}
+						else
+						{
+							Graph.AddMissingEdge(FromNode, SoftRef, bSoftEdge);
+						}
+					}
+				}
+			}
+		}
+
+		UE_LOG(LogIoStoreOnDemand, Display, TEXT("Created %d edges(s) from '%s'"), Graph.GetEdgeCount() - EdgeCount, *Container->Name);
+	}
+}
+
+FIoStatus FPackageGraph::Save(const FPackageGraph& Graph, const FString& Filename)
+{
+	const FString Ext = FPaths::GetExtension(Filename);
+
+	if (Ext == TEXT("json") || Ext == TEXT("cb"))
+	{
+		UE_LOG(LogIoStoreOnDemand, Display, TEXT("Saving graph to '%s'"), *Filename);
+
+		IFileManager&	Ifm = IFileManager::Get();
+		FCbWriter		Writer;
+
+		Writer << Graph;
+
+		TUniquePtr<FArchive> Ar = TUniquePtr<FArchive>(Ifm.CreateFileWriter(*Filename));
+		if (Ar.IsValid())
+		{
+			if (Ext == TEXT("cb"))
+			{
+				Writer.Save(*Ar);
+			}
+			else
+			{
+				FCbObject GraphCb = Writer.Save().AsObject();
+				FUtf8StringBuilderBase Sb;
+				CompactBinaryToJson(GraphCb, Sb);
+				UTF8CHAR UTF8BOM[] = { (UTF8CHAR)0xEF, (UTF8CHAR)0xBB, (UTF8CHAR)0xBF };
+				Ar->Serialize(&UTF8BOM, UE_ARRAY_COUNT(UTF8BOM) * sizeof(UTF8CHAR));
+				Ar->Serialize(Sb.GetData(), Sb.Len() * sizeof(UTF8CHAR));
+			}
+			Ar->Close();
+
+			if (Ar->IsError() || Ar->IsCriticalError())
+			{
+				return FIoStatus(EIoErrorCode::WriteError);
+			}
+
+			return FIoStatus::Ok;
+		}
+		else
+		{
+			return FIoStatus(EIoErrorCode::FileOpenFailed);
+		}
+	}
+	else if (Ext == TEXT("csv"))
+	{
+		return ExportGephiCsv(Graph, Filename);
+	}
+
+	return FIoStatus(EIoErrorCode::InvalidParameter, TEXT("Unsupported file format"));
+}
+
+////////////////////////////////////////////////////////////////////////////////
+static bool LoadFromCompactBinary(FCbFieldView Field, FPackageGraph::FIndex& Index)
+{
+	if (Field.IsInteger())
+	{
+		Index.Index = Field.AsInt32();
+		return true;
+	}
+
+	return false;
+}
+
+////////////////////////////////////////////////////////////////////////////////
+static FCbWriter& operator<<(FCbWriter& Writer, const FPackageGraph::FIndex& Index)
+{
+	Writer.AddInteger(Index.Index);
+	return Writer;
+}
+
+////////////////////////////////////////////////////////////////////////////////
+static bool LoadFromCompactBinary(FCbFieldView Field, FPackageGraph::FEdgeData& Edge)
+{
+	if (FCbObjectView Obj = Field.AsObjectView())
+	{
+		Edge.From	= FPackageGraph::FNode { .Index = Obj["From"].AsInt32() };
+		Edge.To		= FPackageGraph::FNode { .Index = Obj["To"].AsInt32() };
+		Edge.bSoft	= Obj["Soft"].AsBool();
+
+		return true;
+	}
+
+	return false;
+}
+
+////////////////////////////////////////////////////////////////////////////////
+static FCbWriter& operator<<(FCbWriter& Writer, const FPackageGraph::FEdgeData& Edge)
+{
+	Writer.BeginObject();
+	Writer << UTF8TEXT("From") << Edge.From;
+	Writer << UTF8TEXT("To") << Edge.To;
+	Writer << UTF8TEXT("Soft") << Edge.bSoft;
+	Writer.EndObject();
+
+	return Writer;
+}
+
+////////////////////////////////////////////////////////////////////////////////
+static bool LoadFromCompactBinary(FCbFieldView Field, FPackageGraph::FNodeData& Node)
+{
+	FCbObjectView Obj = Field.AsObjectView();
+
+	if (!Obj)
+	{
+		return false;
+	}
+
+	Node.Filename		= FString(Obj["Filename"].AsString());
+	Node.PackageId		= FPackageId::FromValue(Obj["PackageId"].AsUInt64());
+	Node.Size			= Obj["Size"].AsUInt64();
+	Node.ContainerNode	= FPackageGraph::FContainerNode { .Index = Obj["Container"].AsInt32() };
+
+	if (FCbArrayView ArrayView = Obj["In"].AsArrayView())
+	{
+		Node.Connections.In.Reserve(int32(ArrayView.Num()));
+		for (FCbFieldView ArrayField : ArrayView)
+		{
+			Node.Connections.In.Add(FPackageGraph::FEdge { .Index = ArrayField.AsInt32() });
+		}
+	}
+
+	if (FCbArrayView ArrayView = Obj["Out"].AsArrayView())
+	{
+		Node.Connections.Out.Reserve(int32(ArrayView.Num()));
+		for (FCbFieldView ArrayField : ArrayView)
+		{
+			Node.Connections.Out.Add(FPackageGraph::FEdge { .Index = ArrayField.AsInt32() });
+		}
+	}
+
+	return true;
+}
+
+////////////////////////////////////////////////////////////////////////////////
+static FCbWriter& operator<<(FCbWriter& Writer, const FPackageGraph::FNodeData& Data)
+{
+	Writer.BeginObject();
+	Writer << UTF8TEXT("Filename") << Data.Filename;
+	Writer << UTF8TEXT("PackageId") << Data.PackageId.Value();
+	Writer << UTF8TEXT("Size") << Data.Size;
+	Writer << UTF8TEXT("Container") << Data.ContainerNode;
+
+	Writer.BeginArray(UTF8TEXT("In"));
+	for (FPackageGraph::FEdge Edge : Data.Connections.In)
+	{
+		Writer << Edge;
+	}
+	Writer.EndArray();
+
+	Writer.BeginArray(UTF8TEXT("Out"));
+	for (FPackageGraph::FEdge Edge : Data.Connections.Out)
+	{
+		Writer << Edge;
+	}
+	Writer.EndArray();
+
+	Writer.EndObject();
+	return Writer;
+}
+
+////////////////////////////////////////////////////////////////////////////////
+static bool LoadFromCompactBinary(FCbFieldView Field, FPackageGraph::FContainerData& Container)
+{
+	FCbObjectView Obj = Field.AsObjectView();
+
+	if (!Obj)
+	{
+		return false;
+	}
+
+	Container.Name = FString(Obj["Name"].AsString());
+	if (!LoadFromCompactBinary(Obj["Id"], Container.Id))
+	{
+		return false;
+	}
+
+	if (FCbArrayView ArrayView = Obj["In"].AsArrayView())
+	{
+		Container.Connections.In.Reserve(int32(ArrayView.Num()));
+		for (FCbFieldView ArrayField : ArrayView)
+		{
+			Container.Connections.In.Add(FPackageGraph::FEdge { .Index = ArrayField.AsInt32() });
+		}
+	}
+
+	if (FCbArrayView ArrayView = Obj["Out"].AsArrayView())
+	{
+		Container.Connections.Out.Reserve(int32(ArrayView.Num()));
+		for (FCbFieldView ArrayField : ArrayView)
+		{
+			Container.Connections.Out.Add(FPackageGraph::FEdge { .Index = ArrayField.AsInt32() });
+		}
+	}
+
+	return true;
+}
+
+////////////////////////////////////////////////////////////////////////////////
+static FCbWriter& operator<<(FCbWriter& Writer, const FPackageGraph::FContainerData& Container)
+{
+	Writer.BeginObject();
+	Writer << UTF8TEXT("Name") << Container.Name;
+	Writer << UTF8TEXT("Id") << Container.Id;
+
+	Writer.BeginArray(UTF8TEXT("In"));
+	for (FPackageGraph::FEdge Edge : Container.Connections.In)
+	{
+		Writer << Edge;
+	}
+	Writer.EndArray();
+
+	Writer.BeginArray(UTF8TEXT("Out"));
+	for (FPackageGraph::FEdge Edge : Container.Connections.Out)
+	{
+		Writer << Edge;
+	}
+	Writer.EndArray();
+
+	Writer.EndObject();
+	return Writer;
+}
+
+////////////////////////////////////////////////////////////////////////////////
+static bool LoadFromCompactBinary(FCbFieldView Field, FPackageGraph::FMissingEdgeData& MissingEdge)
+{
+	if (FCbObjectView Obj = Field.AsObjectView())
+	{
+		MissingEdge.From	= FPackageGraph::FNode { .Index = Obj["From"].AsInt32() };
+		MissingEdge.To		= FPackageId::FromValue(Obj["To"].AsUInt64());
+		MissingEdge.bSoft	= Obj["Soft"].AsBool();
+
+		return true;
+	}
+
+	return false; 
+}
+
+////////////////////////////////////////////////////////////////////////////////
+static FCbWriter& operator<<(FCbWriter& Writer, const FPackageGraph::FMissingEdgeData& MissingEdge)
+{
+	Writer.BeginObject();
+	Writer << UTF8TEXT("From") << MissingEdge.From;
+	Writer << UTF8TEXT("To") << MissingEdge.To.Value();
+	Writer << UTF8TEXT("Soft") << MissingEdge.bSoft;
+	Writer.EndObject();
+
+	return Writer;
+}
+
+////////////////////////////////////////////////////////////////////////////////
+static bool LoadFromCompactBinary(FCbFieldView Field, FPackageGraph::FRedirect& Redirect)
+{
+	if (FCbObjectView Obj = Field.AsObjectView())
+	{
+		Redirect.From	= FPackageId::FromValue(Obj["From"].AsUInt64());
+		Redirect.To		= FPackageId::FromValue(Obj["To"].AsUInt64());
+
+		return true;
+	}
+
+	return false;
+}
+
+////////////////////////////////////////////////////////////////////////////////
+static FCbWriter& operator<<(FCbWriter& Writer, const FPackageGraph::FRedirect& Redirect)
+{
+	Writer.BeginObject();
+	Writer << UTF8TEXT("From") << Redirect.From.Value();
+	Writer << UTF8TEXT("To") << Redirect.To.Value();
+	Writer.EndObject();
+
+	return Writer;
+}
+
+////////////////////////////////////////////////////////////////////////////////
+FCbWriter& operator<<(FCbWriter& Writer, const FPackageGraph& Graph)
+{
+	Writer.BeginObject();
+
+	if (Graph.ContainerNodeData.IsEmpty() == false)
+	{
+		Writer.BeginArray(UTF8TEXT("Containers"));
+		for (const FPackageGraph::FContainerData& Data : Graph.ContainerNodeData)
+		{
+			Writer << Data;
+		}
+		Writer.EndArray();
+	}
+
+	if (Graph.NodeData.IsEmpty() == false)
+	{
+		Writer.BeginArray(UTF8TEXT("Nodes"));
+		for (const FPackageGraph::FNodeData& Data : Graph.NodeData)
+		{
+			Writer << Data;
+		}
+		Writer.EndArray();
+	}
+
+	if (Graph.EdgeData.IsEmpty() == false)
+	{
+		Writer.BeginArray(UTF8TEXT("Edges"));
+		for (const FPackageGraph::FEdgeData& Edge : Graph.EdgeData)
+		{
+			Writer << Edge;
+		}
+		Writer.EndArray();
+	}
+
+	if (Graph.MissingEdgeData.IsEmpty() == false)
+	{
+		Writer.BeginArray(UTF8TEXT("MissingEdges"));
+		for (const FPackageGraph::FMissingEdgeData& Missing : Graph.MissingEdgeData)
+		{
+			Writer << Missing;
+		}
+		Writer.EndArray();
+	}
+
+	if (Graph.Redirects.IsEmpty() == false)
+	{
+		Writer.BeginArray(UTF8TEXT("Redirects"));
+		for (const FPackageGraph::FRedirect& Redirect : Graph.Redirects)
+		{
+			Writer << Redirect;
+		}
+		Writer.EndArray();
+	}
+
+	Writer.EndObject();
+
+	return Writer;
+}
+
+////////////////////////////////////////////////////////////////////////////////
+bool LoadFromCompactBinary(FCbFieldView Field, FPackageGraph& Graph)
+{
+	FCbObjectView Obj = Field.AsObjectView();
+
+	if (!Obj)
+	{
+		return false;
+	}
+
+	if (FCbArrayView ArrayView = Obj["Containers"].AsArrayView())
+	{
+		Graph.ContainerNodeData.Reserve(int32(ArrayView.Num()));
+		for (FCbFieldView ArrayField : ArrayView)
+		{
+			if (!LoadFromCompactBinary(ArrayField, Graph.ContainerNodeData.AddDefaulted_GetRef()))
+			{
+				return false;
+			}
+		}
+	}
+
+	if (FCbArrayView ArrayView = Obj["Nodes"].AsArrayView())
+	{
+		Graph.NodeData.Reserve(int32(ArrayView.Num()));
+		for (FCbFieldView ArrayField : ArrayView)
+		{
+			const int32					Index = Graph.NodeData.Num();
+			FPackageGraph::FNode		Node { .Index = Index  };
+			FPackageGraph::FNodeData&	Data = Graph.NodeData.AddDefaulted_GetRef();
+
+			if (LoadFromCompactBinary(ArrayField, Data))
+			{
+				Graph.PackageIdToNode.Add(Data.PackageId, Node);
+				Graph.ContainerNodeData[Data.ContainerNode.Index].PackageNodes.Add(Node);
+			}
+			else
+			{
+				return false;
+			}
+		}
+	}
+
+	if (FCbArrayView ArrayView = Obj["Edges"].AsArrayView())
+	{
+		Graph.EdgeData.Reserve(int32(ArrayView.Num()));
+		for (FCbFieldView ArrayField : ArrayView)
+		{
+			if (!LoadFromCompactBinary(ArrayField, Graph.EdgeData.AddDefaulted_GetRef()))
+			{
+				return false;
+			}
+		}
+	}
+
+	if (FCbArrayView ArrayView = Obj["MissingEdges"].AsArrayView())
+	{
+		Graph.MissingEdgeData.Reserve(int32(ArrayView.Num()));
+		for (FCbFieldView ArrayField : ArrayView)
+		{
+			if (!LoadFromCompactBinary(ArrayField, Graph.MissingEdgeData.AddDefaulted_GetRef()))
+			{
+				return false;
+			}
+		}
+	}
+
+	if (FCbArrayView ArrayView = Obj["Redirects"].AsArrayView())
+	{
+		Graph.Redirects.Reserve(int32(ArrayView.Num()));
+		for (FCbFieldView ArrayField : ArrayView)
+		{
+			FPackageGraph::FRedirect Redirect;
+			if (LoadFromCompactBinary(ArrayField, Redirect))
+			{
+				Graph.Redirects.Add(Redirect);
+			}
+			else
+			{
+				return false;
+			}
+		}
+	}
+
+	return true;
+}
+
+////////////////////////////////////////////////////////////////////////////////
+FIoStatus LoadFromCompactBinary(const FString& Filename, FPackageGraph& OutGraph)
+{
+	IFileManager& Ifm = IFileManager::Get();
+
+	TUniquePtr<FArchive> Ar(Ifm.CreateFileReader(*Filename));
+
+	if (Ar.IsValid() == false)
+	{
+		return FIoStatusBuilder(EIoErrorCode::FileNotOpen) << TEXT("Failed to open '") << Filename << TEXT("'");
+	}
+
+	if (FCbField GraphCb = LoadCompactBinary(*Ar))
+	{
+		if (LoadFromCompactBinary(GraphCb, OutGraph))
+		{
+			return FIoStatus::Ok;
+		}
+	}
+
+	return FIoStatusBuilder(EIoErrorCode::ReadError) << TEXT("Failed to read combact binary from '") << Filename << TEXT("'");
+}
+
+////////////////////////////////////////////////////////////////////////////////
+static FKeyChain LoadKeys(const FContext& Context)
+{
+	FKeyChain Ret;
+
+	const FString Path = FString(Context.Get<FStringView>(TEXT("-CryptoKeys")));
+	if (Path.IsEmpty())
+	{
+		return Ret;
+	}
+
+	KeyChainUtilities::LoadKeyChainFromFile(*Path, Ret);
+	return Ret;
+}
+
+////////////////////////////////////////////////////////////////////////////////
+static TArray<FString> GlobContainers(const FContext& Context)
+{
+	IFileManager& Ifm = IFileManager::Get();
+
+	const FString GlobPattern = FString(Context.Get<FStringView>(TEXT("ContainerGlob")));
+	TArray<FString> Ret;
+
+	if (Ifm.FileExists(*GlobPattern))
+	{
+		Ret.Add(GlobPattern);
+	}
+	else if (Ifm.DirectoryExists(*GlobPattern))
+	{
+		FString Directory = GlobPattern;
+		FPaths::NormalizeDirectoryName(Directory);
+
+		TArray<FString> FoundContainerFiles;
+		Ifm.FindFiles(FoundContainerFiles, *(Directory / TEXT("*.utoc")), true, false);
+
+		for (const FString& Filename : FoundContainerFiles)
+		{
+			Ret.Emplace(Directory / Filename);
+		}
+	}
+	else
+	{
+		FString Directory = FPaths::GetPath(GlobPattern);
+		FPaths::NormalizeDirectoryName(Directory);
+
+		TArray<FString> FoundContainerFiles;
+		Ifm.FindFiles(FoundContainerFiles, *GlobPattern, true, false);
+
+		for (const FString& Filename : FoundContainerFiles)
+		{
+			Ret.Emplace(Directory / Filename);
+		}
+	}
+
+	return Ret;
+}
+
+////////////////////////////////////////////////////////////////////////////////
+static int32 BuildGraphCommandEntry(const FContext& Context)
+{
+	GPrintLogCategory = GPrintLogVerbosity = false;
+	GPrintLogTimes = ELogTimes::None;
+
+	TArray<FString> Containers	= GlobContainers(Context);
+	FKeyChain		KeyChain	= LoadKeys(Context);	
+	FString			OutFilename	= FString(Context.Get<FStringView>(TEXT("-Out")));
+
+	if (Containers.IsEmpty())
+	{
+		UE_LOG(LogIoStoreOnDemand, Error, TEXT("Failed to find any container(s)."));
+		return -1;
+	}
+
+	if (OutFilename.IsEmpty())
+	{
+		UE_LOG(LogIoStoreOnDemand, Error, TEXT("No filename specified"));
+		return -1;
+	}
+
+	const bool bSkipSoftRefs = Context.Get<bool>(TEXT("-SkipSoftRefs"));
+	FPackageGraph Graph;
+	FPackageGraph::Build(Containers, KeyChain, !bSkipSoftRefs, Graph);
+	FIoStatus Status = FPackageGraph::Save(Graph, OutFilename);
+
+	if (Status.IsOk() == false)
+	{
+		UE_LOG(LogIoStoreOnDemand, Error, TEXT("%s"), *Status.ToString());
+		return -1;
+	}
+
+	return 0; 
+}
+
+////////////////////////////////////////////////////////////////////////////////
+static int32 QueryGraphCommandEntry(const FContext& Context)
+{
+	GPrintLogCategory = GPrintLogVerbosity = false;
+	GPrintLogTimes = ELogTimes::None;
+
+	FString			GraphFilename = FString(Context.Get<FStringView>(TEXT("Graph")));
+	FPackageGraph	Graph;
+
+	bool		bTraverse	= false;
+	bool		bPrintDeps	= false;
+	FPackageId	PackageId	= FPackageId();
+
+	if (FStringView HexView = Context.Get<FStringView>(TEXT("-Traverse")); HexView.IsEmpty() == false)
+	{
+		uint64 Value	= FParse::HexNumber64(FString(HexView));
+		PackageId		= FPackageId::FromValue(Value);
+		bTraverse		= true;
+	}
+
+	if (FStringView HexView = Context.Get<FStringView>(TEXT("-Deps")); HexView.IsEmpty() == false)
+	{
+		uint64 Value	= FParse::HexNumber64(FString(HexView));
+		PackageId		= FPackageId::FromValue(Value);
+		bPrintDeps		= true;
+	}
+
+	if (PackageId.IsValid() == false)
+	{
+		UE_LOG(LogIoStoreOnDemand, Error, TEXT("Invalid package ID 0x%s"), *LexToString(PackageId));
+		return -1;
+	}
+
+	UE_LOG(LogIoStoreOnDemand, Display, TEXT("Loading graph '%s'"), *GraphFilename);
+	if (FIoStatus Status = LoadFromCompactBinary(GraphFilename, Graph); !Status.IsOk()) 
+	{
+		UE_LOG(LogIoStoreOnDemand, Error, TEXT("%s"), *Status.ToString());
+		return -1;
+	}
+
+	const FPackageGraph::FNode Node = Graph.GetNode(PackageId);
+	if (Node.IsValid() == false)
+	{
+		UE_LOG(LogIoStoreOnDemand, Error, TEXT("Failed to find graph node for 0x%s"), *LexToString(PackageId));
+		return -1;
+	}
+
+	if (bTraverse)
+	{
+		struct FContainerStats
+		{
+			uint64 PackageCount = 0;
+			uint64 TotalSize	= 0;
+		};
+
+		struct FVisitor final
+			: public FPackageGraph::FVisitor
+		{
+			virtual bool Visit(const FPackageGraph& Graph, FPackageGraph::FNode Node, const FPackageGraph::FNodeData& Data) override
+			{
+				bool bAlreadyInSet = false;
+				Visited.Add(Node, &bAlreadyInSet);
+				check(!bAlreadyInSet);
+				TraverseOrder.Add(Node);
+
+				FContainerStats& Stats = ContainerStats.FindOrAdd(Data.ContainerNode);
+				Stats.TotalSize += Data.Size;
+				Stats.PackageCount++;
+				TotalSize += Data.Size;
+
+				return true;
+			}
+
+			TSet<FPackageGraph::FNode>								Visited;
+			TArray<FPackageGraph::FNode>							TraverseOrder;
+			TMap<FPackageGraph::FContainerNode, FContainerStats>	ContainerStats;
+			uint64 TotalSize = 0;
+		}; 
+
+		const bool bSkipSoftRefs	= Context.Get<bool>(TEXT("-SkipSoftRefs"));
+		const bool bIncoming		= Context.Get<bool>(TEXT("-Incoming"));
+
+		UE_LOG(LogIoStoreOnDemand, Display, TEXT("Traversing %s graph node(s) %s for package ID 0x%s"),
+			bIncoming ? TEXT("incoming") : TEXT("outgoing"),
+			bSkipSoftRefs ? TEXT("skipping soft references") : TEXT("including soft references"),
+			*LexToString(PackageId));
+		UE_LOG(LogIoStoreOnDemand, Display, TEXT(""));
+		UE_LOG(LogIoStoreOnDemand, Display, TEXT("%4s %-20s %-40s %-20s %s"),
+			TEXT("#"), TEXT("Container ID"), TEXT("Container"), TEXT("Package ID"), TEXT("Filename"));
+		UE_LOG(LogIoStoreOnDemand, Display, TEXT("----------------------------------------------------------------------------------------------------"));
+
+		FVisitor Visitor;
+		Graph.Traverse(Node, bIncoming, bSkipSoftRefs, Visitor);
+
+		for (int32 Idx = 1; FPackageGraph::FNode Traversed : Visitor.TraverseOrder)
+		{
+			const FPackageGraph::FNodeData&			Data = Graph.GetNodeData(Traversed);
+			const FPackageGraph::FContainerData&	ContainerData = Graph.GetContainerData(Data.ContainerNode);
+
+			UE_LOG(LogIoStoreOnDemand, Display, TEXT("%3d: 0x%-18s %-40s %s%-18s '%s' (%.2fKiB)"),
+				Idx++, *LexToString(ContainerData.Id), *ContainerData.Name, TEXT("0x"), *LexToString(Data.PackageId), *Data.Filename, float(Data.Size) / 1024.0f);
+		}
+
+		UE_LOG(LogIoStoreOnDemand, Display, TEXT("----------------------------------------------------------------------------------------------------"));
+		UE_LOG(LogIoStoreOnDemand, Display, TEXT("Traversed %d package(s) of total %.2f KiB in %d container(s) for package ID 0x%s"),
+				Visitor.TraverseOrder.Num(), float(Visitor.TotalSize) / 1024.0f, Visitor.ContainerStats.Num(), *LexToString(PackageId));
+
+		for (const TPair<FPackageGraph::FContainerNode, FContainerStats>& Kv : Visitor.ContainerStats)
+		{
+			const FPackageGraph::FContainerData& ContainerData = Graph.GetContainerData(Kv.Key);
+			UE_LOG(LogIoStoreOnDemand, Display, TEXT("\t* %-40s : %.2f KiB (%llu)"), *ContainerData.Name, float(Kv.Value.TotalSize) / 1024.0f, Kv.Value.PackageCount);
+		}
+	}
+	else if (bPrintDeps)
+	{
+		const bool								bIncoming		= Context.Get<bool>(TEXT("-Incoming"));
+		const FPackageGraph::FNodeData&			NodeData		= Graph.GetNodeData(Node);
+		const FPackageGraph::FContainerData&	ContainerData	= Graph.GetContainerData(NodeData.ContainerNode);
+
+		UE_LOG(LogIoStoreOnDemand, Display, TEXT(""));
+		UE_LOG(LogIoStoreOnDemand, Display, TEXT("Dependencies: 0x%s"), *LexToString(NodeData.PackageId));
+		UE_LOG(LogIoStoreOnDemand, Display, TEXT("--------------------------------------------------------------------------------"));
+		UE_LOG(LogIoStoreOnDemand, Display, TEXT("%-16s: %s"), TEXT("Container"), *ContainerData.Name);
+		UE_LOG(LogIoStoreOnDemand, Display, TEXT("%-16s: '%s'"), TEXT("Filename"), *LexToString(NodeData.Filename));
+		UE_LOG(LogIoStoreOnDemand, Display, TEXT("%-16s: %.2fKiB"), TEXT("Size"), float(NodeData.Size) / 1024.0f);
+		UE_LOG(LogIoStoreOnDemand, Display, TEXT(""));
+
+		UE_LOG(LogIoStoreOnDemand, Display, TEXT("Outgoing:"));
+		UE_LOG(LogIoStoreOnDemand, Display, TEXT("---------"));
+		for (int32 Idx = 1; FPackageGraph::FEdge Edge : NodeData.Connections.Out)
+		{
+			const FPackageGraph::FEdgeData&		EdgeData		= Graph.GetEdgeData(Edge);
+			const FPackageGraph::FNodeData&		ToData			= Graph.GetNodeData(EdgeData.To);
+			const FPackageGraph::FContainerData ToContainerData = Graph.GetContainerData(ToData.ContainerNode);
+			UE_LOG(LogIoStoreOnDemand, Display, TEXT("%d: %-40s %s%-18s '%s' (%.2f KiB)"),
+				Idx++, *ToContainerData.Name, TEXT("0x"), *LexToString(ToData.PackageId), *ToData.Filename, float(ToData.Size) / 1024.0f);
+			ensure(EdgeData.From == Node);
+		}
+
+		if (bIncoming)
+		{
+			UE_LOG(LogIoStoreOnDemand, Display, TEXT(""));
+			UE_LOG(LogIoStoreOnDemand, Display, TEXT("Incoming:"));
+			UE_LOG(LogIoStoreOnDemand, Display, TEXT("---------"));
+			for (int32 Idx = 1; FPackageGraph::FEdge Edge : NodeData.Connections.In)
+			{
+				const FPackageGraph::FEdgeData&		EdgeData			= Graph.GetEdgeData(Edge);
+				const FPackageGraph::FNodeData&		FromData			= Graph.GetNodeData(EdgeData.From);
+				const FPackageGraph::FContainerData FromContainerData	= Graph.GetContainerData(FromData.ContainerNode);
+				UE_LOG(LogIoStoreOnDemand, Display, TEXT("%d: %-40s %s%-18s '%s' (%.2f KiB)"),
+						Idx++, *FromContainerData.Name, TEXT("0x"), *LexToString(FromData.PackageId), *FromData.Filename, float(FromData.Size) / 1024.0f);
+				ensure(EdgeData.To == Node);
+			}
+		}
+	}
+
+	return 0; 
+}
+
+////////////////////////////////////////////////////////////////////////////////
+static FCommand GraphCommand(
+	BuildGraphCommandEntry,
+	TEXT("BuildGraph"),
+	TEXT(""),
+	{
+		TArgument<FStringView>(TEXT("ContainerGlob"),	TEXT("Path globbed to discover input containers")),
+		TArgument<FStringView>(TEXT("-CryptoKeys"),		TEXT("JSON-format keyring for input containers")),
+		TArgument<bool>(TEXT("-SkipSoftRefs"),			TEXT("Whether to skip traverisng soft package reference(s).")),
+		TArgument<FStringView>(TEXT("-Out"),			TEXT("Filename to save the graph (JSON|CompactBinary)."))
+	}
+);
+
+////////////////////////////////////////////////////////////////////////////////
+static FCommand QueryGraphCommand(
+	QueryGraphCommandEntry,
+	TEXT("QueryGraph"),
+	TEXT(""),
+	{
+		TArgument<FStringView>(TEXT("Graph"),		TEXT("Filepath to the package graph stored in compact binary format.")),
+		TArgument<FStringView>(TEXT("-Traverse"),	TEXT("Traverse the graph starting from the specified package ID.")),
+		TArgument<FStringView>(TEXT("-Deps"),		TEXT("Print dependencies for the specified package ID.")),
+		TArgument<bool>(TEXT("-SkipSoftRefs"),		TEXT("Whether to skip traverisng soft package reference(s).")),
+		TArgument<bool>(TEXT("-Incoming"),			TEXT("Whether to traverse incoming package nodes."))
+	}
+);
+
+} // namespace UE::IoStore::Tool
+
+#endif // UE_WITH_IAS_TOOL
diff --git a/Engine/Source/Runtime/Experimental/IoStore/OnDemand/Private/Tool/Graph.h b/Engine/Source/Runtime/Experimental/IoStore/OnDemand/Private/Tool/Graph.h
new file mode 100644
--- /dev/null
+++ b/Engine/Source/Runtime/Experimental/IoStore/OnDemand/Private/Tool/Graph.h
@@ -0,0 +1,126 @@
+// Copyright Epic Games, Inc. All Rights Reserved.
+
+#if defined(UE_WITH_IAS_TOOL)
+
+#include "CoreTypes.h"
+#include "Containers/Map.h"
+#include "Containers/Set.h"
+#include "IO/IoChunkId.h"
+#include "IO/IoContainerId.h"
+#include "IO/IoStatus.h"
+#include "IO/PackageId.h"
+
+class FCbWriter;
+class FCbFieldView;
+struct FKeyChain;
+
+namespace UE::IoStore::Tool
+{
+
+////////////////////////////////////////////////////////////////////////////////
+class FPackageGraph
+{
+public:
+	struct FIndex
+	{
+		bool			IsValid() const					{ return Index != INDEX_NONE; }
+		friend uint32	GetTypeHash(FIndex Node)		{ return GetTypeHash(Node.Index); }
+		explicit		operator bool() const			{ return IsValid(); }
+		bool			operator==(FIndex Other) const	{ return Index == Other.Index; }
+
+		int32 Index = INDEX_NONE;
+	};
+
+	using FNode				= FIndex;
+	using FEdge				= FIndex;
+	using FContainerNode	= FIndex;
+	using FMissingEdge		= FIndex;
+
+	struct FConnections
+	{
+		TSet<FEdge>	In;
+		TSet<FEdge>	Out;
+	};
+
+	struct FNodeData
+	{
+		FConnections		Connections;
+		FString				Filename;
+		FPackageId			PackageId;
+		uint64				Size;
+		FContainerNode		ContainerNode;
+	};
+
+	struct FEdgeData
+	{
+		FNode	From;
+		FNode	To;
+		bool	bSoft = false;
+	};
+
+	struct FMissingEdgeData
+	{
+		FNode		From;
+		FPackageId	To;
+		bool		bSoft = false;
+	};
+
+	struct FContainerData
+	{
+		FConnections	Connections;
+		TSet<FNode>		PackageNodes; // Transient
+		FString			Name;
+		FIoContainerId	Id;
+	};
+
+	struct FRedirect
+	{
+		friend uint32	GetTypeHash(const FRedirect &Redirect)		{ return GetTypeHash(Redirect.From); }
+		bool			operator==(const FRedirect &Other) const	{ return From == Other.From; }
+		bool			operator<(const FRedirect &Other) const		{ return From < Other.From; }
+
+		FPackageId From;
+		FPackageId To;
+	};
+
+	struct FVisitor
+	{
+		virtual bool Visit(const FPackageGraph& Graph, FNode Node, const FNodeData& Data) = 0;
+	};
+
+	void						Reserve(int32 ContainerCount, int32 PackageCount);
+	FContainerNode				AddContainerNode(FString&& Name, const FIoContainerId& Id);
+	const FContainerData&		GetContainerData(FContainerNode Node) const { return ContainerNodeData[Node.Index]; }
+	FNode						AddNode(FContainerNode ContainerNode, const FPackageId& PackageId, const FString* Filename, uint64 ChunkSize);
+	FNode						GetNode(const FPackageId& PackageId) const { return PackageIdToNode.FindRef(PackageId); }
+	const FNodeData&			GetNodeData(FNode Node) const { return NodeData[Node.Index]; }
+	int32						GetNodeCount() const { return NodeData.Num(); }
+	void						AddEdge(FNode From, FNode To, bool bSoft);
+	const FEdgeData&			GetEdgeData(FEdge Edge) const { return EdgeData[Edge.Index]; }
+	int32						GetEdgeCount() const { return EdgeData.Num(); }
+	void						AddMissingEdge(FNode From, const FPackageId To, bool bSoft);
+	void						AddRedirect(const FPackageId& From, const FPackageId& To);
+
+	void						Traverse(FNode Node, bool bIncoming, bool bSkipSoftRefs, FVisitor& Visitor);
+
+	friend FCbWriter&			operator<<(FCbWriter& Writer, const FPackageGraph& Graph);
+	friend bool					LoadFromCompactBinary(FCbFieldView Field, FPackageGraph& Graph);
+
+	static void					Build(TConstArrayView<FString> ContainerFilenames, const FKeyChain& KeyChain, bool bSkipSoftRefs, FPackageGraph& Graph);
+	static FIoStatus			Save(const FPackageGraph& Graph, const FString& Filename);
+
+private:
+	TArray<FNodeData>			NodeData;
+	TArray<FEdgeData>			EdgeData;
+	TArray<FMissingEdgeData>	MissingEdgeData;
+	TArray<FContainerData>		ContainerNodeData;
+	TSet<FRedirect>				Redirects;
+	TMap<FPackageId, FNode>		PackageIdToNode; // Transient
+};
+
+////////////////////////////////////////////////////////////////////////////////
+FIoStatus	LoadFromCompactBinary(const FString& Filename, FPackageGraph& OutGraph);
+
+} // namespace UE::IoStore::Tool
+
+#endif // UE_WITH_IAS_TOOL
