IasTool graph command build fix.

#rnx

diff --git a/Engine/Source/Runtime/Experimental/IoStore/OnDemand/Private/Tool/Graph.cpp b/Engine/Source/Runtime/Experimental/IoStore/OnDemand/Private/Tool/Graph.cpp
--- a/Engine/Source/Runtime/Experimental/IoStore/OnDemand/Private/Tool/Graph.cpp
+++ b/Engine/Source/Runtime/Experimental/IoStore/OnDemand/Private/Tool/Graph.cpp
@@ -847,7 +847,7 @@ FIoStatus LoadFromCompactBinary(const FString& Filename, FPackageGraph& OutGraph
 }
 
 ////////////////////////////////////////////////////////////////////////////////
-static FKeyChain LoadKeys(const FContext& Context)
+static FKeyChain LoadCryptoKeys(const FContext& Context)
 {
 	FKeyChain Ret;
 
@@ -910,7 +910,7 @@ static int32 BuildGraphCommandEntry(const FContext& Context)
 	GPrintLogTimes = ELogTimes::None;
 
 	TArray<FString> Containers	= GlobContainers(Context);
-	FKeyChain		KeyChain	= LoadKeys(Context);	
+	FKeyChain		KeyChain	= LoadCryptoKeys(Context);	
 	FString			OutFilename	= FString(Context.Get<FStringView>(TEXT("-Out")));
 
 	if (Containers.IsEmpty())
