Can't use checks on untrusted source.

#rnx

diff --git a/Engine/Source/Runtime/Experimental/IoStore/HttpClient/Private/TransactionTwo.inl b/Engine/Source/Runtime/Experimental/IoStore/HttpClient/Private/TransactionTwo.inl
--- a/Engine/Source/Runtime/Experimental/IoStore/HttpClient/Private/TransactionTwo.inl
+++ b/Engine/Source/Runtime/Experimental/IoStore/HttpClient/Private/TransactionTwo.inl
@@ -417,7 +417,11 @@ FOutcome FDriverNg::TickInner(FTlsPeer& Peer)
 
 	if (FrameView.IsEmpty())
 	{
-		check(RecvRemaining >= int32(Dest.GetSize()));
+		if (RecvRemaining < int32(Dest.GetSize()))
+		{
+			return FOutcome::Error("RecvRemaining value too small");
+		}
+
 		if ((RecvRemaining -= int32(Dest.GetSize())) == 0)
 		{
 			RecvRemaining = 0 - int32(sizeof(RecvFrame));
@@ -479,7 +483,10 @@ FOutcome FDriverNg::RecvResponse(FTlsPeer& Peer, TArray<FHeaderNg>* Sink)
 
 		// Assume if we're filling we are expecting something
 		check(RecvRemaining > 0);
-		check(RecvFrame.Type == FFrame::EType::Headers);
+		if (RecvFrame.Type != FFrame::EType::Headers)
+		{
+			return FOutcome::Error("Not a headers frame", int32(RecvFrame.Type));
+		}
 
 		int32 EndOfHeaders = (RecvFrame.Flags & FFrame::Flag_EndHeaders);
 		uint32 ThisId = RecvFrame.GetId();
@@ -507,7 +514,10 @@ FOutcome FDriverNg::RecvBody(FTlsPeer& Peer, FMutableMemoryView Dest)
 		}
 
 		check(RecvRemaining >= 0);
-		check(RecvFrame.Type == FFrame::EType::Data);
+		if (RecvFrame.Type != FFrame::EType::Data)
+		{
+			return FOutcome::Error("Not a data frame", int32(RecvFrame.Type));
+		}
 
 		int32 EndOfStream = (RecvFrame.Flags & FFrame::Flag_EndStream);
 
