Move IasTool graph command common functions to namespace.

#rnx

#FailedValidations Preflight

diff --git a/Engine/Source/Runtime/Experimental/IoStore/OnDemand/Private/Tool/Graph.cpp b/Engine/Source/Runtime/Experimental/IoStore/OnDemand/Private/Tool/Graph.cpp
--- a/Engine/Source/Runtime/Experimental/IoStore/OnDemand/Private/Tool/Graph.cpp
+++ b/Engine/Source/Runtime/Experimental/IoStore/OnDemand/Private/Tool/Graph.cpp
@@ -846,6 +846,10 @@ FIoStatus LoadFromCompactBinary(const FString& Filename, FPackageGraph& OutGraph
 	return FIoStatusBuilder(EIoErrorCode::ReadError) << TEXT("Failed to read combact binary from '") << Filename << TEXT("'");
 }
 
+////////////////////////////////////////////////////////////////////////////////
+namespace GraphTool
+{
+
 ////////////////////////////////////////////////////////////////////////////////
 static FKeyChain LoadCryptoKeys(const FContext& Context)
 {
@@ -903,14 +907,16 @@ static TArray<FString> GlobContainers(const FContext& Context)
 	return Ret;
 }
 
+} // UE::IoStore::Tool::GraphTool
+
 ////////////////////////////////////////////////////////////////////////////////
 static int32 BuildGraphCommandEntry(const FContext& Context)
 {
 	GPrintLogCategory = GPrintLogVerbosity = false;
 	GPrintLogTimes = ELogTimes::None;
 
-	TArray<FString> Containers	= GlobContainers(Context);
-	FKeyChain		KeyChain	= LoadCryptoKeys(Context);	
+	TArray<FString> Containers	= GraphTool::GlobContainers(Context);
+	FKeyChain		KeyChain	= GraphTool::LoadCryptoKeys(Context);	
 	FString			OutFilename	= FString(Context.Get<FStringView>(TEXT("-Out")));
 
 	if (Containers.IsEmpty())
