Extract common command functions to new header.

#rnx

diff --git a/Engine/Source/Runtime/Experimental/IoStore/OnDemand/Private/Tool/ChunkPlugin.cpp b/Engine/Source/Runtime/Experimental/IoStore/OnDemand/Private/Tool/ChunkPlugin.cpp
--- a/Engine/Source/Runtime/Experimental/IoStore/OnDemand/Private/Tool/ChunkPlugin.cpp
+++ b/Engine/Source/Runtime/Experimental/IoStore/OnDemand/Private/Tool/ChunkPlugin.cpp
@@ -3,6 +3,7 @@
 #if defined(UE_WITH_IAS_TOOL)
 
 #include "Command.h"
+#include "Common.h"
 
 #include "Algo/Find.h"
 #include "Containers/StringView.h"
@@ -553,20 +554,6 @@ int32 WriteOnDemandToc(FOnDemandToc& OnDemandToc, const FString& ContainerName,
 	}
 }
 
-static FKeyChain LoadCryptoKeys(const FContext& Context)
-{
-	FKeyChain Ret;
-
-	auto Path = Context.Get<FStringView>(TEXT("-CryptoKeys"));
-	if (Path.IsEmpty())
-	{
-		return Ret;
-	}
-
-	KeyChainUtilities::LoadKeyChainFromFile(Path.GetData(), Ret);
-	return Ret;
-}
-
 ////////////////////////////////////////////////////////////////////////////////
 static int32 ChunkPluginCommandEntry(const FContext& Context)
 {
@@ -711,7 +698,7 @@ static int32 ChunkPluginCommandEntry(const FContext& Context)
 
 	TMap<FGuid, FAES::FAESKey> EncryptionKeys;
 
-	FKeyChain KeyChain = LoadCryptoKeys(Context);
+	FKeyChain KeyChain = Common::LoadCryptoKeys(Context);
 	for (const TPair<FGuid, FNamedAESKey>& KeyPair : KeyChain.GetEncryptionKeys())
 	{
 		EncryptionKeys.Add(KeyPair.Key, KeyPair.Value.Key);
diff --git a/Engine/Source/Runtime/Experimental/IoStore/OnDemand/Private/Tool/Common.cpp b/Engine/Source/Runtime/Experimental/IoStore/OnDemand/Private/Tool/Common.cpp
new file mode 100644
--- /dev/null
+++ b/Engine/Source/Runtime/Experimental/IoStore/OnDemand/Private/Tool/Common.cpp
@@ -0,0 +1,31 @@
+// Copyright Epic Games, Inc. All Rights Reserved.
+
+#if defined(UE_WITH_IAS_TOOL)
+
+#include "Common.h"
+#include "Command.h"
+
+#include "Misc/KeyChainUtilities.h"
+
+namespace UE::IoStore::Tool::Common
+{
+
+////////////////////////////////////////////////////////////////////////////////
+FKeyChain LoadCryptoKeys(const FContext& Context)
+{
+	FKeyChain Ret;
+
+	const FString Path = FString(Context.Get<FStringView>(TEXT("-CryptoKeys")));
+	if (Path.IsEmpty())
+	{
+		return Ret;
+	}
+
+	KeyChainUtilities::LoadKeyChainFromFile(*Path, Ret);
+
+	return Ret;
+}
+
+} // namespace UE::IoStore::Tool::Common
+
+#endif // UE_WITH_IAS_TOOL
diff --git a/Engine/Source/Runtime/Experimental/IoStore/OnDemand/Private/Tool/Common.h b/Engine/Source/Runtime/Experimental/IoStore/OnDemand/Private/Tool/Common.h
new file mode 100644
--- /dev/null
+++ b/Engine/Source/Runtime/Experimental/IoStore/OnDemand/Private/Tool/Common.h
@@ -0,0 +1,18 @@
+// Copyright Epic Games, Inc. All Rights Reserved.
+
+#if defined(UE_WITH_IAS_TOOL)
+
+#include "Containers/StringView.h"
+
+namespace UE::IoStore::Tool { class FContext; }
+struct FKeyChain;
+
+namespace UE::IoStore::Tool::Common
+{
+
+////////////////////////////////////////////////////////////////////////////////
+FKeyChain LoadCryptoKeys(const FContext& Context);
+
+} // namespace UE::IoStore::Tool::Common
+
+#endif // UE_WITH_IAS_TOOL
diff --git a/Engine/Source/Runtime/Experimental/IoStore/OnDemand/Private/Tool/Graph.cpp b/Engine/Source/Runtime/Experimental/IoStore/OnDemand/Private/Tool/Graph.cpp
--- a/Engine/Source/Runtime/Experimental/IoStore/OnDemand/Private/Tool/Graph.cpp
+++ b/Engine/Source/Runtime/Experimental/IoStore/OnDemand/Private/Tool/Graph.cpp
@@ -4,6 +4,7 @@
 
 #include "Graph.h"
 #include "Command.h"
+#include "Common.h"
 
 #include "HAL/FileManager.h"
 #include "HAL/Platform.h"
@@ -850,21 +851,6 @@ FIoStatus LoadFromCompactBinary(const FString& Filename, FPackageGraph& OutGraph
 namespace GraphTool
 {
 
-////////////////////////////////////////////////////////////////////////////////
-static FKeyChain LoadCryptoKeys(const FContext& Context)
-{
-	FKeyChain Ret;
-
-	const FString Path = FString(Context.Get<FStringView>(TEXT("-CryptoKeys")));
-	if (Path.IsEmpty())
-	{
-		return Ret;
-	}
-
-	KeyChainUtilities::LoadKeyChainFromFile(*Path, Ret);
-	return Ret;
-}
-
 ////////////////////////////////////////////////////////////////////////////////
 static TArray<FString> GlobContainers(const FContext& Context)
 {
@@ -916,7 +902,7 @@ static int32 BuildGraphCommandEntry(const FContext& Context)
 	GPrintLogTimes = ELogTimes::None;
 
 	TArray<FString> Containers	= GraphTool::GlobContainers(Context);
-	FKeyChain		KeyChain	= GraphTool::LoadCryptoKeys(Context);	
+	FKeyChain		KeyChain	= Common::LoadCryptoKeys(Context);
 	FString			OutFilename	= FString(Context.Get<FStringView>(TEXT("-Out")));
 
 	if (Containers.IsEmpty())
diff --git a/Engine/Source/Runtime/Experimental/IoStore/OnDemand/Private/Tool/Upload.cpp b/Engine/Source/Runtime/Experimental/IoStore/OnDemand/Private/Tool/Upload.cpp
--- a/Engine/Source/Runtime/Experimental/IoStore/OnDemand/Private/Tool/Upload.cpp
+++ b/Engine/Source/Runtime/Experimental/IoStore/OnDemand/Private/Tool/Upload.cpp
@@ -3,6 +3,7 @@
 #if defined(UE_WITH_IAS_TOOL)
 
 #include "Command.h"
+#include "Common.h"
 
 #include "Async/Async.h"
 #include "Async/Mutex.h"
@@ -177,22 +178,6 @@ static FUploadParams BuildUploadParams(const FContext& Context)
 	return Ret;
 }
 
-////////////////////////////////////////////////////////////////////////////////
-static FKeyChain LoadKeys(const FContext& Context)
-{
-	FKeyChain Ret;
-
-	auto Path = Context.Get<FStringView>(TEXT("-CryptoKeys"));
-	if (Path.IsEmpty())
-	{
-		return Ret;
-	}
-
-	// UE_LOG(LogIoStore, Display, TEXT("Parsing crypto keys from a crypto key cache file '%s'"), *CryptoKeysCacheFilename);
-	KeyChainUtilities::LoadKeyChainFromFile(Path.GetData(), Ret);
-	return Ret;
-}
-
 ////////////////////////////////////////////////////////////////////////////////
 static bool SaveArrayToFile(TArrayView64<uint8> Data, const TCHAR* Filename, uint32 WriteFlags = FILEWRITE_None)
 {
@@ -987,7 +972,7 @@ static TArray<FString> GlobContainers(const FContext& Context)
 static int32 UploadCommandEntry(const FContext& Context)
 {
 	TArray<FString> Containers = GlobContainers(Context);
-	FKeyChain KeyChain = LoadKeys(Context);
+	FKeyChain KeyChain = Common::LoadCryptoKeys(Context);
 	FUploadParams Params = BuildUploadParams(Context);
 
 	TIoStatusOr<FUploadResult> Result = UploadContainerFiles(Params, Containers, KeyChain);
