Collision Filter Refactor Part 3.8: Updating ClusterUnions to use new API and structures.

diff --git a/Engine/Source/Runtime/Experimental/Chaos/Private/PhysicsProxy/ClusterUnionPhysicsProxy.cpp b/Engine/Source/Runtime/Experimental/Chaos/Private/PhysicsProxy/ClusterUnionPhysicsProxy.cpp
--- a/Engine/Source/Runtime/Experimental/Chaos/Private/PhysicsProxy/ClusterUnionPhysicsProxy.cpp
+++ b/Engine/Source/Runtime/Experimental/Chaos/Private/PhysicsProxy/ClusterUnionPhysicsProxy.cpp
@@ -58,16 +58,14 @@ namespace Chaos
 
 			const FShapesArray& ShapeArray = Particle->ShapesArray();
 			BufferData.CollisionData.Empty(ShapeArray.Num());
-			BufferData.QueryData.Empty(ShapeArray.Num());
-			BufferData.SimData.Empty(ShapeArray.Num());
+			BufferData.ShapeFilterData.Empty(ShapeArray.Num());
+			BufferData.ShapeInstanceData.Empty(ShapeArray.Num());
 
 			for (const TUniquePtr<Chaos::FPerShapeData>& ShapeData : ShapeArray)
 			{
 				BufferData.CollisionData.Add(ShapeData->GetCollisionData());
-				FCollisionFilterData QueryData, SimData;
-				Chaos::Filter::FShapeFilterBuilder::GetLegacyShapeFilter(ShapeData->GetCombinedShapeFilterData(), QueryData, SimData);
-				BufferData.QueryData.Add(QueryData);
-				BufferData.SimData.Add(SimData);
+				BufferData.ShapeFilterData.Add(ShapeData->GetShapeFilterData());
+				BufferData.ShapeInstanceData.Add(ShapeData->GetFilterInstanceData());
 			}
 
 			if constexpr (std::is_base_of_v<FClusterUnionPhysicsProxy::FInternalParticle, TParticle>)
@@ -571,17 +569,15 @@ namespace Chaos
 				{
 					ShapeData->SetCollisionData(CurrentPullData.CollisionData[ShapeIndex]);
 				}
-				PRAGMA_DISABLE_DEPRECATION_WARNINGS
-				if (CurrentPullData.QueryData.IsValidIndex(ShapeIndex))
+				if (CurrentPullData.ShapeFilterData.IsValidIndex(ShapeIndex))
 				{
-					ShapeData->SetQueryData(CurrentPullData.QueryData[ShapeIndex]);
+					ShapeData->SetShapeFilterData(CurrentPullData.ShapeFilterData[ShapeIndex]);
 				}
 
-				if (CurrentPullData.SimData.IsValidIndex(ShapeIndex))
+				if (CurrentPullData.ShapeInstanceData.IsValidIndex(ShapeIndex))
 				{
-					ShapeData->SetSimData(CurrentPullData.SimData[ShapeIndex]);
+					ShapeData->SetFilterInstanceData(CurrentPullData.ShapeInstanceData[ShapeIndex]);
 				}
-				PRAGMA_ENABLE_DEPRECATION_WARNINGS
 			}
 		}
 
diff --git a/Engine/Source/Runtime/Experimental/Chaos/Public/Chaos/PullPhysicsDataImp.h b/Engine/Source/Runtime/Experimental/Chaos/Public/Chaos/PullPhysicsDataImp.h
--- a/Engine/Source/Runtime/Experimental/Chaos/Public/Chaos/PullPhysicsDataImp.h
+++ b/Engine/Source/Runtime/Experimental/Chaos/Public/Chaos/PullPhysicsDataImp.h
@@ -127,8 +127,12 @@ struct FDirtyClusterUnionData : public TBasePullData<FClusterUnionPhysicsProxy,
 	TSharedPtr<FImplicitObject, ESPMode::ThreadSafe> SharedGeometry = nullptr;
 	
 	TArray<FCollisionData> CollisionData;
+	UE_DEPRECATED(5.8, "Use ShapeInstanceData and ShapeFilterData instead")
 	TArray<FCollisionFilterData> QueryData;
+	UE_DEPRECATED(5.8, "Use ShapeInstanceData and ShapeFilterData instead")
 	TArray<FCollisionFilterData> SimData;
+	TArray<Chaos::Filter::FShapeFilterData> ShapeFilterData;
+	TArray<Chaos::Filter::FInstanceData> ShapeInstanceData;
 };
 
 struct FJointConstraintOutputData {
