Remove clamp of boolean value

diff --git a/Engine/Source/Runtime/Experimental/Chaos/Public/ChaosSolverConfiguration.h b/Engine/Source/Runtime/Experimental/Chaos/Public/ChaosSolverConfiguration.h
--- a/Engine/Source/Runtime/Experimental/Chaos/Public/ChaosSolverConfiguration.h
+++ b/Engine/Source/Runtime/Experimental/Chaos/Public/ChaosSolverConfiguration.h
@@ -43,7 +43,7 @@ struct FChaosSolverDestructionSettings
 	* This will save runtime memory for unbroken geometry collection at the cost of allocating those particle when the first break happens
 	* It is false by default 
 	*/
-	UPROPERTY(EditAnywhere, Category = "SolverConfiguration|Destruction", meta = (ClampMin = "-1"))
+	UPROPERTY(EditAnywhere, Category = "SolverConfiguration|Destruction")
 	bool bOptimizeForRuntimeMemory = false;
 };
 
