Bias time so that 0 has meaning.

#rnx

diff --git a/Engine/Source/Runtime/Experimental/IoStore/HttpClient/Private/Loop.inl b/Engine/Source/Runtime/Experimental/IoStore/HttpClient/Private/Loop.inl
--- a/Engine/Source/Runtime/Experimental/IoStore/HttpClient/Private/Loop.inl
+++ b/Engine/Source/Runtime/Experimental/IoStore/HttpClient/Private/Loop.inl
@@ -1169,7 +1169,8 @@ static uint32 GetNowMs()
 	uint64 NowBig = ((FPlatformTime::Cycles64() - Base) * 1000) / Freq;
 	uint32 NowMs = uint32(NowBig);
 	check(NowMs == NowBig);
-	return NowMs;
+	return NowMs + 1;
+	// +1 bias ----^ ...so 0 can be considered "not set"
 }
 
 
