Only track time for idle once hitting an idle state.

#rnx

diff --git a/Engine/Source/Runtime/Experimental/IoStore/HttpClient/Private/Loop.inl b/Engine/Source/Runtime/Experimental/IoStore/HttpClient/Private/Loop.inl
--- a/Engine/Source/Runtime/Experimental/IoStore/HttpClient/Private/Loop.inl
+++ b/Engine/Source/Runtime/Experimental/IoStore/HttpClient/Private/Loop.inl
@@ -713,7 +713,15 @@ bool FPeerGroup::MaybeStartNewWork(FTickState& State)
 {
 	if (InflightNum >= InflightMax)		return false;
 	if (IsKeepAlive == 0)				return false;
-	if (!State.Work->HasWork())			return false;
+
+	if (!State.Work->HasWork())
+	{
+		if (InflightNum == 0)
+		{
+			LastUseMs = State.NowMs;
+		}
+		return false;
+	}
 
 	FActivityNode* Activity = State.Work->PopActivity();
 	do
@@ -805,7 +813,6 @@ void FPeerGroup::RecvInternal(FTickState& State)
 	check(Activity != nullptr);
 
 	IsKeepAlive &= uint8(Activity->GetTransaction()->IsKeepAlive() == true);
-	LastUseMs = State.NowMs;
 	bWaiting |= (Outcome.IsWaiting() && Outcome.IsWaitData());
 
 	// New work may have fell foul of peer issues
@@ -915,10 +922,13 @@ void FPeerGroup::TickSend(FTickState& State, FHost& Host, FPoller& Poller)
 	check(Pending != nullptr);
 
 	// Close idle sockets
-	if (Peer.IsValid() && InflightNum == 0 && LastUseMs + GIdleMs < State.NowMs)
+	if (uint32 TimeoutMs = LastUseMs + GIdleMs; LastUseMs && TimeoutMs < State.NowMs)
 	{
-		LastUseMs = State.NowMs;
-		Peer = FHttpPeer();
+		if (InflightNum == 0)
+		{
+			LastUseMs = 0;
+			Peer = FHttpPeer();
+		}
 	}
 
 	// We don't have a connected socket on first use, or if a keep-alive:close
