Fix the IASTool download command.

#rnx

- The code was treating the iochunktoc as a compact binary object but it is being uploaded as a raw memory blob and we need to parse it as such.

diff --git a/Engine/Source/Runtime/Experimental/IoStore/OnDemand/Private/Tool/Download.cpp b/Engine/Source/Runtime/Experimental/IoStore/OnDemand/Private/Tool/Download.cpp
--- a/Engine/Source/Runtime/Experimental/IoStore/OnDemand/Private/Tool/Download.cpp
+++ b/Engine/Source/Runtime/Experimental/IoStore/OnDemand/Private/Tool/Download.cpp
@@ -8,10 +8,11 @@
 #include "IO/IoStore.h"
 #include "IO/IoStoreOnDemand.h"
 #include "Misc/Paths.h"
+#include "S3/S3Client.h"
 #include "Serialization/CompactBinarySerialization.h"
 #include "Serialization/CompactBinaryWriter.h"
 #include "Serialization/LargeMemoryWriter.h"
-#include "S3/S3Client.h"
+#include "Serialization/MemoryReader.h"
 
 namespace UE::IoStore::Tool
 {
@@ -112,7 +113,10 @@ FIoStatus DownloadContainerFiles(const FDownloadParams& DownloadParams, const FS
 	}
 
 	FOnDemandToc OnDemandToc;
-	if (LoadFromCompactBinary(FCbFieldView(TocResponse.GetBody().GetData()), OnDemandToc) == false)
+	FMemoryReaderView Ar(TocResponse.GetBody().GetView());
+	Ar << OnDemandToc;
+
+	if (Ar.IsError())
 	{
 		return FIoStatus(EIoErrorCode::ReadError, TEXT("Failed to load on demand TOC"));
 	}
