Chaos (BubbleResim): Console variable to include/exclude kinematic particles in the growth of the resim bubble. Kinematic particles found to overlap the resim bubble will still fully resimulate even if they are not involved in the growth.

#DisabledValidations Preflight

diff --git a/Engine/Source/Runtime/Experimental/Chaos/Private/RewindData.cpp b/Engine/Source/Runtime/Experimental/Chaos/Private/RewindData.cpp
--- a/Engine/Source/Runtime/Experimental/Chaos/Private/RewindData.cpp
+++ b/Engine/Source/Runtime/Experimental/Chaos/Private/RewindData.cpp
@@ -52,11 +52,17 @@ namespace Chaos
 			TEXT("p.Chaos.BubbleResim.MaxNumTriesForResimBubbleGrowth"),
 			BubbleResimGrowthMaxNumTries,
 			TEXT("When forming bubbles, we grow the bubble by the AABBs of particles it overlaps. This is the maximum number of such growths we allow."));
+
+		bool bIncludeKinematicParticlesInBubbleGrowth = false;
+		static FAutoConsoleVariableRef CVarIncludeKinematicParticlesInBubbleGrowth(
+			TEXT("p.Chaos.BubbleResim.IncludeKinematicParticlesInBubbleGrowth"),
+			bIncludeKinematicParticlesInBubbleGrowth,
+			TEXT("If true, kinematic particles overlapping resim bubbles cause the bubble to grow. Kinematic particles overlapping resim bubbles will always participate to the resim, regardless whether this is enabled or not."));
 	}
 	
 	namespace ResimDefines
 	{
-		FName BubbleResimDebugDrawTag("BubbleResimDebugDraws");
+		const FName BubbleResimDebugDrawTag(TEXT("BubbleResimDebugDraws"));
 	}
 
 FVec3 FGeometryParticleState::ZeroVector = FVec3(0);
@@ -456,6 +462,19 @@ bool FRewindData::CheckQuaternionThreshold(FQuat A, FQuat B, float ThresholdDegr
 	return AngleDegrees >= ThresholdDegrees;
 }
 
+bool FRewindData::OverlapsAnyResimBubble(const Chaos::FAABB3& InAABB) const
+{
+	for (const FAABB3& ResimBubble : ResimBubbles)
+	{
+		if (ResimBubble.Intersects(InAABB))
+		{
+			return true;
+		}
+	}
+
+	return false;
+}
+
 FAABB3 FRewindData::GetParticleSweptAABB(const FGeometryParticleHandle* ParticleHandle, FRewindData::FDirtyParticleInfo* DirtyParticleInfo, const FFrameAndPhase& AtFrameAndPhase, float HorizonDuration, float ExtraSweptDistance, float ExtraThickness) const
 {
 	FAABB3 ResultAABB = FAABB3::ZeroAABB();
@@ -546,9 +565,6 @@ void FRewindData::FindParticlesOverlappingResimBubbles(const FFrameAndPhase& Rew
 	
 	FPBDRigidsEvolution* Evolution = Solver->GetEvolution();
 
-	// Bubbles are swept particle AABBs
-	TArray<Chaos::FAABB3> ResimBubbles;
-
 	// First seed resim bubbles by building swept AABBs around particles needing resim
 	{
 		QUICK_SCOPE_CYCLE_COUNTER(ChaosRewindAndResim_SeedResimBubbles);
@@ -566,7 +582,7 @@ void FRewindData::FindParticlesOverlappingResimBubbles(const FFrameAndPhase& Rew
 				// because they might not be found to overlap the bubble based on the spatial acceleration structure.
 				// The spatial acceleration we have is for the frame we rewind from, not the frame we rewind to
 				OutParticlesOverlappingResimBubbles.Add(ParticleNeedingResim);
-				CVD_TRACE_DEBUG_DRAW_BOX(FBox(ResimAABB.Min(), ResimAABB.Max()), ResimDefines::BubbleResimDebugDrawTag, FColor::Red);
+				CVD_TRACE_DEBUG_DRAW_BOX(FBox(ResimAABB.Min(), ResimAABB.Max()), ResimDefines::BubbleResimDebugDrawTag, FColor(255,255,191)/*Pale Yellow*/);
 			}
 		}
 	}
@@ -581,16 +597,25 @@ void FRewindData::FindParticlesOverlappingResimBubbles(const FFrameAndPhase& Rew
 			int32 NumTries = 0;
 			int32 MaxNumTries = Chaos::ResimConsoleVars::BubbleResimGrowthMaxNumTries;
 			bool FoundNewOverlaps = true;
+			// This lambda causes the bubble to grow when a particle is found to overlap
+			// It will be called only once for every new particle found overlapping the resim bubble as it grows
 			auto LambdaParticleVisitor = [&](const FGeometryParticleHandle* Particle)->bool
 				{
 					FDirtyParticleInfo* DirtyParticleInfo = DirtyParticles.Find(Particle);
 					FAABB3 ParticleResimAABB = GetParticleSweptAABB(Particle, DirtyParticleInfo, RewindFrameAndPhase,
 						TotalResimDuration, Chaos::ResimConsoleVars::BubbleResimExtraSweptDistance, Chaos::ResimConsoleVars::BubbleResimExtraThickness);
+					bool bIsParticleKinematic = (Particle->ObjectState() == EObjectStateType::Kinematic);
+					bool bIncludeInBubbleGrowth = !bIsParticleKinematic || Chaos::ResimConsoleVars::bIncludeKinematicParticlesInBubbleGrowth;
+					if (bIncludeInBubbleGrowth)
+					{
+						ResimAABB.GrowToInclude(ParticleResimAABB);
+						FoundNewOverlaps = true;
+					}
+#if WITH_CHAOS_VISUAL_DEBUGGER
 					float Alpha = float(NumTries) / float(MaxNumTries <= 1 ? 1 : MaxNumTries - 1);
 					FLinearColor OverlapDrawColor = FMath::Lerp(FLinearColor(FColor::Magenta), FLinearColor(FColor::Purple), Alpha);
 					CVD_TRACE_DEBUG_DRAW_BOX(FBox(ParticleResimAABB.Min(), ParticleResimAABB.Max()), ResimDefines::BubbleResimDebugDrawTag, OverlapDrawColor.ToFColor(true));
-					ResimAABB.GrowToInclude(ParticleResimAABB);
-					FoundNewOverlaps = true;
+#endif // WITH_CHAOS_VISUAL_DEBUGGER
 					return true; // Keep going
 				};
 			FNonStaticUniqueOverlapVisitor OverlapVisitor(LambdaParticleVisitor);
@@ -598,10 +623,12 @@ void FRewindData::FindParticlesOverlappingResimBubbles(const FFrameAndPhase& Rew
 			{
 				FoundNewOverlaps = false;
 				SpatialAcceleration->Overlap(ResimAABB, OverlapVisitor);
+#if WITH_CHAOS_VISUAL_DEBUGGER
 				float Alpha = float(NumTries) / float(MaxNumTries <= 1 ? 1 : MaxNumTries - 1);
 				FLinearColor BubbleDrawColor = FMath::Lerp(FLinearColor::Yellow, FLinearColor::Red, Alpha);
-				NumTries++;
 				CVD_TRACE_DEBUG_DRAW_BOX(FBox(ResimAABB.Min(), ResimAABB.Max()), ResimDefines::BubbleResimDebugDrawTag, BubbleDrawColor.ToFColor(true));
+#endif // WITH_CHAOS_VISUAL_DEBUGGER
+				NumTries++;
 			}
 			OutParticlesOverlappingResimBubbles.Append(OverlapVisitor.OverlappingParticles);
 		}
@@ -986,6 +1013,9 @@ void FRewindData::EndResim()
 
 	// Bubble resim only: "Unfreeze" particles that were frozen because they were outside resim bubbles
 	UnfreezeParticlesOutsideResimBubbles();
+
+	// Clear the array of resim bubbles
+	ResimBubbles.Empty();
 }
 
 void FRewindData::StepNonResimParticles(const int32 Frame)
diff --git a/Engine/Source/Runtime/Experimental/Chaos/Public/RewindData.h b/Engine/Source/Runtime/Experimental/Chaos/Public/RewindData.h
--- a/Engine/Source/Runtime/Experimental/Chaos/Public/RewindData.h
+++ b/Engine/Source/Runtime/Experimental/Chaos/Public/RewindData.h
@@ -32,6 +32,11 @@ namespace Chaos
 		class FPBDIsland;
 	}
 
+	namespace ResimDefines
+	{
+		extern const CHAOS_API FName BubbleResimDebugDrawTag;
+	}
+
 /** Base rewind history used in the rewind data */
 struct FBaseRewindHistory
 {
@@ -2163,7 +2168,14 @@ private:
 	TArray<const Private::FPBDIsland*> IslandValidationIslands;
 	TArray<const FGeometryParticleHandle*> IslandValidationIslandParticles;
 
-	// Bubble resim
+	// BUBBLE RESIM
+	// Bubble resim needs to be enabled, see Chaos::ResimConsoleVars::bIsBubbleResimEnabled
+	// Bubble resim mitigates the cost of resimulations by restricting the amount of particles being resimulated
+	// to those for which resimulation could lead to a change in state compared to the original predictive simulation
+public:
+	// This can be used to determine if an update that affects physics and restricts its effects to within an AABB can be skipped during a resim
+	bool CHAOS_API OverlapsAnyResimBubble(const Chaos::FAABB3& InAABB) const;
+private:
 	// Freeze particles outside resim bubbles before rewind/resim
 	void FreezeParticlesOutsideResimBubbles();
 	// Unfreeze particles that were frozen, this happens at the end of a resim (after all the resim frames have occurred)
@@ -2182,6 +2194,10 @@ private:
 	// Used to keep track of objects we "freeze" during resim (by making them kinematic and marking them FrozenDuringResim)
 	// and their object state prior to the freezing so it can be set back to its original value
 	TArray<TPair<FParticleSafeRef, EObjectStateType>> FrozenParticles;
+	// Resim bubbles start as the swept AABBs of particles needing resim, and grow to encompass nearby kinematic and dynamic particles involved in the resim
+	// Their formation happens in FindParticlesOverlappingResimBubbles()
+	TArray<Chaos::FAABB3> ResimBubbles;
+	// END OF BUBBLE RESIM
 
 	template <typename TObj>
 	bool IsResimAndInSync(const TObj& Handle) const { return IsResim() && Handle.SyncState() == ESyncState::InSync; }
