Chaos API - move experimental API into Experimental folders. WIP.

#rnx
#FailedValidations Preflight

diff --git a/Engine/Plugins/Experimental/ChaosRigidPhysicsAsync/ChaosRigidPhysicsAsync.uplugin b/Engine/Plugins/Experimental/ChaosRigidPhysicsAsync/ChaosRigidPhysicsAsync.uplugin
new file mode 100644
--- /dev/null
+++ b/Engine/Plugins/Experimental/ChaosRigidPhysicsAsync/ChaosRigidPhysicsAsync.uplugin
@@ -0,0 +1,27 @@
+{
+	"FileVersion": 3,
+	"Version": 0.1,
+	"VersionName": "0.1",
+	"FriendlyName": "Chaos Rigid Physics Async",
+	"Description": "Provides the Chaos Rigid Body Physics Engine (Async Implementation)",
+	"Category": "Physics",
+	"CreatedBy" : "Epic Games, Inc.",
+	"CreatedByURL" : "https://epicgames.com",
+	"DocsURL": "",
+	"MarketplaceURL": "",
+	"SupportURL": "",
+	"EnabledByDefault": false,
+	"CanContainContent": false,
+	"IsBetaVersion": false,
+	"IsExperimentalVersion": true,
+	"Installed": false,
+	"Modules": [
+		{
+			"Name": "ChaosRigidPhysicsAsync",
+			"Type": "Runtime",
+			"LoadingPhase": "Default"
+		}
+	],
+	"Plugins": [
+	]
+}
\ No newline at end of file
diff --git a/Engine/Plugins/Experimental/ChaosRigidPhysicsAsync/Source/ChaosRigidPhysicsAsync/ChaosRigidPhysicsAsync.Build.cs b/Engine/Plugins/Experimental/ChaosRigidPhysicsAsync/Source/ChaosRigidPhysicsAsync/ChaosRigidPhysicsAsync.Build.cs
new file mode 100644
--- /dev/null
+++ b/Engine/Plugins/Experimental/ChaosRigidPhysicsAsync/Source/ChaosRigidPhysicsAsync/ChaosRigidPhysicsAsync.Build.cs
@@ -0,0 +1,27 @@
+// Copyright Epic Games, Inc. All Rights Reserved.
+
+using System.Linq;
+using System.Collections.Generic;
+
+namespace UnrealBuildTool.Rules
+{
+    public class ChaosRigidPhysicsAsync : ModuleRules
+    {
+        public ChaosRigidPhysicsAsync(ReadOnlyTargetRules Target) : base(Target)
+        {
+            PublicDependencyModuleNames.AddRange(
+                new string[] {
+                "Core",
+                "CoreUObject",
+                "RigidPhysics",
+				}
+			);
+
+			SetupModulePhysicsSupport(Target);
+			
+			//CppCompileWarningSettings.UnsafeTypeCastWarningLevel = WarningLevel.Error;
+			//StaticAnalyzerDisabledCheckers.Add("cplusplus.NewDeleteLeaks"); // To be reevalulated, believed to be invalid warnings.
+			//StaticAnalyzerDisabledCheckers.Add("core.UndefinedBinaryOperatorResult"); // Invalid warning in mass property calculation.
+		}
+    }
+}
diff --git a/Engine/Plugins/Experimental/ChaosRigidPhysicsAsync/Source/ChaosRigidPhysicsAsync/ChaosRigidPhysicsAsyncModule.cpp b/Engine/Plugins/Experimental/ChaosRigidPhysicsAsync/Source/ChaosRigidPhysicsAsync/ChaosRigidPhysicsAsyncModule.cpp
new file mode 100644
--- /dev/null
+++ b/Engine/Plugins/Experimental/ChaosRigidPhysicsAsync/Source/ChaosRigidPhysicsAsync/ChaosRigidPhysicsAsyncModule.cpp
@@ -0,0 +1,5 @@
+// Copyright Epic Games, Inc. All Rights Reserved.
+
+#include "Modules/ModuleManager.h"
+
+IMPLEMENT_MODULE(FDefaultModuleImpl, ChaosRigidPhysicsAsync)
\ No newline at end of file
diff --git a/Engine/Plugins/Experimental/ChaosRigidPhysicsAsync/Source/ChaosRigidPhysicsAsync/Private/ChaosRigidPhysicsAsync/RigidBodyAsync.cpp b/Engine/Plugins/Experimental/ChaosRigidPhysicsAsync/Source/ChaosRigidPhysicsAsync/Private/ChaosRigidPhysicsAsync/RigidBodyAsync.cpp
new file mode 100644
--- /dev/null
+++ b/Engine/Plugins/Experimental/ChaosRigidPhysicsAsync/Source/ChaosRigidPhysicsAsync/Private/ChaosRigidPhysicsAsync/RigidBodyAsync.cpp
@@ -0,0 +1,319 @@
+// Copyright Epic Games, Inc. All Rights Reserved.
+#include "ChaosRigidPhysicsAsync/RigidBodyAsync.h"
+
+#include "ChaosRigidPhysicsAsync/RigidBodyContainerAsync.h"
+#include "ChaosRigidPhysicsAsync/RigidSceneAsync.h"
+#include "RigidPhysics/RigidAccessor.h"
+#include "RigidPhysics/RigidBody.h"
+
+#if UE_RIGIDPHYSICS_API_ENABLED
+
+namespace Chaos::Rigids::Async
+{
+	FRigidBodyAsyncGT::FRigidBodyAsyncGT(
+		FRigidBodyContainerAsyncGT* InContainer, 
+		FSingleParticlePhysicsProxy* InProxy, 
+		const FRigidBodyIdAsync& InBodyId)
+		: Container(InContainer)
+		, Proxy(InProxy)
+		, BodyId(InBodyId)
+	{
+	}
+
+	FRigidBodyAsyncGT::~FRigidBodyAsyncGT()
+	{
+	}
+
+	const char* FRigidBodyAsyncGT::GetTypeName() const
+	{
+		return "RigidBodyAsyncGT";
+	}
+
+	UE::Physics::FRigidBodyId FRigidBodyAsyncGT::GetBodyId() const
+	{
+		CHAOSRIGIDPHYSICSASYNC_CHECK(IsValid());
+		// NOTE: Cannot use Proxy UniqueIdx - not assigned until RegisterObject
+		// so we must create yet another ID!
+		return GetBodyIdAsync().AsRigidBodyId();
+	}
+
+	UE::Physics::FRigidBodyHandle FRigidBodyAsyncGT::GetBodyHandle() const
+	{
+		check(IsValid());
+		return UE::Physics::FRigidBodyHandle(Container->GetContainerHandle(), GetBodyId());
+	}
+
+	FRigidSceneAsyncGT* FRigidBodyAsyncGT::GetSceneAsync() const
+	{
+		return Container->GetSceneAsync();
+	}
+
+	bool FRigidBodyAsyncGT::IsValid() const
+	{
+		return (Proxy != nullptr) && (Proxy->GetParticle_LowLevel() != nullptr);
+	}
+
+	bool FRigidBodyAsyncGT::IsInScene() const
+	{
+		// TODO_CHAOSAPI: need a better way to know if we're in the scene
+		// This works for now because UniqueIdx is set in FPBDRigidsSolver::RegisterObject 
+		// and GetMarkDeleted() is true after FPBDRigidsSolver::UnregisterObject.
+		return (Proxy != nullptr) && !Proxy->GetMarkedDeleted() && (Proxy->GetParticle_LowLevel() != nullptr) && Proxy->GetParticle_LowLevel()->UniqueIdx().IsValid();
+	}
+
+	bool FRigidBodyAsyncGT::IsStatic() const
+	{
+		return (GetProxyAPI().ObjectState() == EObjectStateType::Static);
+	}
+
+	bool FRigidBodyAsyncGT::IsKinematic() const
+	{
+		return (GetProxyAPI().ObjectState() == EObjectStateType::Kinematic);
+	}
+
+	bool FRigidBodyAsyncGT::IsDynamic() const
+	{
+		return (GetProxyAPI().ObjectState() == EObjectStateType::Dynamic) || (GetProxyAPI().ObjectState() == EObjectStateType::Sleeping);
+	}
+
+	bool FRigidBodyAsyncGT::IsSleeping() const
+	{
+		return (GetProxyAPI().ObjectState() == EObjectStateType::Sleeping);
+	}
+
+	void FRigidBodyAsyncGT::InitTransform(const FTransform& InTransform)
+	{
+		// TODO_CHAOSAPI: This is not supported by the Proxy on the GT. 
+		// For now equivalent to UpdateTransform.
+		UpdateTransform(InTransform);
+	}
+
+	void FRigidBodyAsyncGT::UpdateTransform(const FTransform& InTransform)
+	{
+		CHAOSRIGIDPHYSICSASYNC_CHECK(IsValid());
+
+		FRigidBodyHandle_External& ProxyAPI = GetProxyAPI();
+		if (!IsKinematic() && !IsSleeping() && FVec3::IsNearlyEqual(InTransform.GetLocation(), ProxyAPI.X(), SMALL_NUMBER) && FRotation3::IsNearlyEqual(InTransform.GetRotation(), ProxyAPI.R(), SMALL_NUMBER))
+		{
+			// if simulating, don't update X/R if they haven't changed. this allows scale to be set on simulating body without overriding async position/rotation.
+			return;
+		}
+
+		if (IsKinematic())
+		{
+			// NOTE: SetGlobalPose is a teleport for kinematics. Use SetKinematicTarget_AssumesLocked
+			// if the kinematic should calculate its velocity from the transform delta.
+			ProxyAPI.SetKinematicTarget(InTransform);
+			ProxyAPI.SetV(FVector::Zero());
+			ProxyAPI.SetW(FVector::Zero());
+		}
+
+		ProxyAPI.SetX(InTransform.GetTranslation());
+		ProxyAPI.SetR(InTransform.GetRotation());
+		ProxyAPI.UpdateShapeBounds();
+
+		// Notify container (e.g., update the spatial acceleration or other dependents)
+		if (Container != nullptr)
+		{
+			Container->OnRigidBodyTransformUpdated(this);
+		}
+	}
+
+	void FRigidBodyAsyncGT::SetKinematicTarget(const FTransform& InTransform)
+	{
+		CHAOSRIGIDPHYSICSASYNC_CHECK(IsValid());
+		GetProxyAPI().SetKinematicTarget(FKinematicTarget::MakePositionTarget(InTransform.GetTranslation(), InTransform.GetRotation()));
+	}
+
+	FTransform FRigidBodyAsyncGT::GetTransform() const
+	{
+		CHAOSRIGIDPHYSICSASYNC_CHECK(IsValid());
+		return FTransform(GetProxyAPI().R(), GetProxyAPI().X());
+	}
+
+	FBounds3d FRigidBodyAsyncGT::GetBounds() const
+	{
+		CHAOSRIGIDPHYSICSASYNC_CHECK(IsValid());
+		FConstGenericParticleHandle ParticleHandle = Proxy->GetHandle_LowLevel();
+		FAABB3 Bounds;
+		if (GetProxyAPI().GetGeometry()->HasBoundingBox())
+		{
+			Bounds = GetProxyAPI().GetGeometry()->CalculateTransformedBounds(GetTransform());
+		}
+		else
+		{
+			Bounds = FAABB3::FullAABB();
+		}
+		return FBounds3d(Bounds.Min(), Bounds.Max());
+	}
+
+	double FRigidBodyAsyncGT::GetMass() const
+	{
+		CHAOSRIGIDPHYSICSASYNC_CHECK(IsValid());
+		return GetProxyAPI().M();
+	}
+
+	void FRigidBodyAsyncGT::SetMass(double InMass)
+	{
+		CHAOSRIGIDPHYSICSASYNC_CHECK(IsValid());
+		GetProxyAPI().SetM(InMass);
+	}
+
+	//void SetCCDEnabled(bool bInEnabled) {  }
+	//void SetMACDEnabled(bool bInEnabled) {  }
+	//void SetSmoothEdgeCollisionsEnabled(bool bInEnabled) {  }
+	//void SetMaxLinearVelocity(float InMax) {  }
+	//void SetMaxAngularVelocity(float InMax) {  }
+	//void SetInertiaConditioningEnabled(bool bInEnabled) {  }
+	//void SetGyroscopicTorqueEnabled(bool bInEnabled) {  }
+	//void SetGravityGroup(int32 InId) {  }
+
+	void FRigidBodyAsyncGT::SetGravityScale(float InScale)
+	{ 
+		GetProxyAPI().SetGravityEnabled(InScale > 0.0f);
+	}
+
+	//void SetWakeEventsEnabled(bool bInEnabled) {  }
+	//void SetNumPositionIterations(int32 InNum) {  }
+	//void SetNumVelocityIterations(int32 InNum) {  }
+	//void SetNumProjectionIterations(int32 InNum) {  }
+	//void SetIsSleeping(bool bInSleeping) {  }
+}
+
+namespace Chaos::Rigids::Async
+{
+	FRigidBodyAsyncPT::FRigidBodyAsyncPT(
+		FRigidBodyContainerAsyncPT* InContainer, 
+		FSingleParticlePhysicsProxy* InProxy, 
+		const FRigidBodyIdAsync& InBodyId)
+		: Container(InContainer)
+		, Proxy(InProxy)
+		, BodyId(InBodyId)
+	{
+	}
+
+	FRigidBodyAsyncPT::~FRigidBodyAsyncPT()
+	{
+	}
+
+	const char* FRigidBodyAsyncPT::GetTypeName() const
+	{
+		return "RigidBodyAsyncPT";
+	}
+
+	UE::Physics::FRigidBodyId FRigidBodyAsyncPT::GetBodyId() const
+	{
+		CHAOSRIGIDPHYSICSASYNC_CHECK(IsValid());
+		return GetBodyIdAsync().AsRigidBodyId();
+	}
+
+	UE::Physics::FRigidBodyHandle FRigidBodyAsyncPT::GetBodyHandle() const
+	{
+		CHAOSRIGIDPHYSICSASYNC_CHECK(IsValid());
+		return UE::Physics::FRigidBodyHandle(Container->GetContainerHandle(), GetBodyId());
+	}
+
+	bool FRigidBodyAsyncPT::IsValid() const
+	{
+		return (Proxy != nullptr) && (Proxy->GetHandle_LowLevel() != nullptr);
+	}
+
+	bool FRigidBodyAsyncPT::IsInScene() const
+	{
+		// TODO_CHAOSAPI: Asymmetry - proxies are destroyed on remove from scene
+		// so once we are in the scene we will always be in the scene until destroyed
+		// TODO_CHAOSAPI: need a better way to know if we're in the scene
+		return (Proxy != nullptr) && (Proxy->GetHandle_LowLevel() != nullptr) && (Proxy->GetHandle_LowLevel()->PhysicsProxy() != nullptr);
+	}
+
+	bool FRigidBodyAsyncPT::IsStatic() const
+	{
+		return (GetProxyAPI()->ObjectState() == EObjectStateType::Static);
+	}
+
+	bool FRigidBodyAsyncPT::IsKinematic() const
+	{
+		return (GetProxyAPI()->ObjectState() == EObjectStateType::Kinematic);
+	}
+
+	bool FRigidBodyAsyncPT::IsDynamic() const
+	{
+		return (GetProxyAPI()->ObjectState() == EObjectStateType::Dynamic) || (GetProxyAPI()->ObjectState() == EObjectStateType::Sleeping);
+	}
+
+	bool FRigidBodyAsyncPT::IsSleeping() const
+	{
+		return (GetProxyAPI()->ObjectState() == EObjectStateType::Sleeping);
+	}
+
+	void FRigidBodyAsyncPT::InitTransform(const FTransform& InTransform)
+	{
+		CHAOSRIGIDPHYSICSASYNC_CHECK(IsValid());
+		if (FPBDRigidsSolver* Solver = Proxy->GetSolver<FPBDRigidsSolver>())
+		{
+			// NOTE: The Chaos Proxy API does not do everything required when changing
+			// properties on the particle. E.g., marking dirty, etc.
+			// Here'what should be in GetProxyAPI()->SetX()/SetR()
+			FGeometryParticleHandle* ParticleHandle = Proxy->GetHandle_LowLevel();
+			Solver->GetEvolution()->SetParticleTransform(ParticleHandle, InTransform.GetTranslation(), InTransform.GetRotation(), /*bIsTeleport*/true);
+			Solver->GetEvolution()->DirtyParticle(*ParticleHandle);
+		}
+	}
+
+	void FRigidBodyAsyncPT::UpdateTransform(const FTransform& InTransform)
+	{
+		CHAOSRIGIDPHYSICSASYNC_CHECK(IsValid());
+		if (FPBDRigidsSolver* Solver = Proxy->GetSolver<FPBDRigidsSolver>())
+		{
+			// NOTE: The Chaos Proxy API does not do everything required when changing
+			// properties on the particle. E.g., marking dirty, etc.
+			// Here'what should be in GetProxyAPI()->SetX()/SetR()
+			FGeometryParticleHandle* ParticleHandle = Proxy->GetHandle_LowLevel();
+			Solver->GetEvolution()->SetParticleTransform(ParticleHandle, InTransform.GetTranslation(), InTransform.GetRotation(), /*bIsTeleport*/false);
+			Solver->GetEvolution()->DirtyParticle(*ParticleHandle);
+		}
+	}
+
+	void FRigidBodyAsyncPT::SetKinematicTarget(const FTransform& InTransform)
+	{
+		CHAOSRIGIDPHYSICSASYNC_CHECK(IsValid());
+		if (FPBDRigidsSolver* Solver = Proxy->GetSolver<FPBDRigidsSolver>())
+		{
+			FGeometryParticleHandle* ParticleHandle = Proxy->GetHandle_LowLevel();
+			Solver->GetEvolution()->SetParticleKinematicTarget(ParticleHandle, FKinematicTarget::MakePositionTarget(InTransform.GetTranslation(), InTransform.GetRotation()));
+			Solver->GetEvolution()->DirtyParticle(*ParticleHandle);
+		}
+	}
+
+	FTransform FRigidBodyAsyncPT::GetTransform() const
+	{
+		CHAOSRIGIDPHYSICSASYNC_CHECK(IsValid());
+		FConstGenericParticleHandle ParticleHandle = Proxy->GetHandle_LowLevel();
+		return FTransform(ParticleHandle->Q(), ParticleHandle->P());
+	}
+
+	FBounds3d FRigidBodyAsyncPT::GetBounds() const
+	{
+		CHAOSRIGIDPHYSICSASYNC_CHECK(IsValid());
+		FConstGenericParticleHandle ParticleHandle = Proxy->GetHandle_LowLevel();
+		return FBounds3d(ParticleHandle->WorldSpaceInflatedBounds().Min(), ParticleHandle->WorldSpaceInflatedBounds().Max());
+	}
+
+	double FRigidBodyAsyncPT::GetMass() const
+	{
+		CHAOSRIGIDPHYSICSASYNC_CHECK(IsValid());
+		FConstGenericParticleHandle ParticleHandle = Proxy->GetHandle_LowLevel();
+		return ParticleHandle->M();
+	}
+
+	void FRigidBodyAsyncPT::SetMass(double InMass)
+	{
+		CHAOSRIGIDPHYSICSASYNC_CHECK(IsValid());
+		ensure(false);
+		// TODO_CHAOSAPI: Not yet supported on the PT because the mass will not make it back to GT.
+		//FGenericParticleHandle ParticleHandle = Proxy->GetHandle_LowLevel();
+		//return ParticleHandle->SetM(InMass);
+	}
+}
+
+#endif // UE_RIGIDPHYSICS_API_ENABLED
\ No newline at end of file
diff --git a/Engine/Plugins/Experimental/ChaosRigidPhysicsAsync/Source/ChaosRigidPhysicsAsync/Private/ChaosRigidPhysicsAsync/RigidBodyContainerAsync.cpp b/Engine/Plugins/Experimental/ChaosRigidPhysicsAsync/Source/ChaosRigidPhysicsAsync/Private/ChaosRigidPhysicsAsync/RigidBodyContainerAsync.cpp
new file mode 100644
--- /dev/null
+++ b/Engine/Plugins/Experimental/ChaosRigidPhysicsAsync/Source/ChaosRigidPhysicsAsync/Private/ChaosRigidPhysicsAsync/RigidBodyContainerAsync.cpp
@@ -0,0 +1,334 @@
+// Copyright Epic Games, Inc. All Rights Reserved.
+#include "ChaosRigidPhysicsAsync/RigidBodyContainerAsync.h"
+
+#include "ChaosRigidPhysicsAsync/RigidBodyAsync.h"
+#include "ChaosRigidPhysicsAsync/RigidSceneAsync.h"
+#include "RigidPhysics/RigidAccessor.h"
+#include "RigidPhysics/RigidBody.h"
+
+#if UE_RIGIDPHYSICS_API_ENABLED
+
+namespace Chaos::Rigids::Async
+{
+	UE::Physics::FRigidBodyContainerTypeId FRigidBodyContainerAsyncGT::StaticRigidBodyContainerTypeId()
+	{
+		using namespace UE::Physics;
+
+		static const TCHAR* StaticTypeName = TEXT("RigidBodyContainerAsyncGT");
+		static FRigidBodyContainerTypeId StaticTypeId = FRigidBodyContainerTypeId(StaticTypeName);
+		return StaticTypeId;
+	}
+
+	FRigidBodyContainerAsyncGT::FRigidBodyContainerAsyncGT(const UE::Physics::FRigidBodyContainerId& InContainerId)
+		: UE::Physics::IRigidBodyContainer(StaticRigidBodyContainerTypeId())
+		, Scene(nullptr)
+		, ContainerId(InContainerId)
+		, Solver(nullptr)
+		, BodyPool(1000, 0)
+		, Bodies()
+		, BodiesFreeIndices()
+		, NextBodyIdIndex(0)
+	{
+	}
+
+	FRigidBodyContainerAsyncGT::~FRigidBodyContainerAsyncGT()
+	{
+		check(BodyPool.GetNumAllocated() == 0);
+		check(BodiesFreeIndices.Num() == Bodies.Num());
+	}
+
+	UE::Physics::IRigidScene* FRigidBodyContainerAsyncGT::GetScene() const
+	{
+		return GetSceneAsync();
+	}
+
+	UE::Physics::FRigidBodyContainerId FRigidBodyContainerAsyncGT::GetContainerId() const
+	{
+		return ContainerId;
+	}
+
+	FRigidBodyAsyncGT* FRigidBodyContainerAsyncGT::GetRigidBody(const UE::Physics::FRigidBodyId& BodyId) const
+	{
+		// Extract the AsyncId data from the BodyId
+		FRigidBodyIdAsync AsyncId = FRigidBodyIdAsync(BodyId);
+		if (Bodies.IsValidIndex(AsyncId.GetBodyIndex()))
+		{
+			// If the Epochs match the BodyId is current and we can return the body
+			const FRigidBodyArrayEntry& BodyEntry = Bodies[AsyncId.GetBodyIndex()];
+			if (BodyEntry.Epoch == AsyncId.GetEpoch())
+			{
+				return BodyEntry.Body;
+			}
+		}
+
+		// If the Epochs don't match, the Body with BodyId was destroyed
+		return nullptr;
+	}
+
+	UE::Physics::FRigidBodyContainerHandle FRigidBodyContainerAsyncGT::GetContainerHandle() const
+	{
+		return UE::Physics::FRigidBodyContainerHandle(Scene->GetSceneHandle(), GetContainerId());
+	}
+
+	void FRigidBodyContainerAsyncGT::Init(FRigidSceneAsyncGT* InScene, FPBDRigidsSolver* InSolver)
+	{
+		Scene = InScene;
+		Solver = InSolver;
+	}
+
+	FRigidBodyAsyncGT* FRigidBodyContainerAsyncGT::CreateRigidBody(const UE::Physics::FRigidNamePtr& InName, UE::Physics::ERigidMovementType InMovementType)
+	{
+		using namespace UE::Physics;
+
+		TUniquePtr<FGeometryParticle> Particle;
+		if (InMovementType == ERigidMovementType::Static)
+		{
+			Particle = FGeometryParticle::CreateParticle();
+		}
+		else
+		{
+			TUniquePtr<FPBDRigidParticle> RigidParticle = FPBDRigidParticle::CreateParticle();
+			if (InMovementType == ERigidMovementType::Kinematic)
+			{
+				RigidParticle->SetObjectState(EObjectStateType::Kinematic);
+			}
+			Particle = MoveTemp(RigidParticle);
+		}
+
+
+#if CHAOS_DEBUG_NAME
+		Particle->SetDebugName(InName);
+#endif
+
+		// Create a default shape for testing
+		FImplicitObjectPtr Geom = FImplicitObjectPtr(new FImplicitBox3(FVec3(-50), FVec3(50)));
+		FImplicitObjectPtr Union = new FImplicitObjectUnion(TArray<FImplicitObjectPtr>({ Geom }));
+
+		Particle->SetGeometry(Union);
+
+		return CreateRigidBodyForProxy(FSingleParticlePhysicsProxy::Create(MoveTemp(Particle)));
+	}
+
+	FRigidBodyAsyncGT* FRigidBodyContainerAsyncGT::CreateRigidBodyForProxy(FSingleParticlePhysicsProxy* InProxy)
+	{
+		// Store the proxy with a unique Id (unique within this container)
+		return AllocBody(InProxy);
+	}
+
+	void FRigidBodyContainerAsyncGT::DestroyRigidBody(const UE::Physics::FRigidBodyId& BodyId)
+	{
+		if (FRigidBodyAsyncGT* BodyGT = GetRigidBody(BodyId))
+		{
+			// TODO_CHAOSAPI: Register/Unregister asymmetry.
+			// NOTE: This is the correct place for removing the body from the pool, but see RemoveRigidBodyFromScene 
+			// which is where it usually happens atm. It still needs to be here as well to destroy particles that were 
+			// never added to the scene. If the particle is destroyed in RemoveRigidBodyFromScene,  then GetRigidBody() 
+			// will have returned null so there's no danger of duplicate destruction.
+			FreeBody(BodyGT);
+		}
+	}
+
+	void FRigidBodyContainerAsyncGT::AddRigidBodyToScene(const UE::Physics::FRigidBodyId& BodyId)
+	{
+		// This ensure will fire if you try to re-add a body to a scene after removing it. This should be supported
+		// but FPBDRigidsSolver::UnregisterObject destroys the proxy so we would have to recreate it and we would
+		// lose any changes the user may have made. For now just treat this as an error.
+		FRigidBodyAsyncGT* BodyGT = GetRigidBody(BodyId);
+		if (ensure(BodyGT != nullptr))
+		{
+			if (!BodyGT->IsInScene())
+			{
+				// NOTE: Handle GT and PT
+				Solver->RegisterObject(BodyGT->GetProxy());
+			}
+		}
+	}
+
+	void FRigidBodyContainerAsyncGT::RemoveRigidBodyFromScene(const UE::Physics::FRigidBodyId& BodyId)
+	{
+		if (FRigidBodyAsyncGT* BodyGT = GetRigidBody(BodyId))
+		{
+			if (BodyGT->IsInScene())
+			{
+				// NOTE: Handles GT and PT
+				Solver->UnregisterObject(BodyGT->GetProxy());
+			}
+
+			// TODO_CHAOSAPI: Register/Unregister asymmetry.
+			// This should be called from DestroyRigidBody but FPBDRigidsSolver::UnregisterObject 
+			// enqueues a command that destroys the proxy from the physics thread. That command may be executed
+			// before we call DestroyRigidBody which would leave the body with a dangling proxy pointer
+			// unless we also free the body here.
+			FreeBody(BodyGT);
+		}
+	}
+
+	FRigidSceneAsyncGT* FRigidBodyContainerAsyncGT::GetSceneAsync() const
+	{
+		return Scene;
+	}
+
+	void FRigidBodyContainerAsyncGT::OnRigidBodyTransformUpdated(FRigidBodyAsyncGT* Body)
+	{
+		if (Scene != nullptr)
+		{
+			Scene->OnRigidBodyTransformUpdated(Body);
+		}
+	}
+
+	FRigidBodyAsyncGT* FRigidBodyContainerAsyncGT::AllocBody(FSingleParticlePhysicsProxy* InProxy)
+	{
+		// Make sure we have free array entries
+		if (BodiesFreeIndices.IsEmpty())
+		{
+			ReserveBodyArray(Bodies.Num() + 1000);
+		}
+
+		// Assign an array index
+		int32 BodyIndex = BodiesFreeIndices.Pop(EAllowShrinking::No);
+		FRigidBodyArrayEntry& BodyEntry = Bodies[BodyIndex];
+		check(BodyEntry.Body == nullptr);
+
+		// Build the Id from the BodyIndex and Epoch
+		FRigidBodyIdAsync AsyncId = FRigidBodyIdAsync(BodyIndex, BodyEntry.Epoch);
+
+		// Allocate the body from the pool
+		FRigidBodyAsyncGT* Body = BodyPool.Alloc(this, InProxy, AsyncId);
+
+		// Store the body in the ID->Body Map
+		Bodies[BodyIndex].Body = Body;
+
+		return Body;
+	}
+
+	void FRigidBodyContainerAsyncGT::FreeBody(FRigidBodyAsyncGT* BodyGT)
+	{
+		check(BodyGT != nullptr);
+
+		// Get the entry in the Id->Body map
+		FRigidBodyIdAsync AsyncId = BodyGT->GetBodyIdAsync();
+		FRigidBodyArrayEntry& BodyEntry = Bodies[AsyncId.GetBodyIndex()];
+
+		// When FreeBody is called we must know we have a body from this container
+		check(BodyEntry.Epoch == AsyncId.GetEpoch());
+		check(BodyEntry.Body == BodyGT);
+
+		// Reset the entry and increment the epoch to invalidate outstanding handles
+		BodyEntry.Body = nullptr;
+		BodyEntry.Epoch = BodyEntry.Epoch + 1;
+
+		// Add the array item back into the free pool
+		BodiesFreeIndices.Push(AsyncId.GetBodyIndex());
+
+		// Return the body to the pool
+		BodyPool.Free(BodyGT);
+	}
+
+	void FRigidBodyContainerAsyncGT::ReserveBodyArray(int32 MaxBodies)
+	{
+		if (MaxBodies > Bodies.Num())
+		{
+			int32 PrevNumBodies = Bodies.Num();
+			Bodies.SetNum(MaxBodies);
+			BodiesFreeIndices.Reserve(MaxBodies);
+
+			// Loop backawards so we allocate lowest indices first by popping from BodiesFreeIndices
+			for (int32 BodyIndex = MaxBodies - 1; BodyIndex >= PrevNumBodies; --BodyIndex)
+			{
+				Bodies[BodyIndex].Body = nullptr;
+				Bodies[BodyIndex].Epoch = 0;
+				BodiesFreeIndices.Push(BodyIndex);
+			}
+		}
+	}
+}
+
+
+namespace Chaos::Rigids::Async
+{
+	UE::Physics::FRigidBodyContainerTypeId FRigidBodyContainerAsyncPT::StaticRigidBodyContainerTypeId()
+	{
+		using namespace UE::Physics;
+
+		static const TCHAR* StaticTypeName = TEXT("RigidBodyContainerAsyncPT");
+		static FRigidBodyContainerTypeId StaticTypeId = FRigidBodyContainerTypeId(StaticTypeName);
+		return StaticTypeId;
+	}
+
+	FRigidBodyContainerAsyncPT::FRigidBodyContainerAsyncPT(const UE::Physics::FRigidBodyContainerId& InContainerId)
+		: IRigidBodyContainer(StaticRigidBodyContainerTypeId())
+		, Scene(nullptr)
+		, ContainerId(InContainerId)
+		, Solver(nullptr)
+		, BodyPool(1000, 0)
+		, IdToBodyMap()
+	{
+	}
+
+	FRigidBodyContainerAsyncPT::~FRigidBodyContainerAsyncPT()
+	{
+		check(BodyPool.GetNumAllocated() == 0);
+		check(IdToBodyMap.IsEmpty());
+	}
+
+	UE::Physics::IRigidScene* FRigidBodyContainerAsyncPT::GetScene() const
+	{
+		return Scene;
+	}
+
+	UE::Physics::FRigidBodyContainerId FRigidBodyContainerAsyncPT::GetContainerId() const
+	{
+		return ContainerId;
+	}
+
+	FRigidBodyAsyncPT* FRigidBodyContainerAsyncPT::GetRigidBody(const UE::Physics::FRigidBodyId& BodyId) const
+	{
+		const FRigidBodyAsyncPT* const* PBody = IdToBodyMap.Find(BodyId);
+		if (PBody != nullptr)
+		{
+			return const_cast<FRigidBodyAsyncPT*>(*PBody);
+		}
+		return nullptr;
+	}
+
+	void FRigidBodyContainerAsyncPT::Init(FRigidSceneAsyncPT* InScene, FPBDRigidsSolver* InSolver)
+	{
+		Scene = InScene;
+		Solver = InSolver;
+	}
+
+	UE::Physics::FRigidBodyContainerHandle FRigidBodyContainerAsyncPT::GetContainerHandle() const
+	{
+		return UE::Physics::FRigidBodyContainerHandle(Scene->GetSceneHandle(), GetContainerId());
+	}
+
+	void FRigidBodyContainerAsyncPT::CreateRigidBody(const FRigidBodyIdAsync& InBodyId, FSingleParticlePhysicsProxy* InProxy)
+	{
+		FRigidBodyAsyncPT* Body = BodyPool.Alloc(this, InProxy, InBodyId);
+		IdToBodyMap.Add(InBodyId.AsRigidBodyId(), Body);
+	}
+
+	void FRigidBodyContainerAsyncPT::DestroyRigidBody(const UE::Physics::FRigidBodyId& BodyId)
+	{
+		if (FRigidBodyAsyncPT* Body = GetRigidBody(BodyId))
+		{
+			IdToBodyMap.Remove(BodyId);
+			BodyPool.Free(Body);
+		}
+	}
+
+	void FRigidBodyContainerAsyncPT::AddRigidBodyToScene(const UE::Physics::FRigidBodyId& BodyId)
+	{
+		// TODO_CHAOSAPI: nothing to do for now because RegisterObject does GT and PT
+		// See FRigidBodyContainerAsyncGT::AddRigidBodyToScene
+	}
+
+	void FRigidBodyContainerAsyncPT::RemoveRigidBodyFromScene(const UE::Physics::FRigidBodyId& BodyId)
+	{
+		// TODO_CHAOSAPI: nothing to do for now because UnregisterObject does GT and PT
+		// See FRigidBodyContainerAsyncGT::RemoveRigidBodyFromScene
+	}
+
+}
+
+#endif // UE_RIGIDPHYSICS_API_ENABLED
\ No newline at end of file
diff --git a/Engine/Plugins/Experimental/ChaosRigidPhysicsAsync/Source/ChaosRigidPhysicsAsync/Private/ChaosRigidPhysicsAsync/RigidFactoryAsync.cpp b/Engine/Plugins/Experimental/ChaosRigidPhysicsAsync/Source/ChaosRigidPhysicsAsync/Private/ChaosRigidPhysicsAsync/RigidFactoryAsync.cpp
new file mode 100644
--- /dev/null
+++ b/Engine/Plugins/Experimental/ChaosRigidPhysicsAsync/Source/ChaosRigidPhysicsAsync/Private/ChaosRigidPhysicsAsync/RigidFactoryAsync.cpp
@@ -0,0 +1,111 @@
+// Copyright Epic Games, Inc. All Rights Reserved.
+#include "ChaosRigidPhysicsAsync/RigidFactoryAsync.h"
+
+#include "ChaosRigidPhysicsAsync/RigidBodyAsync.h"
+#include "ChaosRigidPhysicsAsync/RigidBodyContainerAsync.h"
+#include "ChaosRigidPhysicsAsync/RigidSceneAsync.h"
+#include "RigidPhysics/RigidAccessor.h"
+#include "RigidPhysics/RigidBody.h"
+
+#if UE_RIGIDPHYSICS_API_ENABLED
+
+namespace Chaos::Rigids::Async
+{
+	UE::Physics::FRigidSceneHandle FRigidFactoryAsync::CreateRigidScene(const FRigidSceneSettingsAsync& InSettings)
+	{
+		// Create a scene and store a hard reference to prevent its destruction
+		TSharedPtr<FRigidSceneAsyncGT, ESPMode::ThreadSafe> Scene = MakeShared<FRigidSceneAsyncGT>(InSettings);
+		Scenes.Add(Scene);
+
+		Scene->Startup();
+
+		// Return a weak/unlocked reference to the scene
+		return UE::Physics::FRigidSceneHandle(Scene);
+	}
+
+	UE::Physics::FRigidSceneHandle FRigidFactoryAsync::CreateRigidSceneFromSolver(FPBDRigidsSolver* InSolver)
+	{
+		// Create a scene and store a hard reference to prevent its destruction
+		TSharedPtr<FRigidSceneAsyncGT, ESPMode::ThreadSafe> Scene = MakeShared<FRigidSceneAsyncGT>(InSolver);
+		Scenes.Add(Scene);
+
+		Scene->Startup();
+
+		// Return a weak/unlocked reference to the scene
+		return UE::Physics::FRigidSceneHandle(Scene);
+	}
+
+	void FRigidFactoryAsync::DestroyRigidScene(const UE::Physics::FRigidSceneHandle& InSceneHandle)
+	{
+		FRigidSceneAsyncGT* Scene = static_cast<FRigidSceneAsyncGT*>(Private::FRigidAccessor::GetIRigidScene(InSceneHandle.LockRW()));
+
+		Scene->Shutdown();
+
+		Scenes.RemoveAll([Scene](const UE::Physics::FRigidSceneSharedPtr& Item) { return !Item.IsValid() || (Item.Get() == Scene); });
+	}
+}
+
+namespace Chaos::Rigids::Async::Private
+{
+	class FRigidAccessorAsync
+	{
+	public:
+		static UE::Physics::FRigidSceneHandle CreateSceneFromSolver(FRigidFactoryAsync& InFactory, FPBDRigidsSolver* InSolver)
+		{
+			return InFactory.CreateRigidSceneFromSolver(InSolver);
+		}
+	};
+
+	FRigidSceneAsyncGT* GetAsyncScene(const UE::Physics::TRigidScenePtr<UE::Physics::FGameContextRW>& InScenePtr)
+	{
+		return static_cast<FRigidSceneAsyncGT*>(Private::FRigidAccessor::GetIRigidScene(InScenePtr));
+	}
+
+	FRigidBodyContainerAsyncGT* GetAsyncBodyContainer(const UE::Physics::TRigidScenePtr<UE::Physics::FGameContextRW>& InScenePtr)
+	{
+		if (FRigidSceneAsyncGT* AsyncScene = GetAsyncScene(InScenePtr))
+		{
+			return static_cast<FRigidBodyContainerAsyncGT*>(AsyncScene->GetRigidBodyContainer(UE::Physics::FRigidBodyContainerId(0)));
+		}
+		return nullptr;
+	}
+
+	FRigidBodyAsyncGT* GetAsyncBody(const UE::Physics::TRigidBodyPtr<UE::Physics::FGameContextRO>& InBodyPtr)
+	{
+		return static_cast<FRigidBodyAsyncGT*>(Private::FRigidAccessor::GetIRigidBody(InBodyPtr));
+	}
+
+	UE::Physics::FRigidSceneHandle CreateSceneFromSolver(FRigidFactoryAsync& InFactory, FPBDRigidsSolver* InSolver)
+	{
+		return FRigidAccessorAsync::CreateSceneFromSolver(InFactory, InSolver);
+	}
+
+	UE::Physics::FRigidBodyHandle CreateBodyFromProxy(const UE::Physics::FRigidSceneHandle& InSceneHandle, FSingleParticlePhysicsProxy* InProxy)
+	{
+		using namespace UE::Physics;
+
+		if (TRigidScenePtr<FGameContextRW> ScenePtr = InSceneHandle.LockRW())
+		{
+			if (FRigidBodyContainerAsyncGT* BodyContainer = GetAsyncBodyContainer(ScenePtr))
+			{
+				FRigidBodyAsyncGT* Body = BodyContainer->CreateRigidBodyForProxy(InProxy);
+				return FRigidBodyHandle(FRigidBodyContainerHandle(InSceneHandle, BodyContainer->GetContainerId()), Body->GetBodyId());
+			}
+		}
+		return FRigidBodyHandle();
+	}
+
+	FSingleParticlePhysicsProxy* GetProxyFromBody(const UE::Physics::FRigidBodyHandle& InBodyHandle)
+	{
+		using namespace UE::Physics;
+
+		if (TRigidBodyPtr<FGameContextRO> BodyPtr = InBodyHandle.LockRO())
+		{
+			FRigidBodyAsyncGT* Body = GetAsyncBody(BodyPtr);
+			return Body->GetProxy();
+		}
+		return nullptr;
+	}
+}
+
+#endif // UE_RIGIDPHYSICS_API_ENABLED
\ No newline at end of file
diff --git a/Engine/Plugins/Experimental/ChaosRigidPhysicsAsync/Source/ChaosRigidPhysicsAsync/Private/ChaosRigidPhysicsAsync/RigidSceneAsync.cpp b/Engine/Plugins/Experimental/ChaosRigidPhysicsAsync/Source/ChaosRigidPhysicsAsync/Private/ChaosRigidPhysicsAsync/RigidSceneAsync.cpp
new file mode 100644
--- /dev/null
+++ b/Engine/Plugins/Experimental/ChaosRigidPhysicsAsync/Source/ChaosRigidPhysicsAsync/Private/ChaosRigidPhysicsAsync/RigidSceneAsync.cpp
@@ -0,0 +1,559 @@
+// Copyright Epic Games, Inc. All Rights Reserved.
+#include "ChaosRigidPhysicsAsync/RigidSceneAsync.h"
+#include "RigidPhysics/RigidModifier.h"
+
+#if UE_RIGIDPHYSICS_API_ENABLED
+
+namespace Chaos::Rigids::Async
+{
+	FRigidSceneAsyncGT::FRigidSceneAsyncGT(const FRigidSceneSettingsAsync& InSceneSettings)
+		: SceneSettings(InSceneSettings)
+		, Solver(nullptr)
+		, SpatialAcceleration(nullptr)
+		, SingleBodyContainer(UE::Physics::FRigidBodyContainerId(0))
+		, SceneModifiers()
+		, ScenePT()
+		, bOwnsSolver(false)
+	{
+		ScenePT = MakeUnique<FRigidSceneAsyncPT>(UE::Physics::FRigidBodyContainerId(0));
+	}
+
+	FRigidSceneAsyncGT::FRigidSceneAsyncGT(FPBDRigidsSolver* InSolver)
+		: SceneSettings()
+		, Solver(InSolver)
+		, SpatialAcceleration(nullptr)
+		, SingleBodyContainer(UE::Physics::FRigidBodyContainerId(0))
+		, SceneModifiers()
+		, ScenePT()
+		, bOwnsSolver(false)
+	{
+		ScenePT = MakeUnique<FRigidSceneAsyncPT>(UE::Physics::FRigidBodyContainerId(0));
+	}
+
+	FRigidSceneAsyncGT::~FRigidSceneAsyncGT()
+	{
+		check(Solver == nullptr);
+	}
+
+	UE::Physics::FRigidSceneHandle FRigidSceneAsyncGT::GetSceneHandle() const
+	{
+		return SceneHandle;
+	}
+
+	void FRigidSceneAsyncGT::Startup()
+	{
+		using namespace UE::Physics;
+
+		SceneHandle = FRigidSceneHandle(AsShared());
+
+		if (Solver == nullptr)
+		{
+			// NOTE: Not using full Async for now - it works but it complicates testing (time shift and interpolation)
+			FChaosSolversModule* Module = FChaosSolversModule::GetModule();
+			Solver = Module->CreateSolver(/*Owner=*/nullptr, /*AsyncDt*/-1, EThreadingMode::TaskGraph, SceneSettings.SceneName);
+			bOwnsSolver = true;
+		}
+
+		SingleBodyContainer.Init(this, Solver);
+
+		ScenePT->Startup(SceneHandle, Solver);
+	}
+
+	void FRigidSceneAsyncGT::Shutdown()
+	{
+		ScenePT->Shutdown();
+
+		if (Solver != nullptr)
+		{
+			if (bOwnsSolver)
+			{
+				FChaosSolversModule* Module = FChaosSolversModule::GetModule();
+				Module->DestroySolver(Solver);
+				bOwnsSolver = false;
+			}
+			else
+			{
+				// From FPhysicsSolverBase::DestroySolver - we have no easy way to flush commands so this will do for now
+				Solver->SetThreadingMode_External(EThreadingModeTemp::SingleThread);
+				Solver->AdvanceAndDispatch_External(0);	//flush any pending commands are executed (for example unregister object)
+			}
+
+			Solver = nullptr;
+		}
+	}
+
+	void FRigidSceneAsyncGT::Lock(UE::Physics::EContextLockType InLockType)
+	{
+		using namespace UE::Physics;
+
+		CHAOSRIGIDPHYSICSASYNC_CHECK(IsInGameThreadContext());
+
+		if (Solver != nullptr)
+		{
+			switch (InLockType)
+			{
+			case EContextLockType::ReadOnly:
+				Solver->GetExternalDataLock_External().ReadLock();
+				break;
+			case EContextLockType::ReadWrite:
+				Solver->GetExternalDataLock_External().WriteLock();
+				break;
+			}
+		}
+	}
+
+	void FRigidSceneAsyncGT::Unlock(UE::Physics::EContextLockType InLockType)
+	{
+		using namespace UE::Physics;
+
+		CHAOSRIGIDPHYSICSASYNC_CHECK(IsInGameThreadContext());
+
+		if (Solver != nullptr)
+		{
+			switch (InLockType)
+			{
+			case EContextLockType::ReadOnly:
+				Solver->GetExternalDataLock_External().ReadUnlock();
+				break;
+			case EContextLockType::ReadWrite:
+				Solver->GetExternalDataLock_External().WriteUnlock();
+				break;
+			}
+		}
+	}
+
+	FRigidBodyAsyncGT* FRigidSceneAsyncGT::CreateRigidBody(const UE::Physics::FRigidNamePtr& InName, UE::Physics::ERigidMovementType InMovementType)
+	{
+		CheckGameWriteLocked();
+
+		FRigidBodyAsyncGT* BodyGT = SingleBodyContainer.CreateRigidBody(InName, InMovementType);
+
+		Solver->EnqueueCommandImmediate(
+			[this, BodyId = BodyGT->GetBodyIdAsync(), Proxy = BodyGT->Proxy]()
+			{
+				ScenePT->CreateRigidBody(BodyId, Proxy);
+			});
+
+		return BodyGT;
+	}
+
+	void FRigidSceneAsyncGT::DestroyRigidBody(const UE::Physics::FRigidBodyHandle& BodyHandle)
+	{
+		CheckGameWriteLocked();
+
+		if (BodyHandle.GetContainerHandle().GetContainerId() == SingleBodyContainer.GetContainerId())
+		{
+			UE::Physics::FRigidBodyId BodyId = BodyHandle.GetBodyId();
+
+			SingleBodyContainer.DestroyRigidBody(BodyId);
+
+			Solver->EnqueueCommandImmediate(
+				[BodyId, this]()
+				{
+					ScenePT->DestroyRigidBody(BodyId);
+				});
+		}
+	}
+
+	void FRigidSceneAsyncGT::RegisterRigidBody(const UE::Physics::FRigidBodyHandle& BodyHandle)
+	{
+		CheckGameWriteLocked();
+
+		if (BodyHandle.GetContainerHandle().GetContainerId() == SingleBodyContainer.GetContainerId())
+		{
+			SingleBodyContainer.AddRigidBodyToScene(BodyHandle.GetBodyId());
+
+			Solver->EnqueueCommandImmediate(
+				[BodyId = BodyHandle.GetBodyId(), this]()
+				{
+					ScenePT->RegisterRigidBody(BodyId);
+				});
+		}
+	}
+
+	void FRigidSceneAsyncGT::UnregisterRigidBody(const UE::Physics::FRigidBodyHandle& BodyHandle)
+	{
+		CheckGameWriteLocked();
+
+		if (BodyHandle.GetContainerHandle().GetContainerId() == SingleBodyContainer.GetContainerId())
+		{
+			SingleBodyContainer.RemoveRigidBodyFromScene(BodyHandle.GetBodyId());
+
+			Solver->EnqueueCommandImmediate(
+				[BodyId = BodyHandle.GetBodyId(), this]()
+				{
+					ScenePT->UnregisterRigidBody(BodyId);
+				});
+		}
+	}
+
+	FRigidBodyContainerAsyncGT* FRigidSceneAsyncGT::GetRigidBodyContainer(const UE::Physics::FRigidBodyContainerId& InBodyContainerId)
+	{
+		CheckGameReadLocked();
+
+		if (InBodyContainerId == SingleBodyContainer.GetContainerId())
+		{
+			return &SingleBodyContainer;
+		}
+		return nullptr;
+	}
+
+	void FRigidSceneAsyncGT::RegisterModifier(UE::Physics::IRigidSceneModifier* Modifier)
+	{
+		CheckGameWriteLocked();
+
+		SceneModifiers.Add(Modifier);
+
+		Solver->EnqueueCommandImmediate(
+			[Modifier, this]()
+			{
+				ScenePT->RegisterModifier(Modifier);
+			});
+	}
+
+	void FRigidSceneAsyncGT::UnregisterModifier(UE::Physics::IRigidSceneModifier* Modifier)
+	{
+		CheckGameWriteLocked();
+
+		SceneModifiers.Remove(Modifier);
+
+		Solver->EnqueueCommandImmediate(
+			[Modifier, this]()
+			{
+				ScenePT->UnregisterModifier(Modifier);
+			});
+	}
+
+	bool FRigidSceneAsyncGT::Query(const UE::Physics::FRigidQueryDesc& QueryDesc, TArray<UE::Physics::FRigidQueryHit>& OutHits)
+	{
+		CheckGameReadLocked();
+
+		// TODO: Runs the query on the GT
+		return false;
+	}
+
+	void FRigidSceneAsyncGT::Tick(float Dt)
+	{
+		CheckGameWriteLocked();
+
+		if (Solver != nullptr)
+		{
+			OnPreSimulate(Dt);
+
+			Solver->AdvanceAndDispatch_External(Dt);
+			Solver->WaitOnPendingTasks_External();
+			Solver->UpdateGameThreadStructures();
+
+			OnPostSimulate(Dt);
+		}
+	}
+
+	const char* FRigidSceneAsyncGT::GetTypeName() const
+	{ 
+		return "FRigidSceneAsyncGT";
+	}
+
+	void FRigidSceneAsyncGT::OnRigidBodyTransformUpdated(FRigidBodyAsyncGT* Body)
+	{
+		UpdateRigidBodyInAccelerationStructure(Body);
+	}
+
+	void FRigidSceneAsyncGT::UpdateRigidBodyInAccelerationStructure(FRigidBodyAsyncGT* Body)
+	{
+		using namespace Chaos;
+		if (Body == nullptr)
+		{
+			return;
+		}
+
+		if (Body->IsInScene())
+		{
+			FGeometryParticle* Particle = Body->GetProxy()->GetParticle_LowLevel();
+
+			if (SpatialAcceleration != nullptr)
+			{
+				const FBounds3d Bounds = Body->GetBounds();
+				const FAABB3 WorldBounds = FAABB3(Bounds.Min, Bounds.Max);
+				const bool bHasBounds = !WorldBounds.IsFull();
+
+				FAccelerationStructureHandle AccelerationHandle(Particle);
+				SpatialAcceleration->UpdateElementIn(AccelerationHandle, WorldBounds, bHasBounds, Body->GetProxyAPI().SpatialIdx());
+			}
+
+			if (Solver != nullptr)
+			{
+				Solver->UpdateParticleInAccelerationStructure_External(Particle, EPendingSpatialDataOperation::Update);
+			}
+		}
+	}
+
+	UE::Physics::FRigidSceneHandle FRigidSceneAsyncGT::GetHandle()
+	{
+		return UE::Physics::FRigidSceneHandle(AsShared());
+	}
+
+	void FRigidSceneAsyncGT::CheckGameReadLocked() const
+	{
+		CHAOSRIGIDPHYSICSASYNC_CHECK(IsInGameThreadContext());
+
+		// TODO_CHAOSAPI: Check read or write lock
+		//if (Solver != nullptr)
+		//{
+		//	Solver->GetExternalDataLock_External().IsReadLocked();
+		//}
+	}
+
+	void FRigidSceneAsyncGT::CheckGameWriteLocked() const
+	{
+		CHAOSRIGIDPHYSICSASYNC_CHECK(IsInGameThreadContext());
+
+		// TODO_CHAOSAPI: Check write lock
+		//if (Solver != nullptr)
+		//{
+		//	Solver->GetExternalDataLock_External().IsWriteLocked();
+		//}
+	}
+
+	void FRigidSceneAsyncGT::OnPreSimulate(FReal Dt)
+	{
+		using namespace UE::Physics;
+
+		CheckGameWriteLocked();
+
+		FGameContextRW Context = FGameContextRW(EContextGamePhase::PreSimulate);
+		TRigidScenePtr<FGameContextRW> ScenePtr = UE::Physics::Private::FRigidAccessor::MakeRigidScenePtr(Context, this);
+
+		for (IRigidSceneModifier* SceneModifier : SceneModifiers)
+		{
+			SceneModifier->PreSimulate(ScenePtr);
+		}
+	}
+
+	void FRigidSceneAsyncGT::OnPostSimulate(FReal Dt)
+	{
+		using namespace UE::Physics;
+
+		CheckGameWriteLocked();
+
+		FGameContextRW Context = FGameContextRW(EContextGamePhase::PostSimulate);
+		TRigidScenePtr<FGameContextRW> ScenePtr = UE::Physics::Private::FRigidAccessor::MakeRigidScenePtr(Context, this);
+
+		for (IRigidSceneModifier* SceneModifier : SceneModifiers)
+		{
+			SceneModifier->PostSimulate(ScenePtr);
+		}
+	}
+}
+
+
+namespace Chaos::Rigids::Async
+{
+	FRigidSceneAsyncPT::FRigidSceneAsyncPT(const UE::Physics::FRigidBodyContainerId& InSingleBodyContainerId)
+		: Solver(nullptr)
+		, SingleBodyContainer(InSingleBodyContainerId)
+		, SceneModifiers()
+		, Context()
+	{
+	}
+
+	void FRigidSceneAsyncPT::Lock(UE::Physics::EContextLockType InLockType)
+	{
+		CheckSimContext();
+		// Locking not required on the physics thread.
+		// All parallel access patterns are controlled by the physics thread itself.
+		// No callbacks are called in parallel (yet).
+	}
+
+	void FRigidSceneAsyncPT::Unlock(UE::Physics::EContextLockType InLockType)
+	{
+		CheckSimContext();
+		// No lock - see Lock() for comments.
+	}
+
+	// TODO_CHAOSAPI: PT and GT APIs are diverging - perhaps IRigidScene needs to be split
+	// Or we need to handle particle creation and registeration from the PT...
+	FRigidBodyAsyncPT* FRigidSceneAsyncPT::CreateRigidBody(const UE::Physics::FRigidNamePtr& InName, UE::Physics::ERigidMovementType InMovementType)
+	{
+		CheckSimContext();
+		checkf(false, TEXT("Not supported from physics thread"));
+		return nullptr;
+	}
+
+	void FRigidSceneAsyncPT::DestroyRigidBody(const UE::Physics::FRigidBodyHandle& BodyHandle)
+	{
+		CheckSimContext();
+		checkf(false, TEXT("Not supported from physics thread"));
+	}
+
+	void FRigidSceneAsyncPT::RegisterRigidBody(const UE::Physics::FRigidBodyHandle& BodyHandle)
+	{
+		CheckSimContext();
+		checkf(false, TEXT("Not supported from physics thread"));
+	}
+
+	void FRigidSceneAsyncPT::UnregisterRigidBody(const UE::Physics::FRigidBodyHandle& BodyHandle)
+	{
+		CheckSimContext();
+		checkf(false, TEXT("Not supported from physics thread"));
+	}
+
+	FRigidBodyContainerAsyncPT* FRigidSceneAsyncPT::GetRigidBodyContainer(const UE::Physics::FRigidBodyContainerId& InBodyContainerId)
+	{
+		CheckSimContext();
+		if (InBodyContainerId == SingleBodyContainer.GetContainerId())
+		{
+			return &SingleBodyContainer;
+		}
+		return nullptr;
+	}
+
+	void FRigidSceneAsyncPT::RegisterModifier(UE::Physics::IRigidSceneModifier* Modifier)
+	{
+		CheckSimContext();
+		SceneModifiers.Add(Modifier);
+	}
+
+	void FRigidSceneAsyncPT::UnregisterModifier(UE::Physics::IRigidSceneModifier* Modifier)
+	{
+		CheckSimContext();
+		SceneModifiers.Remove(Modifier);
+	}
+
+	bool FRigidSceneAsyncPT::Query(const UE::Physics::FRigidQueryDesc& QueryDesc, TArray<UE::Physics::FRigidQueryHit>& OutHits)
+	{
+		CheckSimContext();
+
+		// TODO: Runs the query on the PT
+
+		return false;
+	}
+
+	void FRigidSceneAsyncPT::Tick(float Dt)
+	{
+		checkf(false, TEXT("Not supported from physics thread"));
+	}
+
+	const char* FRigidSceneAsyncPT::GetTypeName() const
+	{ 
+		return "FRigidSceneAsyncPT";
+	}
+
+	UE::Physics::FRigidSceneHandle FRigidSceneAsyncPT::GetSceneHandle() const
+	{
+		CheckSimContext();
+		return SceneHandle;
+	}
+
+	void FRigidSceneAsyncPT::Startup(UE::Physics::FRigidSceneHandle InSceneHandle, FPBDRigidsSolver* InSolver)
+	{
+		CheckGameContext();
+
+		SceneHandle = InSceneHandle;
+
+		Solver = InSolver;
+
+		SingleBodyContainer.Init(this, Solver);
+
+		Solver->GetEvolution()->SetPreIntegrateCallback(
+			[this](FReal Dt)
+			{
+				OnPreTick(Dt);
+			});
+
+		Solver->GetEvolution()->SetPostSolveCallback(
+			[this](FReal Dt)
+			{
+				OnPostTick(Dt);
+			});
+	}
+
+	void FRigidSceneAsyncPT::Shutdown()
+	{
+		CheckGameContext();
+
+		if (Solver != nullptr)
+		{
+			Solver->GetEvolution()->SetPreIntegrateCallback([](FReal Dt) {});
+			Solver->GetEvolution()->SetPostSolveCallback([](FReal Dt) {});
+		}
+	}
+
+	void FRigidSceneAsyncPT::CreateRigidBody(const FRigidBodyIdAsync& InBodyId, FSingleParticlePhysicsProxy* Proxy)
+	{
+		CheckSimContext();
+
+		SingleBodyContainer.CreateRigidBody(InBodyId, Proxy);
+	}
+
+	void FRigidSceneAsyncPT::DestroyRigidBody(const UE::Physics::FRigidBodyId& InBodyId)
+	{
+		SingleBodyContainer.DestroyRigidBody(InBodyId);
+	}
+
+	void FRigidSceneAsyncPT::RegisterRigidBody(const UE::Physics::FRigidBodyId& InBodyId)
+	{
+		CheckSimContext();
+		// Nothing to do because RegisterObject handle GT and PT
+		// See FRigidSceneAsyncGT::RegisterRigidBody
+	}
+
+	void FRigidSceneAsyncPT::UnregisterRigidBody(const UE::Physics::FRigidBodyId& InBodyId)
+	{
+		CheckSimContext();
+		// Nothing to do because UnegisterObject handle GT and PT
+		// See FRigidSceneAsyncGT::UnregisterRigidBody
+	}
+
+	void FRigidSceneAsyncPT::CheckGameContext() const
+	{
+		CHAOSRIGIDPHYSICSASYNC_CHECK(IsInGameThreadContext());
+	}
+
+	void FRigidSceneAsyncPT::CheckSimContext() const
+	{
+		CHAOSRIGIDPHYSICSASYNC_CHECK(IsInPhysicsThreadContext());
+	}
+
+	void FRigidSceneAsyncPT::OnPreTick(FReal Dt)
+	{
+		using namespace UE::Physics;
+
+		CheckSimContext();
+
+		Context.SetSimPhase(EContextSimPhase::PreTick);
+
+		TRigidScenePtr<FSimContextRW> ScenePtr = Private::FRigidAccessor::MakeRigidScenePtr(Context, this);
+		for (IRigidSceneModifier* SceneModifier : SceneModifiers)
+		{
+			SceneModifier->SimPreTick(ScenePtr);
+		}
+
+		// Handle legacy modifiers
+		Solver->OnEvolutionPreIntegrate(Dt);
+
+		Context.SetSimPhase(EContextSimPhase::None);
+	}
+
+	void FRigidSceneAsyncPT::OnPostTick(FReal Dt)
+	{
+		using namespace UE::Physics;
+
+		CheckSimContext();
+
+		Context.SetSimPhase(EContextSimPhase::PostTick);
+
+		TRigidScenePtr<FSimContextRW> ScenePtr = Private::FRigidAccessor::MakeRigidScenePtr(Context, this);
+		for (IRigidSceneModifier* SceneModifier : SceneModifiers)
+		{
+			SceneModifier->SimPostTick(ScenePtr);
+		}
+
+		// Handle legacy modifiers
+		Solver->OnEvolutionPostIntegrate(Dt);
+
+		Context.SetSimPhase(EContextSimPhase::None);
+	}
+
+}
+
+#endif // UE_RIGIDPHYSICS_API_ENABLED
\ No newline at end of file
diff --git a/Engine/Plugins/Experimental/ChaosRigidPhysicsAsync/Source/ChaosRigidPhysicsAsync/Public/ChaosRigidPhysicsAsync/RigidBodyAsync.h b/Engine/Plugins/Experimental/ChaosRigidPhysicsAsync/Source/ChaosRigidPhysicsAsync/Public/ChaosRigidPhysicsAsync/RigidBodyAsync.h
new file mode 100644
--- /dev/null
+++ b/Engine/Plugins/Experimental/ChaosRigidPhysicsAsync/Source/ChaosRigidPhysicsAsync/Public/ChaosRigidPhysicsAsync/RigidBodyAsync.h
@@ -0,0 +1,225 @@
+// Copyright Epic Games, Inc. All Rights Reserved.
+#pragma once 
+
+#include "ChaosRigidPhysicsAsync/RigidFwdAsync.h"
+
+#if UE_RIGIDPHYSICS_API_ENABLED
+
+#include "Logging/StructuredLog.h"
+#include "Math/Bounds.h"
+#include "PBDRigidsSolver.h"
+#include "Physics/Experimental/PhysScene_Chaos.h"
+#include "PhysicsProxy/SingleParticlePhysicsProxy.h"
+#include "RigidPhysics/RigidAccessor.h"
+#include "RigidPhysics/RigidBody.h"
+
+namespace Chaos::Rigids::Async
+{
+	// Utility class to standardize conversion from FRigidBodyID (uint64) to 
+	// ints for BodyIndex and Epoch, as used by the Async Body Container.
+	class UE_INTERNAL FRigidBodyIdAsync
+	{
+	public:
+		FRigidBodyIdAsync(int32 InBodyIndex, int32 InEpoch)
+			: BodyIndex(InBodyIndex)
+			, Epoch(InEpoch)
+		{
+		}
+
+		explicit FRigidBodyIdAsync(const UE::Physics::FRigidBodyId& BodyId)
+			: BodyIndex(INDEX_NONE)
+			, Epoch(INDEX_NONE)
+		{
+			uint64 Id = Private::FRigidAccessor::GetRigidBodyIdInt(BodyId);
+			BodyIndex = (int32)(Id & 0xFFFFFFFF);
+			Epoch = (int32)((Id >> 32) & 0xFFFFFFFF);
+		}
+
+		UE::Physics::FRigidBodyId AsRigidBodyId() const
+		{
+			uint64 Id = ((uint64)BodyIndex | ((uint64)Epoch << 32));
+			return UE::Physics::FRigidBodyId(Id);
+		}
+
+		int32 GetBodyIndex() const
+		{
+			return BodyIndex;
+		}
+
+		int32 GetEpoch() const
+		{
+			return Epoch;
+		}
+
+	private:
+		int32 BodyIndex;
+		int32 Epoch;
+	};
+
+	//
+	//
+	// RIGID BODY GT/PT
+	//
+	// Implemented in terms of the existing FSingleParticlePhysicsProxy.
+	// GT and PT bodies implement the IRigidBody API and forward the functions
+	// to the appropriate Proxy function. 
+	//
+	// NOTE: The current Proxy implementation does not sync back changes made in 
+	// simulation callbacks if those changes are to data that is not part of the 
+	// simulation output (e.g., SetGravityGroup, and most things that are not
+	// position, velocity, etc)
+	//
+	//
+
+	class UE_INTERNAL FRigidBodyAsyncGT : public UE::Physics::IRigidBody
+	{
+	private:
+		inline FRigidBodyHandle_External& GetProxyAPI()
+		{ 
+			return Proxy->GetGameThreadAPI();
+		}
+		
+		inline const FRigidBodyHandle_External& GetProxyAPI() const
+		{
+			return Proxy->GetGameThreadAPI();
+		}
+
+	public:
+
+		UE_INTERNAL CHAOSRIGIDPHYSICSASYNC_API FRigidBodyAsyncGT(
+			FRigidBodyContainerAsyncGT* InContainer, 
+			FSingleParticlePhysicsProxy* InProxy, 
+			const FRigidBodyIdAsync& InBodyId);
+
+		virtual ~FRigidBodyAsyncGT();
+
+		//
+		// IRIGIDBODY BEGIN
+		//
+
+		virtual const char* GetTypeName() const override final;
+		virtual UE::Physics::FRigidBodyId GetBodyId() const override final;
+		virtual UE::Physics::FRigidBodyHandle GetBodyHandle() const override final;
+		virtual bool IsValid() const override final;
+		virtual bool IsInScene() const override final;
+		virtual bool IsStatic() const override final;
+		virtual bool IsKinematic() const override final;
+		virtual bool IsDynamic() const override final;
+		virtual bool IsSleeping() const override final;
+		virtual void InitTransform(const FTransform& InTransform) override final;
+		virtual void UpdateTransform(const FTransform& InTransform) override final;
+		virtual void SetKinematicTarget(const FTransform& InTransform) override final;
+		virtual FTransform GetTransform() const override final;
+		virtual FBounds3d GetBounds() const override final;
+		virtual double GetMass() const override final;
+		virtual void SetMass(double InMass) override final;
+		//virtual void SetCCDEnabled(bool bInEnabled) override final {  }
+		//virtual void SetMACDEnabled(bool bInEnabled) override final {  }
+		//virtual void SetSmoothEdgeCollisionsEnabled(bool bInEnabled) override final {  }
+		//virtual void SetMaxLinearVelocity(float InMax) override final {  }
+		//virtual void SetMaxAngularVelocity(float InMax) override final {  }
+		//virtual void SetInertiaConditioningEnabled(bool bInEnabled) override final {  }
+		//virtual void SetGyroscopicTorqueEnabled(bool bInEnabled) override final {  }
+		//virtual void SetGravityGroup(int32 InId) override final {  }
+		virtual void SetGravityScale(float InScale) override final;
+		//virtual void SetWakeEventsEnabled(bool bInEnabled) override final {  }
+		//virtual void SetNumPositionIterations(int32 InNum) override final {  }
+		//virtual void SetNumVelocityIterations(int32 InNum) override final {  }
+		//virtual void SetNumProjectionIterations(int32 InNum) override final {  }
+		//virtual void SetIsSleeping(bool bInSleeping) override final {  }
+
+		//
+		// IRIGIDBODY END
+		//
+
+		UE_INTERNAL FRigidBodyIdAsync GetBodyIdAsync() const
+		{
+			return BodyId;
+		}
+
+		UE_INTERNAL const FSingleParticlePhysicsProxy* GetProxy() const
+		{
+			return Proxy;
+		}
+		
+		UE_INTERNAL FSingleParticlePhysicsProxy* GetProxy()
+		{
+			return Proxy;
+		}
+
+	private:
+		friend class FRigidSceneAsyncGT;
+		friend class FRigidBodyContainerAsyncGT;
+
+		FRigidSceneAsyncGT* GetSceneAsync() const;
+
+		FRigidBodyContainerAsyncGT* Container;
+		FSingleParticlePhysicsProxy* Proxy;
+		
+		FRigidBodyIdAsync BodyId;
+	};
+
+
+	class UE_INTERNAL FRigidBodyAsyncPT : public UE::Physics::IRigidBody
+	{
+	private:
+		inline FRigidBodyHandle_Internal* GetProxyAPI()
+		{
+			return Proxy->GetPhysicsThreadAPI();
+		}
+		
+		inline const FRigidBodyHandle_Internal* GetProxyAPI() const
+		{
+			return Proxy->GetPhysicsThreadAPI();
+		}
+
+	public:
+		UE_INTERNAL CHAOSRIGIDPHYSICSASYNC_API FRigidBodyAsyncPT(
+			FRigidBodyContainerAsyncPT* InContainer, 
+			FSingleParticlePhysicsProxy* InProxy, 
+			const FRigidBodyIdAsync& InBodyId);
+
+		virtual ~FRigidBodyAsyncPT();
+
+		//
+		// IRIGIDBODY BEGIN
+		//
+
+		virtual const char* GetTypeName() const override final;
+		virtual UE::Physics::FRigidBodyId GetBodyId() const override final;
+		virtual UE::Physics::FRigidBodyHandle GetBodyHandle() const override final;
+		virtual bool IsValid() const override final;
+		virtual bool IsInScene() const override final;
+		virtual bool IsStatic() const override final;
+		virtual bool IsKinematic() const override final;
+		virtual bool IsDynamic() const override final;
+		virtual bool IsSleeping() const override final;
+		virtual void InitTransform(const FTransform& InTransform) override final;
+		virtual void UpdateTransform(const FTransform& InTransform) override final;
+		virtual void SetKinematicTarget(const FTransform& InTransform) override final;
+		virtual FTransform GetTransform() const override final;
+		virtual FBounds3d GetBounds() const override final;
+		virtual double GetMass() const override final;
+		virtual void SetMass(double InMass) override final;
+
+		//
+		// IRIGIDBODY END
+		//
+
+		UE_INTERNAL FRigidBodyIdAsync GetBodyIdAsync() const
+		{
+			return BodyId;
+		}
+
+	private:
+		friend class FRigidBodyContainerAsyncPT;
+
+		FRigidBodyContainerAsyncPT* Container;
+		FSingleParticlePhysicsProxy* Proxy;
+
+		FRigidBodyIdAsync BodyId;
+	};
+
+}
+
+#endif // UE_RIGIDPHYSICS_API_ENABLED
\ No newline at end of file
diff --git a/Engine/Plugins/Experimental/ChaosRigidPhysicsAsync/Source/ChaosRigidPhysicsAsync/Public/ChaosRigidPhysicsAsync/RigidBodyContainerAsync.h b/Engine/Plugins/Experimental/ChaosRigidPhysicsAsync/Source/ChaosRigidPhysicsAsync/Public/ChaosRigidPhysicsAsync/RigidBodyContainerAsync.h
new file mode 100644
--- /dev/null
+++ b/Engine/Plugins/Experimental/ChaosRigidPhysicsAsync/Source/ChaosRigidPhysicsAsync/Public/ChaosRigidPhysicsAsync/RigidBodyContainerAsync.h
@@ -0,0 +1,137 @@
+// Copyright Epic Games, Inc. All Rights Reserved.
+#pragma once
+
+#include "ChaosRigidPhysicsAsync/RigidFwdAsync.h"
+
+#if UE_RIGIDPHYSICS_API_ENABLED
+
+#include "Async/TaskGraphInterfaces.h"
+#include "Chaos/ObjectPool.h"
+#include "ChaosRigidPhysicsAsync/RigidBodyAsync.h"
+#include "ChaosSolversModule.h"
+#include "Containers/Map.h"
+#include "PBDRigidsSolver.h"
+#include "PhysicsProxy/SingleParticlePhysicsProxy.h"
+#include "RigidPhysics/RigidAccessor.h"
+#include "RigidPhysics/RigidBody.h"
+
+namespace Chaos::Rigids::Async
+{
+	//
+	//
+	// RIGID BODY CONTAINER GT/PT
+	//
+	// The GT and PT containers hold the set of GT and PT particles respectively.
+	//
+	//
+
+	class UE_INTERNAL FRigidBodyContainerAsyncGT : public UE::Physics::IRigidBodyContainer
+	{
+	public:
+		UE_INTERNAL CHAOSRIGIDPHYSICSASYNC_API static UE::Physics::FRigidBodyContainerTypeId StaticRigidBodyContainerTypeId();
+
+		UE_INTERNAL CHAOSRIGIDPHYSICSASYNC_API FRigidBodyContainerAsyncGT(const UE::Physics::FRigidBodyContainerId& InContainerId);
+
+		virtual ~FRigidBodyContainerAsyncGT();
+
+		//
+		// IRIGIDBODYCONTAINER BEGIN
+		//
+
+		virtual UE::Physics::IRigidScene* GetScene() const override final;
+		virtual UE::Physics::FRigidBodyContainerId GetContainerId() const override final;
+		virtual FRigidBodyAsyncGT* GetRigidBody(const UE::Physics::FRigidBodyId& BodyId) const override final;
+
+		//
+		// IRIGIDBODYCONTAINER END
+		//
+
+		UE_INTERNAL CHAOSRIGIDPHYSICSASYNC_API UE::Physics::FRigidBodyContainerHandle GetContainerHandle() const;
+
+		UE_INTERNAL CHAOSRIGIDPHYSICSASYNC_API void Init(FRigidSceneAsyncGT* InScene, FPBDRigidsSolver* InSolver);
+
+		UE_INTERNAL CHAOSRIGIDPHYSICSASYNC_API FRigidBodyAsyncGT* CreateRigidBody(const UE::Physics::FRigidNamePtr& InName, UE::Physics::ERigidMovementType InMovementType);
+
+		UE_INTERNAL CHAOSRIGIDPHYSICSASYNC_API FRigidBodyAsyncGT* CreateRigidBodyForProxy(FSingleParticlePhysicsProxy* InProxy);
+
+		UE_INTERNAL CHAOSRIGIDPHYSICSASYNC_API void DestroyRigidBody(const UE::Physics::FRigidBodyId& BodyId);
+
+		UE_INTERNAL CHAOSRIGIDPHYSICSASYNC_API void AddRigidBodyToScene(const UE::Physics::FRigidBodyId& BodyId);
+
+		UE_INTERNAL CHAOSRIGIDPHYSICSASYNC_API void RemoveRigidBodyFromScene(const UE::Physics::FRigidBodyId& BodyId);
+
+		UE_INTERNAL CHAOSRIGIDPHYSICSASYNC_API void OnRigidBodyTransformUpdated(FRigidBodyAsyncGT* Body);
+
+		UE_INTERNAL CHAOSRIGIDPHYSICSASYNC_API FRigidSceneAsyncGT* GetSceneAsync() const;
+
+	private:
+
+		// Store the Epoch with the body to validate handles
+		struct FRigidBodyArrayEntry
+		{
+			FRigidBodyAsyncGT* Body;
+			int32 Epoch;
+		};
+
+		FRigidBodyAsyncGT* AllocBody(FSingleParticlePhysicsProxy* InProxy);
+		void FreeBody(FRigidBodyAsyncGT* BodyGT);
+		void ReserveBodyArray(int32 MaxBodies);
+
+		FRigidSceneAsyncGT* Scene;
+		UE::Physics::FRigidBodyContainerId ContainerId;
+		FPBDRigidsSolver* Solver;
+
+		// Body allocator
+		TObjectPool<FRigidBodyAsyncGT> BodyPool;
+
+		// BodyId->Body mapping
+		TArray<FRigidBodyArrayEntry> Bodies;
+		TArray<int32> BodiesFreeIndices;
+		int64 NextBodyIdIndex;
+	};
+
+
+	class UE_INTERNAL FRigidBodyContainerAsyncPT : public UE::Physics::IRigidBodyContainer
+	{
+	public:
+		UE_INTERNAL static UE::Physics::FRigidBodyContainerTypeId StaticRigidBodyContainerTypeId();
+
+		UE_INTERNAL FRigidBodyContainerAsyncPT(const UE::Physics::FRigidBodyContainerId& InContainerId);
+
+		virtual ~FRigidBodyContainerAsyncPT();
+
+		//
+		// IRIGIDBODYCONTAINER BEGIN
+		//
+
+		virtual UE::Physics::IRigidScene* GetScene() const override final;
+		virtual UE::Physics::FRigidBodyContainerId GetContainerId() const override final;
+		virtual FRigidBodyAsyncPT* GetRigidBody(const UE::Physics::FRigidBodyId& BodyId) const override final;
+
+		//
+		// IRIGIDBODYCONTAINER END
+		//
+
+		UE_INTERNAL void Init(FRigidSceneAsyncPT* InScene, FPBDRigidsSolver* InSolver);
+
+		UE_INTERNAL UE::Physics::FRigidBodyContainerHandle GetContainerHandle() const;
+
+		UE_INTERNAL void CreateRigidBody(const FRigidBodyIdAsync& InBodyId, FSingleParticlePhysicsProxy* InProxy);
+
+		UE_INTERNAL void DestroyRigidBody(const UE::Physics::FRigidBodyId& BodyId);
+
+		UE_INTERNAL void AddRigidBodyToScene(const UE::Physics::FRigidBodyId& BodyId);
+
+		UE_INTERNAL void RemoveRigidBodyFromScene(const UE::Physics::FRigidBodyId& BodyId);
+
+	private:
+		FRigidSceneAsyncPT* Scene;
+		UE::Physics::FRigidBodyContainerId ContainerId;
+		FPBDRigidsSolver* Solver;
+		TObjectPool<FRigidBodyAsyncPT> BodyPool;
+		TMap<UE::Physics::FRigidBodyId, FRigidBodyAsyncPT*> IdToBodyMap;
+	};
+
+}
+
+#endif // UE_RIGIDPHYSICS_API_ENABLED
\ No newline at end of file
diff --git a/Engine/Plugins/Experimental/ChaosRigidPhysicsAsync/Source/ChaosRigidPhysicsAsync/Public/ChaosRigidPhysicsAsync/RigidFactoryAsync.h b/Engine/Plugins/Experimental/ChaosRigidPhysicsAsync/Source/ChaosRigidPhysicsAsync/Public/ChaosRigidPhysicsAsync/RigidFactoryAsync.h
new file mode 100644
--- /dev/null
+++ b/Engine/Plugins/Experimental/ChaosRigidPhysicsAsync/Source/ChaosRigidPhysicsAsync/Public/ChaosRigidPhysicsAsync/RigidFactoryAsync.h
@@ -0,0 +1,60 @@
+// Copyright Epic Games, Inc. All Rights Reserved.
+#pragma once
+
+#include "RigidPhysics/RigidFwd.h"
+
+#if UE_RIGIDPHYSICS_API_ENABLED
+
+#include "Templates/SharedPointer.h"
+#include "Templates/UniquePtr.h"
+
+namespace Chaos
+{
+	class FPBDRigidsSolver;
+	class FSingleParticlePhysicsProxy;
+}
+
+namespace Chaos::Rigids::Async::Private
+{
+	using FRigidAccessor = UE::Physics::Private::FRigidAccessor;
+	class FRigidAccessorAsync;
+}
+
+namespace Chaos::Rigids::Async
+{
+	class FRigidSceneSettingsAsync
+	{
+	public:
+		FName SceneName = NAME_None;
+		float AsyncDt = (1.0f / 60.0f);
+	};
+
+	class FRigidFactoryAsync
+	{
+	public:
+		CHAOSRIGIDPHYSICSASYNC_API UE::Physics::FRigidSceneHandle CreateRigidScene(const FRigidSceneSettingsAsync& InSettings);
+		CHAOSRIGIDPHYSICSASYNC_API void DestroyRigidScene(const UE::Physics::FRigidSceneHandle& InSceneHandle);
+
+	private:
+		// TODO_CHAOSAPI: Temp until solver is created through new API
+		CHAOSRIGIDPHYSICSASYNC_API UE::Physics::FRigidSceneHandle CreateRigidSceneFromSolver(Chaos::FPBDRigidsSolver* InSolver);
+
+		friend class Private::FRigidAccessorAsync;
+		TArray<UE::Physics::FRigidSceneSharedPtr> Scenes;
+	};
+}
+
+// TODO_CHAOSAPI: Temp until solver is created through new API
+namespace Chaos::Rigids::Async::Private
+{
+	// Create a scene for and externally-created solver (simulation proxy)
+	UE_INTERNAL CHAOSRIGIDPHYSICSASYNC_API UE::Physics::FRigidSceneHandle CreateSceneFromSolver(FRigidFactoryAsync& InFactory, FPBDRigidsSolver* InSolver);
+
+	// Create a body handle for an externally-created single particle proxy
+	UE_INTERNAL CHAOSRIGIDPHYSICSASYNC_API UE::Physics::FRigidBodyHandle CreateBodyFromProxy(const UE::Physics::FRigidSceneHandle& InSceneHandle, FSingleParticlePhysicsProxy* InProxy);
+
+	// Get the proxy from the body handle
+	UE_INTERNAL CHAOSRIGIDPHYSICSASYNC_API FSingleParticlePhysicsProxy* GetProxyFromBody(const UE::Physics::FRigidBodyHandle& InBodyHandle);
+}
+
+#endif // UE_RIGIDPHYSICS_API_ENABLED
\ No newline at end of file
diff --git a/Engine/Plugins/Experimental/ChaosRigidPhysicsAsync/Source/ChaosRigidPhysicsAsync/Public/ChaosRigidPhysicsAsync/RigidFwdAsync.h b/Engine/Plugins/Experimental/ChaosRigidPhysicsAsync/Source/ChaosRigidPhysicsAsync/Public/ChaosRigidPhysicsAsync/RigidFwdAsync.h
new file mode 100644
--- /dev/null
+++ b/Engine/Plugins/Experimental/ChaosRigidPhysicsAsync/Source/ChaosRigidPhysicsAsync/Public/ChaosRigidPhysicsAsync/RigidFwdAsync.h
@@ -0,0 +1,32 @@
+// Copyright Epic Games, Inc. All Rights Reserved.
+#pragma once 
+
+#include "RigidPhysics/RigidFwd.h"
+
+#if UE_RIGIDPHYSICS_API_ENABLED
+
+#include "Framework/Threading.h"
+
+namespace Chaos::Rigids::Async
+{
+	class UE_INTERNAL FRigidBodyAsyncGT;
+	class UE_INTERNAL FRigidBodyAsyncPT;
+	class UE_INTERNAL FRigidBodyContainerAsyncGT;
+	class UE_INTERNAL FRigidBodyContainerAsyncPT;
+	class UE_INTERNAL FRigidSceneAsyncGT;
+	class UE_INTERNAL FRigidSceneAsyncPT;
+}
+
+namespace Chaos::Rigids::Async::Private
+{
+	using UE_INTERNAL FRigidAccessor = UE::Physics::Private::FRigidAccessor;
+}
+
+// Enable asserts if we have a context to check (not TEST or SHIPPING builds)
+#if PHYSICS_THREAD_CONTEXT && DO_CHECK
+#define CHAOSRIGIDPHYSICSASYNC_CHECK(X) check(X)
+#else
+#define CHAOSRIGIDPHYSICSASYNC_CHECK(X)
+#endif
+
+#endif // UE_RIGIDPHYSICS_API_ENABLED
\ No newline at end of file
diff --git a/Engine/Plugins/Experimental/ChaosRigidPhysicsAsync/Source/ChaosRigidPhysicsAsync/Public/ChaosRigidPhysicsAsync/RigidSceneAsync.h b/Engine/Plugins/Experimental/ChaosRigidPhysicsAsync/Source/ChaosRigidPhysicsAsync/Public/ChaosRigidPhysicsAsync/RigidSceneAsync.h
new file mode 100644
--- /dev/null
+++ b/Engine/Plugins/Experimental/ChaosRigidPhysicsAsync/Source/ChaosRigidPhysicsAsync/Public/ChaosRigidPhysicsAsync/RigidSceneAsync.h
@@ -0,0 +1,149 @@
+// Copyright Epic Games, Inc. All Rights Reserved.
+#pragma once
+
+#include "ChaosRigidPhysicsAsync/RigidFwdAsync.h"
+
+#if UE_RIGIDPHYSICS_API_ENABLED
+
+#include "Async/TaskGraphInterfaces.h"
+#include "Chaos/ObjectPool.h"
+#include "ChaosRigidPhysicsAsync/RigidBodyAsync.h"
+#include "ChaosRigidPhysicsAsync/RigidBodyContainerAsync.h"
+#include "ChaosRigidPhysicsAsync/RigidFactoryAsync.h"
+#include "ChaosSolversModule.h"
+#include "Containers/Map.h"
+#include "Logging/StructuredLog.h"
+#include "PBDRigidsSolver.h"
+#include "Physics/Experimental/PhysScene_Chaos.h"
+#include "PhysicsProxy/SingleParticlePhysicsProxy.h"
+#include "RigidPhysics/RigidAccessor.h"
+#include "RigidPhysics/RigidBody.h"
+
+namespace Chaos::Rigids::Async
+{
+	//
+	//
+	// RIGID SCENE GT/PT
+	//
+	// The GT/PT scenes are implemented as a wrapper around Chaos::FPBDRigidsSolver
+	// (the equivalen of a "SceneProxy"). GT and PT Scenes reference the same solver,
+	// and each implements IRigidScene, forwarding the functions to the appropriate
+	// solver function for the context.
+	//
+	//
+
+	class UE_INTERNAL FRigidSceneAsyncGT : public UE::Physics::IRigidScene, public TSharedFromThis<FRigidSceneAsyncGT, ESPMode::ThreadSafe>
+	{
+	public:
+		UE_INTERNAL CHAOSRIGIDPHYSICSASYNC_API FRigidSceneAsyncGT(const FRigidSceneSettingsAsync& InSceneSettings);
+		UE_INTERNAL CHAOSRIGIDPHYSICSASYNC_API FRigidSceneAsyncGT(FPBDRigidsSolver* InSolver);
+		virtual ~FRigidSceneAsyncGT();
+
+		//
+		// IRIGIDSCENE BEGIN
+		//
+	public:
+		virtual void Lock(UE::Physics::EContextLockType InLockType) override final;
+		virtual void Unlock(UE::Physics::EContextLockType InLockType) override final;
+		virtual FRigidBodyAsyncGT* CreateRigidBody(const UE::Physics::FRigidNamePtr& InName, UE::Physics::ERigidMovementType InMovementType) override final;
+		virtual void DestroyRigidBody(const UE::Physics::FRigidBodyHandle& BodyHandle) override final;
+		virtual void RegisterRigidBody(const UE::Physics::FRigidBodyHandle& BodyHandle) override final;
+		virtual void UnregisterRigidBody(const UE::Physics::FRigidBodyHandle& BodyHandle) override final;
+		virtual FRigidBodyContainerAsyncGT* GetRigidBodyContainer(const UE::Physics::FRigidBodyContainerId& InBodyContainerId) override final;
+		virtual void RegisterModifier(UE::Physics::IRigidSceneModifier* Modifier) override final;
+		virtual void UnregisterModifier(UE::Physics::IRigidSceneModifier* Modifier) override final;
+		virtual bool Query(const UE::Physics::FRigidQueryDesc& QueryDesc, TArray<UE::Physics::FRigidQueryHit>& OutHits) override final;
+		virtual void Tick(float Dt) override final;
+
+	protected:
+		virtual const char* GetTypeName() const override final;
+
+		//
+		// IRIGIDSCENE END
+		//
+
+	public:
+		UE_INTERNAL CHAOSRIGIDPHYSICSASYNC_API UE::Physics::FRigidSceneHandle GetSceneHandle() const;
+		UE_INTERNAL CHAOSRIGIDPHYSICSASYNC_API void Startup();
+		UE_INTERNAL CHAOSRIGIDPHYSICSASYNC_API void Shutdown();
+		UE_INTERNAL CHAOSRIGIDPHYSICSASYNC_API void OnRigidBodyTransformUpdated(FRigidBodyAsyncGT* Body);
+
+	private:
+		UE::Physics::FRigidSceneHandle GetHandle();
+		void CheckGameReadLocked() const;
+		void CheckGameWriteLocked() const;
+		void OnPreSimulate(FReal Dt);
+		void OnPostSimulate(FReal Dt);
+		void UpdateRigidBodyInAccelerationStructure(FRigidBodyAsyncGT* Body);
+
+		FRigidSceneSettingsAsync SceneSettings;
+
+		FPBDRigidsSolver* Solver;
+		ISpatialAccelerationCollection<FAccelerationStructureHandle, FReal, 3>* SpatialAcceleration;
+
+		UE::Physics::FRigidSceneHandle SceneHandle;
+		FRigidBodyContainerAsyncGT SingleBodyContainer;
+
+		TArray<UE::Physics::IRigidSceneModifier*> SceneModifiers;
+
+		TUniquePtr<FRigidSceneAsyncPT> ScenePT;
+
+		// TODO_CHAOSAPI: Temp until solver is created through new API
+		bool bOwnsSolver;
+	};
+
+
+	class UE_INTERNAL FRigidSceneAsyncPT : public UE::Physics::IRigidScene
+	{
+	public:
+		UE_INTERNAL CHAOSRIGIDPHYSICSASYNC_API FRigidSceneAsyncPT(const UE::Physics::FRigidBodyContainerId& InSingleBodyContainerId);
+
+		//
+		// IRIGIDSCENE BEGIN
+		//
+
+	public:
+		virtual void Lock(UE::Physics::EContextLockType InLockType) override final;
+		virtual void Unlock(UE::Physics::EContextLockType InLockType) override final;
+		virtual FRigidBodyAsyncPT* CreateRigidBody(const UE::Physics::FRigidNamePtr& InName, UE::Physics::ERigidMovementType InMovementType) override final;
+		virtual void DestroyRigidBody(const UE::Physics::FRigidBodyHandle& BodyHandle) override final;
+		virtual void RegisterRigidBody(const UE::Physics::FRigidBodyHandle& BodyHandle) override final;
+		virtual void UnregisterRigidBody(const UE::Physics::FRigidBodyHandle& BodyHandle) override final;
+		virtual FRigidBodyContainerAsyncPT* GetRigidBodyContainer(const UE::Physics::FRigidBodyContainerId& InBodyContainerId) override final;
+		virtual void RegisterModifier(UE::Physics::IRigidSceneModifier* Modifier) override final;
+		virtual void UnregisterModifier(UE::Physics::IRigidSceneModifier* Modifier) override final;
+		virtual bool Query(const UE::Physics::FRigidQueryDesc& QueryDesc, TArray<UE::Physics::FRigidQueryHit>& OutHits) override final;
+		virtual void Tick(float Dt) override final;
+
+	protected:
+		virtual const char* GetTypeName() const override final;
+
+		//
+		// IRIGIDSCENE END
+		//
+
+	public:
+		UE_INTERNAL CHAOSRIGIDPHYSICSASYNC_API UE::Physics::FRigidSceneHandle GetSceneHandle() const;
+		UE_INTERNAL CHAOSRIGIDPHYSICSASYNC_API void Startup(UE::Physics::FRigidSceneHandle InSceneHandle, FPBDRigidsSolver* InSolver);
+		UE_INTERNAL CHAOSRIGIDPHYSICSASYNC_API void Shutdown();
+		UE_INTERNAL CHAOSRIGIDPHYSICSASYNC_API void CreateRigidBody(const FRigidBodyIdAsync& InBodyId, FSingleParticlePhysicsProxy* Proxy);
+		UE_INTERNAL CHAOSRIGIDPHYSICSASYNC_API void DestroyRigidBody(const UE::Physics::FRigidBodyId& InBodyId);
+		UE_INTERNAL CHAOSRIGIDPHYSICSASYNC_API void RegisterRigidBody(const UE::Physics::FRigidBodyId& InBodyId);
+		UE_INTERNAL CHAOSRIGIDPHYSICSASYNC_API void UnregisterRigidBody(const UE::Physics::FRigidBodyId& InBodyId);
+
+	private:
+		void CheckGameContext() const;
+		void CheckSimContext() const;
+
+		void OnPreTick(FReal Dt);
+		void OnPostTick(FReal Dt);
+
+		UE::Physics::FRigidSceneHandle SceneHandle;
+		FPBDRigidsSolver* Solver;
+		FRigidBodyContainerAsyncPT SingleBodyContainer;
+		TArray<UE::Physics::IRigidSceneModifier*> SceneModifiers;
+		UE::Physics::FSimContextRW Context;
+	};
+}
+
+#endif // UE_RIGIDPHYSICS_API_ENABLED
diff --git a/Engine/Plugins/Experimental/ChaosRigidPhysicsAsync/Tests/ChaosRigidPhysicsAsyncTests.Build.cs b/Engine/Plugins/Experimental/ChaosRigidPhysicsAsync/Tests/ChaosRigidPhysicsAsyncTests.Build.cs
new file mode 100644
--- /dev/null
+++ b/Engine/Plugins/Experimental/ChaosRigidPhysicsAsync/Tests/ChaosRigidPhysicsAsyncTests.Build.cs
@@ -0,0 +1,60 @@
+// Copyright Epic Games, Inc. All Rights Reserved.
+
+using UnrealBuildTool;
+
+public class ChaosRigidPhysicsAsyncTests : TestModuleRules
+{
+	static ChaosRigidPhysicsAsyncTests()
+	{
+		if (InTestMode)
+		{
+			TestMetadata = new Metadata();
+			TestMetadata.TestName = "ChaosRigidPhysicsAsyncTests";
+			TestMetadata.TestShortName = "ChaosRigidPhysicsAsyncTests";
+			TestMetadata.ReportType = "xml";
+			TestMetadata.SupportedPlatforms.Add(UnrealTargetPlatform.Linux);
+			TestMetadata.SupportedPlatforms.Add(UnrealTargetPlatform.Mac);
+			TestMetadata.SupportedPlatforms.Add(UnrealTargetPlatform.Android);
+			TestMetadata.SupportedPlatforms.Add(UnrealTargetPlatform.IOS);
+
+			string PlatformCompilationArgs;
+			foreach (var Platform in UnrealTargetPlatform.GetValidPlatforms())
+			{
+				if (Platform == UnrealTargetPlatform.Android)
+				{
+					PlatformCompilationArgs = "-allmodules -architectures=arm64";
+				}
+				else
+				{
+					PlatformCompilationArgs = "-allmodules";
+				}
+				TestMetadata.PlatformCompilationExtraArgs.Add(Platform, PlatformCompilationArgs);
+			}
+
+			// Platform-specific tags
+			TestMetadata.PlatformTags.Add(UnrealTargetPlatform.Linux, "~[.]~[Slow]");
+			TestMetadata.PlatformTags.Add(UnrealTargetPlatform.Android, "~[Perf]~[Slow]~[AndroidSkip]");
+
+			// Allow Android run for this test
+			// Will remove Android from PlatformsRunUnsupported as more diverse types of tests can run on this platform
+			TestMetadata.PlatformsRunUnsupported.Remove(UnrealTargetPlatform.Android);
+		}
+	}
+
+	public ChaosRigidPhysicsAsyncTests(ReadOnlyTargetRules Target) : base(Target)
+	{
+		PrivateDependencyModuleNames.AddRange(
+			new string[] {
+                "Core",
+                "CoreUObject",
+                "Engine",
+				"ChaosRigidPhysicsAsync",
+				"ChaosTestHarness",
+                "RigidPhysics",
+		});
+
+		SetupModulePhysicsSupport(Target);
+
+		bWarningsAsErrors = true;
+	}
+}
\ No newline at end of file
diff --git a/Engine/Plugins/Experimental/ChaosRigidPhysicsAsync/Tests/ChaosRigidPhysicsAsyncTests.Target.cs b/Engine/Plugins/Experimental/ChaosRigidPhysicsAsync/Tests/ChaosRigidPhysicsAsyncTests.Target.cs
new file mode 100644
--- /dev/null
+++ b/Engine/Plugins/Experimental/ChaosRigidPhysicsAsync/Tests/ChaosRigidPhysicsAsyncTests.Target.cs
@@ -0,0 +1,27 @@
+// Copyright Epic Games, Inc. All Rights Reserved.
+
+using UnrealBuildTool;
+
+[SupportedPlatforms(UnrealPlatformClass.All)]
+ public class ChaosRigidPhysicsAsyncTestsTarget : TestTargetRules
+ {
+     public ChaosRigidPhysicsAsyncTestsTarget(TargetInfo Target) : base(Target)
+     {
+		bCompileAgainstEngine = true;
+		bCompileAgainstCoreUObject = true;
+		bMockEngineDefaults = true;
+		bUsePlatformFileStub = true;
+
+		bWithLowLevelTestsOverride = true;
+
+		// Enable the new API
+		GlobalDefinitions.Add("UE_RIGIDPHYSICS_API_ENABLED=1");
+
+		// Enable constraint handle debugging in Chaos
+		GlobalDefinitions.Add("CHAOS_CONSTRAINTHANDLE_DEBUG_ENABLED=1");
+
+		GlobalDefinitions.Add("CATCH_CONFIG_ENABLE_BENCHMARKING=1");
+
+		bWarningsAsErrors = true;
+     }
+ }
\ No newline at end of file
diff --git a/Engine/Plugins/Experimental/ChaosRigidPhysicsAsync/Tests/Private/SceneLockBenchmarks.cpp b/Engine/Plugins/Experimental/ChaosRigidPhysicsAsync/Tests/Private/SceneLockBenchmarks.cpp
new file mode 100644
--- /dev/null
+++ b/Engine/Plugins/Experimental/ChaosRigidPhysicsAsync/Tests/Private/SceneLockBenchmarks.cpp
@@ -0,0 +1,238 @@
+// Copyright Epic Games, Inc. All Rights Reserved.
+
+#include "Chaos/LowLevelTest/ChaosTestHarness.h"
+#include "Chaos/LowLevelTest/ChaosTestScene.h"
+#include "ChaosRigidPhysicsAsync/RigidFactoryAsync.h"
+#include "PBDRigidsSolver.h"
+#include "PhysicsInterfaceDeclaresCore.h"
+#include "Physics/Experimental/PhysInterface_Chaos.h"
+#include "PhysicsProxy/SingleParticlePhysicsProxy.h"
+#include "RigidPhysics/RigidBody.h"
+
+
+namespace Chaos::LowLevelTest
+{
+
+	// This test is only intended to be run with an attached sampling profiler
+	// It effectively loops forever, so don't enable it for regular testing
+	TEST_CASE("ChaosAPIBenchmark", "[.LoopForever][ChaosRBPerBodyLock]")
+	{
+		using namespace Chaos;
+		using namespace Chaos::Rigids::Async;
+		using namespace UE::Physics;
+
+		int32 NumBodies = 10000;
+		FTransform Transform = FTransform::Identity;
+
+		FRigidFactoryAsync Factory;
+		FRigidSceneHandle SceneHandle = Factory.CreateRigidScene(FRigidSceneSettingsAsync());
+
+		// Create the bodies
+		TArray<FRigidBodyHandle> BodyHandles;
+
+		if (TRigidScenePtr<FGameContextRW> Scene = SceneHandle.LockRW())
+		{
+			BodyHandles.Reserve(NumBodies);
+			for (int32 BodyIndex = 0; BodyIndex < NumBodies; ++BodyIndex)
+			{
+				TRigidBodyPtr<FGameContextRW> Body = Scene->CreateRigidBody(MakeRigidName(TEXT("Body{0}"), BodyIndex), ERigidMovementType::Dynamic);
+				Scene->RegisterRigidBody(Body);
+
+				BodyHandles.Add(Body);
+			}
+		}
+
+		// Change transform to prevent early-out
+		for (int32 LoopIndex = 0; LoopIndex < 1000000; ++LoopIndex)
+		{
+			FVector Pos = Transform.GetTranslation() + FVector(100, 0, 0);
+			Transform.SetTranslation(Pos);
+
+			for (FRigidBodyHandle BodyHandle : BodyHandles)
+			{
+				if (TRigidBodyPtr<FGameContextRW> Body = BodyHandle.LockRW())
+				{
+					Body->InitTransform(Transform);
+				}
+			}
+		}
+
+		if (TRigidScenePtr<FGameContextRW> Scene = SceneHandle.LockRW())
+		{
+			for (FRigidBodyHandle BodyHandle : BodyHandles)
+			{
+				Scene->UnregisterRigidBody(BodyHandle);
+				Scene->DestroyRigidBody(BodyHandle);
+			}
+		}
+		BodyHandles.Reset();
+
+		Factory.DestroyRigidScene(SceneHandle);
+	}
+
+	// This test is only intended to be run with an attached sampling profiler
+	// It effectively loops forever, so don't enable it for regular testing
+	TEST_CASE("ChaosAPIBenchmark ProxyComparison", "[.LoopForever][ChaosPerBodyLock]")
+	{
+		using namespace Chaos;
+		using namespace UE::Physics;
+
+		int32 NumBodies = 10000;
+		FTransform Transform = FTransform::Identity;
+
+		FChaosTestScene Scene;
+
+		// Create the bodies
+		TArray<FSingleParticlePhysicsProxy*> Proxies;
+		Proxies.Reserve(NumBodies);
+		for (int32 BodyIndex = 0; BodyIndex < NumBodies; ++BodyIndex)
+		{
+			FSingleParticlePhysicsProxy* Proxy = Scene.AddDynamicBox(TEXT("Box"), Transform, FVector(50, 50, 50));
+			Proxies.Add(Proxy);
+		}
+
+		for (int32 LoopIndex = 0; LoopIndex < 1000000; ++LoopIndex)
+		{
+			FVector Pos = Transform.GetTranslation() + FVector(100, 0, 0);
+			Transform.SetTranslation(Pos);
+
+			for (FSingleParticlePhysicsProxy* Proxy : Proxies)
+			{
+				FPhysicsCommand::ExecuteWrite(Proxy, 
+					[Proxy, &Transform](const FPhysicsActorHandle& Actor)
+					{
+						FPhysicsInterface::SetGlobalPose_AssumesLocked(Proxy, Transform);
+					});
+			}
+		}
+
+		Proxies.Reset();
+	}
+
+
+
+	TEST_CASE("ChaosAPIBenchmark BodyHandle Lock", "[!benchmark][ChaosAPIBenchmark]")
+	{
+		using namespace Chaos;
+		using namespace Chaos::Rigids::Async;
+		using namespace UE::Physics;
+
+		int32 NumBodies = 10000;
+		FTransform Transform = FTransform::Identity;
+
+		FRigidFactoryAsync Factory;
+		FRigidSceneHandle SceneHandle = Factory.CreateRigidScene(FRigidSceneSettingsAsync());
+
+		// Create the bodies
+		TArray<FRigidBodyHandle> BodyHandles;
+
+		if (TRigidScenePtr<FGameContextRW> Scene = SceneHandle.LockRW())
+		{
+			BodyHandles.Reserve(NumBodies);
+			for (int32 BodyIndex = 0; BodyIndex < NumBodies; ++BodyIndex)
+			{
+				TRigidBodyPtr<FGameContextRW> Body = Scene->CreateRigidBody(MakeRigidName(TEXT("Body{0}"), BodyIndex), ERigidMovementType::Dynamic);
+				Scene->RegisterRigidBody(Body);
+
+				BodyHandles.Add(Body);
+			}
+		}
+
+		BENCHMARK("Per-Body Lock")
+		{
+			// Change transform to prevent early-out
+			Transform.SetTranslation(Transform.GetTranslation() + FVector(100, 0, 0));
+
+			for (FRigidBodyHandle BodyHandle : BodyHandles)
+			{
+				if (TRigidBodyPtr<FGameContextRW> Body = BodyHandle.LockRW())
+				{
+					Body->InitTransform(Transform);
+				}
+			}
+		};
+
+		BENCHMARK("Scene Lock")
+		{
+			// Change transform to prevent early-out
+			Transform.SetTranslation(Transform.GetTranslation() + FVector(100, 0, 0));
+
+			if (TRigidScenePtr<FGameContextRW> Scene = SceneHandle.LockRW())
+			{
+				for (FRigidBodyHandle BodyHandle : BodyHandles)
+				{
+					// Pin does not add a lock - we are already in a scene lock
+					if (TRigidBodyPtr<FGameContextRW> Body = BodyHandle.Pin(Scene))
+					{
+						Body->InitTransform(Transform);
+					}
+				}
+			}
+		};
+
+		if (TRigidScenePtr<FGameContextRW> Scene = SceneHandle.LockRW())
+		{
+			for (FRigidBodyHandle BodyHandle : BodyHandles)
+			{
+				Scene->UnregisterRigidBody(BodyHandle);
+				Scene->DestroyRigidBody(BodyHandle);
+			}
+		}
+		BodyHandles.Reset();
+
+		Factory.DestroyRigidScene(SceneHandle);
+	}
+
+	TEST_CASE("ChaosAPIBenchmark BodyHandle Lock ProxyComparison", "[!benchmark][ChaosAPIBenchmark]")
+	{
+		using namespace Chaos;
+		using namespace UE::Physics;
+
+		int32 NumBodies = 10000;
+		FTransform Transform = FTransform::Identity;
+
+		FChaosTestScene Scene;
+
+		// Create the bodies
+		TArray<FSingleParticlePhysicsProxy*> Proxies;
+		Proxies.Reserve(NumBodies);
+		for (int32 BodyIndex = 0; BodyIndex < NumBodies; ++BodyIndex)
+		{
+			FSingleParticlePhysicsProxy* Proxy = Scene.AddDynamicBox(TEXT("Box"), Transform, FVector(50,50,50));
+			Proxies.Add(Proxy);
+		}
+
+		BENCHMARK("Per-Body Lock")
+		{
+			// Change transform to prevent early-out
+			Transform.SetTranslation(Transform.GetTranslation() + FVector(100, 0, 0));
+
+			for (FSingleParticlePhysicsProxy* Proxy : Proxies)
+			{
+				FPhysicsCommand::ExecuteWrite(Proxy,
+					[Proxy, &Transform](const FPhysicsActorHandle& Actor)
+					{
+						FPhysicsInterface::SetGlobalPose_AssumesLocked(Proxy, Transform);
+					});
+			}
+		};
+
+		BENCHMARK("Scene Lock")
+		{
+			// Change transform to prevent early-out
+			Transform.SetTranslation(Transform.GetTranslation() + FVector(100, 0, 0));
+
+			FPhysicsCommand::ExecuteWrite(Proxies[0],
+				[Proxies, &Transform](const FPhysicsActorHandle& Actor)
+				{
+					for (FSingleParticlePhysicsProxy* Proxy : Proxies)
+					{
+						FPhysicsInterface::SetGlobalPose_AssumesLocked(Proxy, Transform);
+					}
+				});
+		};
+
+		Proxies.Reset();
+	}
+
+}
\ No newline at end of file
diff --git a/Engine/Plugins/Experimental/ChaosRigidPhysicsAsync/Tests/Private/SceneTests.cpp b/Engine/Plugins/Experimental/ChaosRigidPhysicsAsync/Tests/Private/SceneTests.cpp
new file mode 100644
--- /dev/null
+++ b/Engine/Plugins/Experimental/ChaosRigidPhysicsAsync/Tests/Private/SceneTests.cpp
@@ -0,0 +1,314 @@
+// Copyright Epic Games, Inc. All Rights Reserved.
+
+#include "Chaos/Collision/ContactPoint.h"
+#include "Chaos/LowLevelTest/ChaosTestHarness.h"
+#include "Chaos/LowLevelTest/ChaosTestScene.h"
+#include "ChaosRigidPhysicsAsync/RigidFactoryAsync.h"
+#include "RigidPhysics/RigidBody.h"
+#include "RigidPhysics/RigidModifier.h"
+
+namespace Chaos::LowLevelTest
+{
+	class FTestRigidSceneModifier : public UE::Physics::IRigidSceneModifier
+	{
+	public:
+		virtual void SimPreTick(UE::Physics::TRigidScenePtr<UE::Physics::FSimContextRW>& Context) override final
+		{
+			if (SimPreTickFunc)
+			{
+				SimPreTickFunc(Context);
+			}
+		}
+
+		TFunction<void(UE::Physics::TRigidScenePtr<UE::Physics::FSimContextRW>&)> SimPreTickFunc;
+	};
+
+
+	TEST_CASE("ChaosAsyncExample", "[Chaos][API][Scene][unit]")
+	{
+		using namespace Chaos;
+		using namespace Chaos::Rigids::Async;
+		using namespace UE::Physics;
+
+		float Dt = 1.0f / 60.0f;
+
+		SECTION("CreateSceneAndBodyExample")
+		{
+			// Create an async physics scene
+			FRigidFactoryAsync Factory;
+			FRigidSceneHandle SceneHandle = Factory.CreateRigidScene(FRigidSceneSettingsAsync());
+
+			// Create a rigid body
+			FRigidBodyHandle BodyHandle = SceneHandle.LockRW()->CreateRigidBody(MakeRigidName(TEXT("Body")), ERigidMovementType::Dynamic);
+
+			// Not Allowed from ReadOnly scene handle
+			//FRigidBodyHandle BodyHandle2 = SceneHandle.LockRO().CreateRigidBody(MakeRigidName(TEXT("Body")), ERigidMovementType::Dynamic);
+
+			// Set body position from the game thread
+			if (TRigidBodyPtr<FGameContextRW> BodyPtr = BodyHandle.LockRW())
+			{
+				BodyPtr->InitTransform(FTransform(FVector(0.0, 0.0, 100.0)));
+			}
+
+			// Tick physics for a while
+			for (int TickIndex = 0; TickIndex< 100; ++TickIndex)
+			{
+				SceneHandle.LockRW()->Tick(Dt);
+			}
+
+			// Destroy the body
+			SceneHandle.LockRW()->DestroyRigidBody(BodyHandle);
+
+			// Destroy the scene
+			Factory.DestroyRigidScene(SceneHandle);
+		}
+
+		SECTION("CallbackExample")
+		{
+			// Create an async physics scene
+			FRigidFactoryAsync Factory;
+			FRigidSceneHandle SceneHandle = Factory.CreateRigidScene(FRigidSceneSettingsAsync());
+
+			// Create a rigid body
+			FRigidBodyHandle BodyHandle = SceneHandle.LockRW()->CreateRigidBody(MakeRigidName(TEXT("Body")), ERigidMovementType::Dynamic);
+
+			// An example callback
+			FTestRigidSceneModifier SceneModifier;
+			SceneModifier.SimPreTickFunc =
+				[BodyHandle](TRigidScenePtr<FSimContextRW>& Scene)
+				{
+					// Pinning the handle with the specified scene restricts the API to only legal functions
+					if (TRigidBodyPtr<FSimContextRW> PTBodyPtr = BodyHandle.Pin(Scene))
+					{
+						// TODO: Still calls the GT functions so it asserts in SingleParticlePhysicsProxy
+						PTBodyPtr->InitTransform(FTransform(FVector(0.0, 0.0, 100.0)));
+
+						// Not Allowed from Simulation scene handle (GT-only)
+						//PTBody.SetM(200.0f);	
+					}
+				};
+
+			// Tick physics for a while
+			for (int TickIndex = 0; TickIndex < 100; ++TickIndex)
+			{
+				SceneHandle.LockRW()->Tick(Dt);
+			}
+
+			// Destroy the body
+			SceneHandle.LockRW()->DestroyRigidBody(BodyHandle);
+
+			// Destroy the scene
+			Factory.DestroyRigidScene(SceneHandle);
+		}
+	}
+
+
+	TEST_CASE("ChaosAsync", "[Chaos][API][Scene][unit]")
+	{
+		using namespace Chaos;
+		using namespace Chaos::Rigids::Async;
+		using namespace UE::Physics;
+
+		float Dt = 1.0f / 60.0f;
+
+		SECTION("CreateScene")
+		{
+			FRigidFactoryAsync Factory;
+
+			FRigidSceneSettingsAsync SceneSettings;
+			FRigidSceneHandle SceneHandle = Factory.CreateRigidScene(SceneSettings);
+
+			SECTION("TickScene")
+			{
+				SceneHandle.LockRW()->Tick(Dt);
+				SceneHandle.LockRW()->Tick(Dt);
+			}
+
+			SECTION("CreateAndDestroyBody")
+			{
+				// Empty RigidBodyHandle
+				FRigidBodyHandle BodyHandle;
+
+				// Handle should be invalid (explicit and implicit checks)
+				CHECK(!BodyHandle.LockRO().IsValid());
+				CHECK(!BodyHandle.LockRO());
+
+				// Create a body
+				BodyHandle = SceneHandle.LockRW()->CreateRigidBody(MakeRigidName(TEXT("Body")), ERigidMovementType::Dynamic);
+
+				// Body Handle should be valid
+				CHECK(BodyHandle.LockRO().IsValid());
+				CHECK(!!BodyHandle.LockRO());
+
+				SceneHandle.LockRW()->Tick(Dt);
+				SceneHandle.LockRW()->Tick(Dt);
+
+				// Destroy the body
+				SceneHandle.LockRW()->DestroyRigidBody(BodyHandle);
+
+				// RigidBody handle should is invalid (internally the container will 
+				// still be valid but the BodyId will not)
+				CHECK(!BodyHandle.LockRO().IsValid());
+				CHECK(!BodyHandle.LockRO());
+			}
+
+			SECTION("BodyHandleDuplication")
+			{
+				// Create a body
+				FRigidBodyHandle BodyHandle = SceneHandle.LockRW()->CreateRigidBody(MakeRigidName(TEXT("Body")), ERigidMovementType::Dynamic);
+				CHECK(!!BodyHandle.LockRO());
+
+				// Create a duplicate body handle - it should also be valid
+				FRigidBodyHandle BodyHandleDupe = BodyHandle;
+				CHECK(!!BodyHandleDupe.LockRO());
+
+				SceneHandle.LockRW()->Tick(Dt);
+				SceneHandle.LockRW()->Tick(Dt);
+
+				// Destroy the body
+				SceneHandle.LockRW()->DestroyRigidBody(BodyHandle);
+
+				// Body handle should is invalid (references destroyed body)
+				CHECK(!BodyHandle.LockRO());
+				CHECK(!BodyHandleDupe.LockRO());
+			}
+
+			SECTION("SetBodyTransform")
+			{
+				FRigidBodyHandle BodyHandle = SceneHandle.LockRW()->CreateRigidBody(MakeRigidName(TEXT("Body")), ERigidMovementType::Dynamic);
+				FRigidBodyHandle BodyHandleDupe = BodyHandle;
+
+				// Get write access to the body
+				if (TRigidBodyPtr<FGameContextRW> BodyPtr = BodyHandle.LockRW())
+				{
+					REQUIRE(!!BodyPtr);
+
+					// Body starts at the origin
+					FVector3d InitialBodyX = BodyPtr->GetTransform().GetTranslation();
+					CHECK(InitialBodyX == FVector(0,0,0));
+
+					// Move body
+					FVector3d NewBodyX = FVector(100,0,0);
+					BodyPtr->InitTransform(FTransform(NewBodyX));
+
+					// Check body has moved and the duplicated handle agrees
+					FVector3d FinalBodyX = BodyPtr->GetTransform().GetTranslation();
+					CHECK(FinalBodyX == NewBodyX);
+					CHECK(BodyHandleDupe.LockRO()->GetTransform().GetTranslation() == FinalBodyX);
+				}
+
+				SceneHandle.LockRW()->Tick(Dt);
+				SceneHandle.LockRW()->Tick(Dt);
+
+				// Destroy the body
+				SceneHandle.LockRW()->DestroyRigidBody(BodyHandle);
+			}
+
+			SECTION("InvalidBodyHandleAssertsOnUse")
+			{
+				// Create a body and immediately destroy it
+				FRigidBodyHandle BodyHandle = SceneHandle.LockRW()->CreateRigidBody(MakeRigidName(TEXT("Body")), ERigidMovementType::Dynamic);
+				CHECK(!!BodyHandle.LockRO());
+
+				SceneHandle.LockRW()->Tick(Dt);
+
+				SceneHandle.LockRW()->DestroyRigidBody(BodyHandle);
+				CHECK(!BodyHandle.LockRO());
+
+				// Get write access to the body - will generate invalid handle
+				TRigidBodyPtr<FGameContextRW> BodyPtrRW = BodyHandle.LockRW();
+				CHECK(!BodyPtrRW);
+
+				// These calls assert but there's no way to safely verify that happens
+				// because we don't use exceptions and check() uses __debugbreak
+				//CHECK_THROWS(BodyPtrRW->GetTransform());
+				//CHECK_THROWS(BodyPtrRW->SetTransform(FTransform::Identity));
+			}
+
+			SECTION("SceneModifierIsCalled")
+			{
+				std::atomic<int32> SceneModifierCallCount = 0;
+				FTestRigidSceneModifier SceneModifier;
+				SceneModifier.SimPreTickFunc =
+					[&SceneModifierCallCount](TRigidScenePtr<FSimContextRW>& Context)
+					{
+						++SceneModifierCallCount;
+					};
+
+				SceneHandle.LockRW()->RegisterModifier(&SceneModifier);
+
+				SceneHandle.LockRW()->Tick(Dt);
+				CHECK(SceneModifierCallCount == 1);
+
+				SceneHandle.LockRW()->Tick(Dt);
+				CHECK(SceneModifierCallCount == 2);
+
+				SceneHandle.LockRW()->UnregisterModifier(&SceneModifier);
+
+				SceneHandle.LockRW()->Tick(Dt);
+				CHECK(SceneModifierCallCount == 2);
+			}
+
+			SECTION("SceneModifierBodyHandles")
+			{
+				const FVector3d GTBodyX = FVector(100, 0, 0);
+				const FVector3d PTBodyX = FVector(200, 0, 0);
+
+				FRigidBodyHandle BodyHandle = SceneHandle.LockRW()->CreateRigidBody(MakeRigidName(TEXT("Body")), ERigidMovementType::Dynamic);
+				SceneHandle.LockRW()->RegisterRigidBody(BodyHandle);
+
+				BodyHandle.LockRW()->InitTransform(FTransform(GTBodyX));
+				CHECK(BodyHandle.LockRO()->GetTransform().GetTranslation() == GTBodyX);
+
+				// We don't want to particles to move on their own...
+				BodyHandle.LockRW()->SetGravityScale(0.0f);
+
+				// A BeginTick SceneModifier that changes the body position to PTBodyX
+				FTestRigidSceneModifier SceneModifier;
+				SceneModifier.SimPreTickFunc =
+					[BodyHandle, GTBodyX, PTBodyX](TRigidScenePtr<FSimContextRW>& Scene)
+					{
+						TRigidBodyPtr<FSimContextRW> PTBodyPtr = BodyHandle.Pin(Scene);
+						CHECK(PTBodyPtr.IsValid());
+
+						if (PTBodyPtr.IsValid())
+						{
+							// Make sure GT position change is reflected on PT
+							CHECK(PTBodyPtr->GetTransform().GetTranslation() == GTBodyX);
+
+							// Change the position of the body
+							PTBodyPtr->InitTransform(FTransform(PTBodyX));
+							CHECK(PTBodyPtr->GetTransform().GetTranslation() == PTBodyX);
+						}
+					};
+
+				SceneHandle.LockRW()->RegisterModifier(&SceneModifier);
+
+				// Tick the scene - modifier will be called
+				SceneHandle.LockRW()->Tick(Dt);
+
+				// Modifier should have moved the body, and the results
+				// should have synched back to the GT
+				CHECK(BodyHandle.LockRO()->GetTransform().GetTranslation() == PTBodyX);
+
+				// Put the body back
+				BodyHandle.LockRW()->InitTransform(FTransform(GTBodyX));
+				CHECK(BodyHandle.LockRO()->GetTransform().GetTranslation() == GTBodyX);
+
+				// Tick the scene - modifier will be called
+				SceneHandle.LockRW()->Tick(Dt);
+
+				// Modifier should have moved the body again
+				CHECK(BodyHandle.LockRO()->GetTransform().GetTranslation() == PTBodyX);
+
+				// Clean up
+				SceneHandle.LockRW()->UnregisterModifier(&SceneModifier);
+				SceneHandle.LockRW()->UnregisterRigidBody(BodyHandle);
+				SceneHandle.LockRW()->DestroyRigidBody(BodyHandle);
+			}
+
+			Factory.DestroyRigidScene(SceneHandle);
+		}
+	}
+
+} // Chaos::LowLevelTest
diff --git a/Engine/Source/Developer/ChaosTestHarness/Public/Chaos/LowLevelTest/ChaosTestMatchers.h b/Engine/Source/Developer/ChaosTestHarness/Public/Chaos/LowLevelTest/ChaosTestMatchers.h
--- a/Engine/Source/Developer/ChaosTestHarness/Public/Chaos/LowLevelTest/ChaosTestMatchers.h
+++ b/Engine/Source/Developer/ChaosTestHarness/Public/Chaos/LowLevelTest/ChaosTestMatchers.h
@@ -1,7 +1,6 @@
 // Copyright Epic Games, Inc. All Rights Reserved.
 #pragma once
 
-#include "CoreMinimal.h"
 #include "Chaos/Core.h"
 #include "Chaos/LowLevelTest/ChaosTestCatch2.h"
 
diff --git a/Engine/Source/Runtime/Experimental/RigidPhysics/Private/RigidPhysics/RigidScene.cpp b/Engine/Source/Runtime/Experimental/RigidPhysics/Private/RigidPhysics/RigidScene.cpp
new file mode 100644
--- /dev/null
+++ b/Engine/Source/Runtime/Experimental/RigidPhysics/Private/RigidPhysics/RigidScene.cpp
@@ -0,0 +1,55 @@
+// Copyright Epic Games, Inc. All Rights Reserved.
+#include "RigidPhysics/RigidScene.h"
+#include "RigidPhysics/RigidBody.h"
+#include "RigidPhysics/RigidBodyContainer.h"
+#include "RigidPhysics/RigidLog.h"
+
+#if UE_RIGIDPHYSICS_API_ENABLED
+
+DEFINE_LOG_CATEGORY(LogChaosRB);
+
+namespace UE::Physics
+{
+	template class TGameContext<EContextLockType::None>;
+	template class TGameContext<EContextLockType::ReadOnly>;
+	template class TGameContext<EContextLockType::ReadWrite>;
+
+	IRigidScene::IRigidScene()
+	{
+	}
+
+	IRigidScene::~IRigidScene()
+	{
+	}
+
+
+	IRigidBodyContainer::IRigidBodyContainer(const FRigidBodyContainerTypeId& InTypeId)
+		: TypeId(InTypeId)
+	{
+	}
+
+	IRigidBodyContainer::~IRigidBodyContainer()
+	{
+	}
+
+	const FRigidBodyContainerTypeId& IRigidBodyContainer::GetTypeId() const
+	{
+		return TypeId;
+	}
+
+	bool IRigidBodyContainer::IsCastableTo(const FRigidBodyContainerTypeId& InTypeId) const
+	{
+		return GetTypeId() == InTypeId;
+	}
+
+
+	IRigidBody::IRigidBody()
+	{
+	}
+
+	IRigidBody::~IRigidBody()
+	{
+	}
+}
+
+#endif // UE_RIGIDPHYSICS_API_ENABLED
\ No newline at end of file
diff --git a/Engine/Source/Runtime/Experimental/RigidPhysics/Public/RigidPhysics/RigidAccessor.h b/Engine/Source/Runtime/Experimental/RigidPhysics/Public/RigidPhysics/RigidAccessor.h
new file mode 100644
--- /dev/null
+++ b/Engine/Source/Runtime/Experimental/RigidPhysics/Public/RigidPhysics/RigidAccessor.h
@@ -0,0 +1,52 @@
+// Copyright Epic Games, Inc. All Rights Reserved.
+#pragma once
+
+#include "RigidPhysics/RigidFwd.h"
+
+#if UE_RIGIDPHYSICS_API_ENABLED
+
+#include "RigidPhysics/RigidBody.h"
+#include "RigidPhysics/RigidScene.h"
+
+namespace UE::Physics::Private
+{
+	class UE_INTERNAL FRigidAccessor
+	{
+	public:
+		UE_INTERNAL static const FRigidBodyId& GetRigidBodyId(const FRigidBodyHandle& BodyHandle)
+		{
+			return BodyHandle.BodyId;
+		}
+
+		UE_INTERNAL static uint64 GetRigidBodyIdInt(const FRigidBodyHandle& BodyHandle)
+		{
+			return BodyHandle.BodyId.IntId;
+		}
+
+		UE_INTERNAL static uint64 GetRigidBodyIdInt(const FRigidBodyId& BodyId)
+		{
+			return BodyId.IntId;
+		}
+
+		template<typename ContextType>
+		UE_INTERNAL static TRigidScenePtr<ContextType> MakeRigidScenePtr(const ContextType& Context, IRigidScene* Scene)
+		{
+			return TRigidScenePtr<ContextType>(Scene);
+		}
+
+		template<typename ContextType>
+		UE_INTERNAL static IRigidScene* GetIRigidScene(const TRigidScenePtr<ContextType>& InScenePtr)
+		{
+			return InScenePtr.SceneImpl.SceneRaw;
+		}
+
+		template<typename ContextType>
+		UE_INTERNAL static IRigidBody* GetIRigidBody(const TRigidBodyPtr<ContextType>& InBodyPtr)
+		{
+			return InBodyPtr.BodyImpl.BodyRaw;
+		}
+
+	};
+}
+
+#endif // UE_RIGIDPHYSICS_API_ENABLED
diff --git a/Engine/Source/Runtime/Experimental/RigidPhysics/Public/RigidPhysics/RigidBody.h b/Engine/Source/Runtime/Experimental/RigidPhysics/Public/RigidPhysics/RigidBody.h
new file mode 100644
--- /dev/null
+++ b/Engine/Source/Runtime/Experimental/RigidPhysics/Public/RigidPhysics/RigidBody.h
@@ -0,0 +1,686 @@
+// Copyright Epic Games, Inc. All Rights Reserved.
+#pragma once
+
+#include "RigidPhysics/RigidFwd.h"
+
+#if UE_RIGIDPHYSICS_API_ENABLED
+
+#include "Math/Bounds.h"
+#include "Misc/Crc.h"
+#include "RigidPhysics/RigidBodyContainer.h"
+#include "RigidPhysics/RigidLog.h"
+#include "RigidPhysics/RigidScene.h"
+#include "Templates/SharedPointer.h"
+
+// NOTES
+// 
+// Using SharedPtr and WeakPtr in a few places where they are technically required, 
+// but in practice may not be (e.g., FBodyHandle). However this is hidden as
+// an implementation detail so we have the option to switch to raw pointers if we
+// can guarantee lifetime, or retain pointer tracking in non-ship builds, etc.
+//
+
+extern uint32 GetTypeHash(UE::Physics::FRigidBodyId& RigidBodyId);
+
+namespace UE::Physics
+{
+	// An ID for a rigid body that is guaranteed to be unique and never reused
+	// within its container. There is no guarantee of uniqueness between
+	// containers, which means you need a Conainer + BodyId to access a body.
+	// (see FRigidBodyHandle, and FRigidBodyContainerHandle).
+	class FRigidBodyId
+	{
+	public:
+		FRigidBodyId()
+			: IntId(INDEX_NONE)
+		{
+		}
+
+		explicit FRigidBodyId(uint64 InId)
+			: IntId(InId)
+		{
+		}
+
+		void Reset()
+		{
+			IntId = INDEX_NONE;
+		}
+
+		friend bool operator==(const FRigidBodyId& L, const FRigidBodyId& R) = default;
+		
+		friend bool operator<(const FRigidBodyId& L, const FRigidBodyId& R)
+		{
+			return L.IntId < R.IntId;
+		}
+		
+		friend uint32 GetTypeHash(const FRigidBodyId& BodyId)
+		{ 
+			return GetTypeHash(BodyId.IntId);
+		}
+
+	private:
+		friend class Private::FRigidAccessor;
+
+		uint64 IntId;
+	};
+
+	// The full rigid body API. There will be several implementations of this API. E.g., synchronous
+	// physics, asynchronous physics, immediate physics, remote physics, etc.
+	// NOTE: for some implementations like async physics, there are both Game- and Physics-Context
+	// implementations within the async implementation. Locking a handle will return the object 
+	// appropriate for the context.
+	class RIGIDPHYSICS_API IRigidBody
+	{
+	public:
+		IRigidBody();
+		virtual ~IRigidBody();
+
+		virtual bool IsValid() const = 0;
+		virtual bool IsInScene() const = 0;
+
+		// The ID for this body. This is only guaranteed to be unique in its container.
+		// i.e., you need a Conainer+ID to access a body (or a FRigidBodyHandle).
+		virtual FRigidBodyId GetBodyId() const = 0;
+
+		// The handle to this body
+		virtual FRigidBodyHandle GetBodyHandle() const = 0;
+
+		// Is this body Static? Static bodies can be moved but have no velocity.
+		virtual bool IsStatic() const
+		{
+			RIGIDPHYSICS_API_UNSUPPORTED();
+			return true;
+		}
+
+		// Is this body kinematic? Kinematic bodies can be explicitly moved and have velocity, but do
+		// not react to forces, collisions, etc.
+		virtual bool IsKinematic() const
+		{
+			RIGIDPHYSICS_API_UNSUPPORTED();
+			return false;
+		}
+
+		// Is this a dynamic body? Dynamic bodies have momentum and react to forces and collisions.
+		virtual bool IsDynamic() const
+		{
+			RIGIDPHYSICS_API_UNSUPPORTED();
+			return false;
+		}
+
+		// Is this a sleeping dynamic body? Sleeping bodies are dynamic bodies that are ignored by
+		// the simulation (as an optimization) until the scene changes in a way that might affect them, 
+		// or they are explicitly woken.
+		virtual bool IsSleeping() const
+		{
+			RIGIDPHYSICS_API_UNSUPPORTED();
+			return false;
+		}
+
+		// Set the transform and initialize body state (velocity etc.) at the new location.
+		virtual void InitTransform(const FTransform& InTransform)
+		{
+			RIGIDPHYSICS_API_UNSUPPORTED();
+		}
+
+		// Set the transform without changing any other rigid body state.
+		virtual void UpdateTransform(const FTransform& InTransform)
+		{
+			RIGIDPHYSICS_API_UNSUPPORTED();
+		}
+
+		// Set the kinematic target. The body will attempt to reach the target at the next
+		// tick. Kinematic bodies will always reach the target bscause they do not respond
+		// to forces or collisions. Dynamic bodies will only reach their target if their
+		// velocity in the absence of interactions (forces, collisions, or other constraints).
+		virtual void SetKinematicTarget(const FTransform& InTransform)
+		{
+			RIGIDPHYSICS_API_UNSUPPORTED();
+		}
+
+		// Get the current transform. NOTE: This will immediately reflect any recent calls to
+		// to InitTransform or UpdateTransform, but will not reflect the latest KinematicTarget
+		// until after the next physics tick.
+		virtual FTransform GetTransform() const
+		{
+			RIGIDPHYSICS_API_UNSUPPORTED();
+			return FTransform::Identity;
+		}
+
+		// Get the world-space bounds at the current transform
+		virtual FBounds3d GetBounds() const
+		{
+			RIGIDPHYSICS_API_UNSUPPORTED();
+			return FBounds3d(FVector3d::Zero(), FVector3d::Zero());
+		}
+
+		// TODO: Maybe just Set/GetMassProperties
+		virtual double GetMass() const
+		{
+			RIGIDPHYSICS_API_UNSUPPORTED();
+			return 0;
+		}
+
+		virtual void SetMass(double InM)
+		{
+			RIGIDPHYSICS_API_UNSUPPORTED();
+		}
+
+		virtual void SetCCDEnabled(bool bInEnabled)
+		{
+			RIGIDPHYSICS_API_UNSUPPORTED();
+		}
+
+		virtual void SetMACDEnabled(bool bInEnabled)
+		{
+			RIGIDPHYSICS_API_UNSUPPORTED();
+		}
+		
+		virtual void SetSmoothEdgeCollisionsEnabled(bool bInEnabled)
+		{
+			RIGIDPHYSICS_API_UNSUPPORTED();
+		}
+
+		virtual void SetMaxLinearVelocity(float InMax)
+		{
+			RIGIDPHYSICS_API_UNSUPPORTED();
+		}
+		
+		virtual void SetMaxAngularVelocity(float InMax)
+		{
+			RIGIDPHYSICS_API_UNSUPPORTED();
+		}
+		
+		virtual void SetInertiaConditioningEnabled(bool bInEnabled)
+		{
+			RIGIDPHYSICS_API_UNSUPPORTED();
+		}
+		
+		virtual void SetGyroscopicTorqueEnabled(bool bInEnabled)
+		{
+			RIGIDPHYSICS_API_UNSUPPORTED();
+		}
+
+		virtual void SetGravityGroup(int32 InId)
+		{
+			RIGIDPHYSICS_API_UNSUPPORTED();
+		}
+		
+		virtual void SetGravityScale(float InScale)
+		{
+			RIGIDPHYSICS_API_UNSUPPORTED();
+		}
+
+		virtual void SetWakeEventsEnabled(bool bInEnabled)
+		{
+			RIGIDPHYSICS_API_UNSUPPORTED();
+		}
+		
+		virtual void SetNumPositionIterations(int32 InNum)
+		{
+			RIGIDPHYSICS_API_UNSUPPORTED();
+		}
+		
+		virtual void SetNumVelocityIterations(int32 InNum)
+		{
+			RIGIDPHYSICS_API_UNSUPPORTED();
+		}
+		
+		virtual void SetNumProjectionIterations(int32 InNum)
+		{
+			RIGIDPHYSICS_API_UNSUPPORTED();
+		}
+		
+		virtual void SetIsSleeping(bool bInSleeping)
+		{
+			RIGIDPHYSICS_API_UNSUPPORTED();
+		}
+
+	protected:
+
+		virtual const char* GetTypeName() const
+		{
+			return "Unknown";
+		}
+	};
+
+	// Internal API for rigid bodies. Not intended for public use.
+	class RIGIDPHYSICS_API IRigidBodyInternal : public IRigidBody
+	{
+	public:
+		IRigidBodyInternal() {}
+		virtual ~IRigidBodyInternal() {}
+
+		virtual void SetSpatialIdx() {}
+	};
+
+
+	// This class only exists so that TRigidBodyPtr can have pointer semantics rather than object semantics.
+	// See TRigidBodyPtr for comments.
+	template<typename ContextType>
+	class TRigidBodyPtrImpl
+	{
+	public:
+		using FContext = ContextType;
+		using FRigidScenePtr = TRigidScenePtr<ContextType>;
+
+		~TRigidBodyPtrImpl()
+		{
+			Reset();
+		}
+
+		// Forward to IRigidBody. See IRigidBody for comments
+		bool IsStatic() const
+		{
+			return BodyRaw->IsStatic();
+		}
+		
+		bool IsKinematic() const
+		{
+			return BodyRaw->IsKinematic();
+		}
+
+		bool IsDynamic() const
+		{
+			return BodyRaw->IsDynamic();
+		}
+
+		void InitTransform(const FTransform& InTransform)
+		{
+			BodyRaw->InitTransform(InTransform);
+		}
+
+		void UpdateTransform(const FTransform& InTransform)
+		{
+			BodyRaw->UpdateTransform(InTransform);
+		}
+
+		void SetKinematicTarget(const FTransform& InTransform)
+		{
+			BodyRaw->SetKinematicTarget(InTransform);
+		}
+
+		FTransform GetTransform() const
+		{
+			return BodyRaw->GetTransform();
+		}
+
+		FBounds3d GetBounds() const
+		{
+			return BodyRaw->GetBounds();
+		}
+		
+		double GetMass() const
+		{
+			return BodyRaw->GetMass();
+		}
+
+		void SetMass(double InMass)
+		{
+			BodyRaw->SetMass(InMass);
+		}
+		
+		void SetCCDEnabled(bool bInEnabled)
+		{
+			BodyRaw->SetCCDEnabled(bInEnabled);
+		}
+
+		void SetMACDEnabled(bool bInEnabled)
+		{
+			BodyRaw->SetMACDEnabled(bInEnabled);
+		}
+
+		void SetSmoothEdgeCollisionsEnabled(bool bInEnabled)
+		{
+			BodyRaw->SetSmoothEdgeCollisionsEnabled(bInEnabled);
+		}
+
+		void SetMaxLinearVelocity(float InMax)
+		{
+			BodyRaw->SetMaxLinearVelocity(InMax);
+		}
+		
+		void SetMaxAngularVelocity(float InMax)
+		{
+			BodyRaw->SetMaxAngularVelocity(InMax);
+		}
+
+		void SetInertiaConditioningEnabled(bool bInEnabled)
+		{
+			BodyRaw->SetInertiaConditioningEnabled(bInEnabled);
+		}
+
+		void SetGyroscopicTorqueEnabled(bool bInEnabled)
+		{
+			BodyRaw->SetGyroscopicTorqueEnabled(bInEnabled);
+		}
+
+		void SetGravityGroup(int32 InId)
+		{
+			BodyRaw->SetGravityGroup(InId);
+		}
+		
+		void SetGravityScale(float InScale)
+		{
+			BodyRaw->SetGravityScale(InScale);
+		}
+
+		void SetWakeEventsEnabled(bool bInEnabled)
+		{
+			BodyRaw->SetWakeEventsEnabled(bInEnabled);
+		}
+		void SetNumPositionIterations(int32 InNum)
+		{
+			BodyRaw->SetNumPositionIterations(InNum);
+		}
+		void SetNumVelocityIterations(int32 InNum)
+		{
+			BodyRaw->SetNumVelocityIterations(InNum);
+		}
+
+		void SetNumProjectionIterations(int32 InNum)
+		{
+			BodyRaw->SetNumProjectionIterations(InNum);
+		}
+		
+		void SetIsSleeping(bool bInSleeping)
+		{
+			BodyRaw->SetIsSleeping(bInSleeping);
+		}
+
+		friend bool operator==(const TRigidBodyPtrImpl& L, const TRigidBodyPtrImpl& R) = default;
+		friend bool operator!=(const TRigidBodyPtrImpl& L, const TRigidBodyPtrImpl& R) = default;
+
+	private:
+		friend class Private::FRigidAccessor;
+		friend class TRigidBodyPtr<FContext>;
+
+		TRigidBodyPtrImpl(const TRigidBodyPtrImpl& R) = default;
+		TRigidBodyPtrImpl(TRigidBodyPtrImpl&& R) = default;
+
+		TRigidBodyPtrImpl(IRigidBody* InBody)
+			: BodyRaw(InBody)
+			, SceneLock()
+		{
+		}
+
+		// NOTE: Receiving a transient lock, so we must store it
+		void Lock(FRigidScenePtr&& InScene)
+		{
+			SceneLock = MoveTemp(InScene);
+		}
+
+		void Reset()
+		{
+			BodyRaw = nullptr;
+		}
+
+		bool IsValid() const
+		{
+			return (BodyRaw != nullptr);
+		}
+
+		// IRigidBody implementation depends on Context type
+		IRigidBody* BodyRaw;
+
+		// Holds a lock on the scene (if set)
+		FRigidScenePtr SceneLock;
+	};
+
+
+	// A context-aware IRigidBody pointer that simply forwards all context-legal operations to IRigidBody. 
+	// Provides access to a restricted RigidBody API where the restrictions depend on the context 
+	// (template-parameter), and attempts to call disallowed functions (in the context) produce compiler errors.
+	// 
+	// Uses pointer semantics and may be "null". Check IsValid() before dereferencing. Also implicitly casts 
+	// to boolean for validity checks.
+	// 
+	// Cannot be copied or stored - use FRigidBodyHandle for that.
+	template<typename ContextType>
+	class TRigidBodyPtr
+	{
+	public:
+		using FContext = ContextType;
+		using FRigidBodyPtrImpl = TRigidBodyPtrImpl<FContext>;
+		using FRigidBodyContainerPtr = TRigidBodyContainerPtr<FContext>;
+		using FRigidScenePtr = TRigidScenePtr<ContextType>;
+
+		TRigidBodyPtr()
+			: BodyImpl(nullptr)
+		{
+		}
+
+		~TRigidBodyPtr()
+		{
+		}
+
+		// Has this pointer been initialized with a body pointer?
+		bool IsValid() const
+		{
+			return BodyImpl.IsValid();
+		}
+
+		// Auto cast to bool to support using in if statements:
+		//	if (TRigidBodyPtr<FGameContextRO> Body = BodyHandle.LockRO())...
+		operator bool() const
+		{
+			return IsValid();
+		}
+
+		// Get the handle for this body
+		FRigidBodyHandle GetHandle() const;
+
+		operator FRigidBodyHandle() const;
+
+		// Pointer sematics
+		const FRigidBodyPtrImpl* operator->() const
+		{
+			checkf(IsValid(), TEXT("Check IsValid() before dereferencing"));
+			return &BodyImpl;
+		}
+
+		// Pointer sematics
+		FRigidBodyPtrImpl* operator->() const requires (FContext::bWriteEnabled)
+		{
+			checkf(IsValid(), TEXT("Check IsValid() before dereferencing"));
+			return &BodyImpl;
+		}
+
+		friend bool operator==(const TRigidBodyPtr& L, const TRigidBodyPtr& R) = default;
+		friend bool operator!=(const TRigidBodyPtr& L, const TRigidBodyPtr& R) = default;
+
+	private:
+		friend class Private::FRigidAccessor;
+		friend class TRigidBodyContainerPtrImpl<FContext>;
+		friend class TRigidScenePtrImpl<FContext>;
+		friend class FRigidBodyHandle;
+
+		TRigidBodyPtr(const TRigidBodyPtr& R) = default;
+		TRigidBodyPtr(TRigidBodyPtr&& R) = default;
+
+		// Create a pointer with no scene lock. Used in situations were the scene is known to  be 
+		// locked before the rigid body pointer is retrieved from the ID. This will be when 
+		// the user wants to use multiple physics objects under one lock.
+		TRigidBodyPtr(IRigidBody* InBody)
+			: BodyImpl(InBody)
+		{
+		}
+
+		// NOTE: Receiving a transient lock, so we must store it. Should be used
+		// when only a small number of physics objects are being accessed and extra locking
+		// overhead is not an issue. Otherise an external scene lock is preferred.
+		TRigidBodyPtr& Lock(FRigidScenePtr&& InScene)
+		{
+			BodyImpl.Lock(MoveTemp(InScene));
+			return *this;
+		}
+
+		mutable FRigidBodyPtrImpl BodyImpl;
+	};
+
+
+	// A handle to a rigid body in a container.
+	// This handle may be copied and stored, and can be used to retrieve a useful
+	// body reference by "locking" it and providing the necessary context.
+	// 
+	// Example contexts are: 
+	// - Read-Only Game-Thread
+	// - Read-Write Game Thread
+	// - Physics-Thread BeginTick Callback
+	// - etc.
+	//
+	// NOTE: If you have many bodies from the same container, it will be more efficient
+	// to store a set of BodyIds separately from the ContainerHandle. See TRigidBodyPtr.
+	//
+	class FRigidBodyHandle
+	{
+	public:
+		FRigidBodyHandle() = default;
+
+		FRigidBodyHandle(const FRigidBodyContainerHandle& InBodyContainerHandle, const FRigidBodyId& InBodyId)
+			: BodyContainerHandle(InBodyContainerHandle)
+			, BodyId(InBodyId)
+		{
+		}
+
+		void Reset()
+		{
+			BodyContainerHandle.Reset();
+			BodyId.Reset();
+		}
+
+		const FRigidSceneHandle& GetSceneHandle() const
+		{
+			return BodyContainerHandle.GetSceneHandle();
+		}
+
+		const FRigidBodyContainerHandle& GetContainerHandle() const
+		{
+			return BodyContainerHandle;
+		}
+
+		const FRigidBodyId& GetBodyId() const
+		{
+			return BodyId;
+		}
+
+		// Get a pointer to the GameContext rigid body and hold a read-lock on the scene for the duration of the pointer.
+		TRigidBodyPtr<FGameContextRO> LockRO() const
+		{
+			// NOTE: Passing a transient lock.
+			return Lock(GetSceneHandle().LockRO());
+		}
+
+		// Get a pointer to the GameContext rigid body and hold a write-lock on the scene for the duration of the pointer.
+		TRigidBodyPtr<FGameContextRW> LockRW() const
+		{
+			// NOTE: Passing a transient lock.
+			return Lock(GetSceneHandle().LockRW());
+		}
+
+		// Get a pointer to the Context-specific rigid body. This does not hold a lock on the scene,
+		// which is assumed to be locked externally with a TRigidScenePtr whose lifetime
+		// exceeds this pointer. Using this pointer after the external lock is released will 
+		// cause race conditions and random crashes.
+		template<typename ContextType>
+		TRigidBodyPtr<ContextType> Pin(const TRigidScenePtr<ContextType>& InScene) const
+		{
+			if (TRigidBodyContainerPtr<ContextType> Container = BodyContainerHandle.Pin(InScene))
+			{
+				return Container->GetRigidBody(BodyId);
+			}
+			return TRigidBodyPtr<ContextType>();
+		}
+
+		friend bool operator==(const FRigidBodyHandle& L, const FRigidBodyHandle& R) = default;
+		friend bool operator!=(const FRigidBodyHandle& L, const FRigidBodyHandle& R) = default;
+
+	private:
+		friend class Private::FRigidAccessor;
+
+		// Rvalue version is private and deleted to prevent users from being able to Pin with a
+		// lock whose lifetime is limited to the scope of the call. Do not implement this 
+		// as it could easily be used incorrectly. See Lock() instead.
+		template<typename ContextType>
+		TRigidBodyPtr<ContextType> Pin(TRigidScenePtr<ContextType>&& InScene) const = delete;
+
+		// NOTE: Receiving a transient lock, so we must store it in RigidBodyPtr.
+		template<typename ContextType>
+		TRigidBodyPtr<ContextType> Lock(TRigidScenePtr<ContextType>&& InScene) const
+		{
+			return Pin(InScene).Lock(MoveTemp(InScene));
+		}
+
+		FRigidBodyContainerHandle BodyContainerHandle;
+		FRigidBodyId BodyId;
+	};
+
+
+	template<typename ContextType>
+	FRigidBodyHandle TRigidBodyPtr<ContextType>::GetHandle() const
+	{
+		if (IsValid())
+		{
+			return BodyImpl.BodyRaw->GetBodyHandle();
+		}
+		return FRigidBodyHandle();
+	}
+
+	template<typename ContextType>
+	TRigidBodyPtr<ContextType>::operator FRigidBodyHandle() const
+	{
+		return GetHandle();
+	}
+
+	//
+	// Container implementations to avoid circular dependency in headers
+	//
+
+	template<typename ContextType>
+	TRigidBodyPtr<ContextType> TRigidBodyContainerPtrImpl<ContextType>::GetRigidBody(const FRigidBodyId& BodyId) const
+	{
+		UE_RIGIDPHYSICS_CHECK(IsValid());
+		return TRigidBodyPtr<ContextType>(BodyContainerRaw->GetRigidBody(BodyId));
+	}
+
+	//
+	// Scene implementations to avoid circular dependency in headers
+	//
+
+	template<typename ContextType>
+	TRigidBodyPtr<ContextType> TRigidScenePtrImpl<ContextType>::CreateRigidBody(const FRigidNamePtr& InName, ERigidMovementType InMovementType) const
+	{
+		UE_RIGIDPHYSICS_CHECK(IsValid());
+		return TRigidBodyPtr<FContext>(SceneRaw->CreateRigidBody(InName, InMovementType));
+	}
+
+	template<typename ContextType>
+	void TRigidScenePtrImpl<ContextType>::DestroyRigidBody(const FRigidBodyHandle& BodyHandle) const
+	{
+		UE_RIGIDPHYSICS_CHECK(IsValid());
+		SceneRaw->DestroyRigidBody(BodyHandle);
+	}
+
+	template<typename ContextType>
+	void TRigidScenePtrImpl<ContextType>::RegisterRigidBody(const FRigidBodyHandle& BodyHandle)
+	{
+		UE_RIGIDPHYSICS_CHECK(IsValid());
+		SceneRaw->RegisterRigidBody(BodyHandle);
+	}
+
+	template<typename ContextType>
+	void TRigidScenePtrImpl<ContextType>::UnregisterRigidBody(const FRigidBodyHandle& BodyHandle)
+	{
+		UE_RIGIDPHYSICS_CHECK(IsValid());
+		SceneRaw->UnregisterRigidBody(BodyHandle);
+	}
+
+	template<typename ContextType>
+	TRigidBodyContainerPtr<ContextType> TRigidScenePtrImpl<ContextType>::GetRigidBodyContainer(const FRigidBodyContainerId& InBodyContainerId) const
+	{
+		UE_RIGIDPHYSICS_CHECK(IsValid());
+		return TRigidBodyContainerPtr<ContextType>(SceneRaw->GetRigidBodyContainer(InBodyContainerId));
+	}
+
+}
+
+#endif // UE_RIGIDPHYSICS_API_ENABLED
\ No newline at end of file
diff --git a/Engine/Source/Runtime/Experimental/RigidPhysics/Public/RigidPhysics/RigidBodyContainer.h b/Engine/Source/Runtime/Experimental/RigidPhysics/Public/RigidPhysics/RigidBodyContainer.h
new file mode 100644
--- /dev/null
+++ b/Engine/Source/Runtime/Experimental/RigidPhysics/Public/RigidPhysics/RigidBodyContainer.h
@@ -0,0 +1,332 @@
+// Copyright Epic Games, Inc. All Rights Reserved.
+#pragma once
+
+#include "RigidPhysics/RigidFwd.h"
+
+#if UE_RIGIDPHYSICS_API_ENABLED
+
+#include "Misc/Crc.h"
+#include "RigidPhysics/RigidScene.h"
+#include "Templates/SharedPointer.h"
+
+namespace UE::Physics
+{
+	// A type ID for a rigid body containers. There will be multiple types of
+	// container, and this can grow over time. Initially we will have:
+	// SingleBodyContainer, GeometryCollection, ClusterUnion.
+	class FRigidBodyContainerTypeId
+	{
+	public:
+		FRigidBodyContainerTypeId(const TCHAR* InTypeName)
+			: TypeId(FCrc::StrCrc32(InTypeName))
+			, TypeName(InTypeName)
+		{
+		}
+
+		friend bool operator==(const FRigidBodyContainerTypeId& L, const FRigidBodyContainerTypeId& R) = default;
+		friend bool operator<(const FRigidBodyContainerTypeId& L, const FRigidBodyContainerTypeId& R)
+		{
+			return L.TypeId < R.TypeId;
+		}
+
+	private:
+		uint64 TypeId;
+		const TCHAR* TypeName;
+	};
+
+
+	// A container ID. All containers are registered with with RigidScene and can be found
+	// from their ID (usually from a FRigidBodyContainerHandle).
+	class FRigidBodyContainerId
+	{
+	public:
+		FRigidBodyContainerId()
+		{
+			IntId = -1;
+		}
+
+		explicit FRigidBodyContainerId(uint64 InId)
+		{
+			IntId = InId;
+		}
+
+		void Reset()
+		{
+			IntId = -1;
+		}
+
+		friend bool operator==(const FRigidBodyContainerId& L, const FRigidBodyContainerId& R) = default;
+		friend bool operator<(const FRigidBodyContainerId& L, const FRigidBodyContainerId& R)
+		{
+			return L.IntId < R.IntId;
+		}
+		friend uint32 GetTypeHash(const FRigidBodyContainerId& BodyContainerId)
+		{ 
+			return GetTypeHash(BodyContainerId.IntId);
+		}
+
+	private:
+		friend class Private::FRigidAccessor;
+
+		uint64 IntId;
+	};
+
+
+	// The interface for rigid body containers. There will be several implementations of this API. E.g., synchronous
+	// physics, asynchronous physics, immediate physics, remote physics, etc. 
+	class RIGIDPHYSICS_API IRigidBodyContainer
+	{
+	public:
+		IRigidBodyContainer(const FRigidBodyContainerTypeId& InTypeId);
+		virtual ~IRigidBodyContainer();
+
+		const FRigidBodyContainerTypeId& GetTypeId() const;
+
+		bool IsCastableTo(const FRigidBodyContainerTypeId& InTypeId) const;
+
+		// Downcast to the specified container type
+		template<typename T>
+		T* TryAsA()
+		{
+			if (IsCastableTo(T::StaticRigidBodyContainerTypeId()))
+			{
+				return static_cast<T&>(*this);
+			}
+			return nullptr;
+		}
+
+		virtual FRigidBodyContainerId GetContainerId() const = 0;
+
+		virtual IRigidScene* GetScene() const = 0;
+
+		virtual IRigidBody* GetRigidBody(const FRigidBodyId& BodyId) const = 0;
+
+	private:
+		FRigidBodyContainerTypeId TypeId;
+	};
+
+
+	// Utility wrapper around IRigidBodyContainer that adds Context and allows
+	// TRigidBodyContainerPtr to have pointer semantics
+	template<typename ContextType>
+	class TRigidBodyContainerPtrImpl
+	{
+	public:
+		using FContext = ContextType;
+		using FRigidScenePtr = TRigidScenePtr<ContextType>;
+
+		TRigidBodyContainerPtrImpl(IRigidBodyContainer* InBodyContainer)
+			: BodyContainerRaw(InBodyContainer)
+			, SceneLock()
+		{
+		}
+
+		virtual ~TRigidBodyContainerPtrImpl()
+		{
+			Reset();
+		}
+
+		void Reset()
+		{
+			BodyContainerRaw = nullptr;
+		}
+
+		bool IsValid() const
+		{
+			return (BodyContainerRaw != nullptr);
+		}
+
+		// Get the type ID of this container
+		const FRigidBodyContainerTypeId& GetContainerTypeId() const
+		{
+			UE_RIGIDPHYSICS_CHECK(IsValid());
+			return BodyContainerRaw->GetTypeId();
+		}
+
+		// Get the ID of this container
+		FRigidBodyContainerId GetContainerId() const
+		{
+			UE_RIGIDPHYSICS_CHECK(IsValid());
+			return BodyContainerRaw->GetContainerId();
+		}
+
+		// Impl in RigidBody.h
+		TRigidBodyPtr<FContext> GetRigidBody(const FRigidBodyId& BodyId) const;
+
+	private:
+		friend class TRigidBodyContainerPtr<ContextType>;
+
+		TRigidBodyContainerPtrImpl(const TRigidBodyContainerPtrImpl& R)
+			: BodyContainerRaw(R.BodyContainerRaw)
+			, SceneLock(R.SceneLock)
+		{
+		}
+
+		// We are passed a transient lock, so we must store it
+		void Lock(FRigidScenePtr&& InScene)
+		{
+			SceneLock = MoveTemp(InScene);
+		}
+
+		// Pointer to implementation, object depends on Context type
+		IRigidBodyContainer* BodyContainerRaw;
+
+		// Holds a lock on the scene if set
+		FRigidScenePtr SceneLock;
+	};
+
+
+	// A context-aware IRigidBodyContainer pointer that simply forwards all context-legal operations to IRigidBodyContainer. 
+	// Provides access to a restricted RigidBodyContainer API where the restrictions depend on the context 
+	// (template-parameter), and attempts to call disallowed functions (in the context) produce compiler errors.
+	// 
+	// Uses pointer semantics and may be "null". Check IsValid() before dereferencing. Also implicitly casts 
+	// to boolean for validity checks.
+	// 
+	// Cannot be copied or stored - use FRigidBodyContainerHandle for that.
+	template<typename ContextType>
+	class TRigidBodyContainerPtr
+	{
+	public:
+		using FContext = ContextType;
+		using FRigidBodyContainerPtrImpl = TRigidBodyContainerPtrImpl<FContext>;
+		using FRigidScenePtr = TRigidScenePtr<ContextType>;
+
+		TRigidBodyContainerPtr()
+			: BodyContainerImpl(nullptr)
+		{
+		}
+
+		~TRigidBodyContainerPtr()
+		{
+		}
+
+		void Reset()
+		{
+			BodyContainerImpl.Reset();
+		}
+
+		bool IsValid() const
+		{
+			return BodyContainerImpl.IsValid();
+		}
+
+		operator bool() const
+		{
+			return IsValid();
+		}
+
+		const FRigidBodyContainerPtrImpl* operator->() const
+		{
+			return &BodyContainerImpl;
+		}
+
+	protected:
+		friend class Private::FRigidAccessor;
+		friend class FRigidBodyContainerHandle;
+		friend class TRigidScenePtrImpl<FContext>;
+
+		TRigidBodyContainerPtr(IRigidBodyContainer* InBodyContainer)
+			: BodyContainerImpl(InBodyContainer)
+		{
+		}
+
+		TRigidBodyContainerPtr(const TRigidBodyContainerPtr& R)
+			: BodyContainerImpl(R.BodyContainerImpl)
+		{
+		}
+
+		// We are passed a transient lock, so we must store it
+		TRigidBodyContainerPtr& Lock(FRigidScenePtr&& InScene)
+		{
+			BodyContainerImpl.Lock(MoveTemp(InScene));
+			return *this;
+		}
+
+		FRigidBodyContainerPtrImpl BodyContainerImpl;
+	};
+
+
+	class FRigidBodyContainerHandle
+	{
+	public:
+		FRigidBodyContainerHandle()
+		{
+		}
+
+		FRigidBodyContainerHandle(FRigidSceneHandle InSceneHandle, const FRigidBodyContainerId& InBodyContainerId)
+			: SceneHandle(InSceneHandle)
+			, BodyContainerId(InBodyContainerId)
+		{
+		}
+
+		void Reset()
+		{
+			SceneHandle.Reset();
+			BodyContainerId.Reset();
+		}
+
+		const FRigidSceneHandle& GetSceneHandle() const
+		{
+			return SceneHandle;
+		}
+
+		const FRigidBodyContainerId& GetContainerId() const
+		{
+			return BodyContainerId;
+		}
+
+		// Get a pointer to the container and hold a read-lock on the scene for the duration of the pointer.
+		TRigidBodyContainerPtr<FGameContextRO> LockRO() const
+		{
+			// NOTE: Passing a transient lock to Lock()
+			return Lock(SceneHandle.LockRO());
+		}
+
+		// Get a pointer to the container and hold a write-lock on the scene for the duration of the pointer.
+		TRigidBodyContainerPtr<FGameContextRW> LockRW() const
+		{
+			// NOTE: Passing a transient lock to Lock()
+			return Lock(SceneHandle.LockRW());
+		}
+
+		// Get a pointer to the GameContext rigid body container. This does not hold a lock on the scene,
+		// which is assumed to be locked externally with a TRigidScenePtr whose lifetime
+		// exceeds this pointer's. Using this pointer after the external lock is released will 
+		// cause race conditions and random crashes.
+		template<typename ContextType>
+		TRigidBodyContainerPtr<ContextType> Pin(const TRigidScenePtr<ContextType>& InScene) const
+		{
+			if (InScene.IsValid())
+			{
+				return InScene->GetRigidBodyContainer(BodyContainerId);
+			}
+			return TRigidBodyContainerPtr<ContextType>();
+		}
+
+		friend bool operator==(const FRigidBodyContainerHandle& L, const FRigidBodyContainerHandle& R) = default;
+		friend bool operator!=(const FRigidBodyContainerHandle& L, const FRigidBodyContainerHandle& R) = default;
+
+	private:
+		friend class Private::FRigidAccessor;
+
+		// Rvalue version is private and deleted to prevent users from being able to Pin with a
+		// lock whose lifetime is limited to the scope of the call. Do not implement this 
+		// as it could easily be used incorrectly. See Lock() instead.
+		template<typename ContextType>
+		TRigidBodyContainerPtr<ContextType> Pin(TRigidScenePtr<ContextType>&& InScene) const = delete;
+
+		// NOTE: Receiving a transient lock, so we must store it in TRigidBodyContainerPtr.
+		template<typename ContextType>
+		TRigidBodyContainerPtr<ContextType> Lock(TRigidScenePtr<ContextType>&& InScene) const
+		{
+			return Pin(InScene).Lock(MoveTemp(InScene));
+		}
+
+
+		FRigidSceneHandle SceneHandle;
+		FRigidBodyContainerId BodyContainerId;
+	};
+}
+
+#endif // UE_RIGIDPHYSICS_API_ENABLED
\ No newline at end of file
diff --git a/Engine/Source/Runtime/Experimental/RigidPhysics/Public/RigidPhysics/RigidContext.h b/Engine/Source/Runtime/Experimental/RigidPhysics/Public/RigidPhysics/RigidContext.h
new file mode 100644
--- /dev/null
+++ b/Engine/Source/Runtime/Experimental/RigidPhysics/Public/RigidPhysics/RigidContext.h
@@ -0,0 +1,100 @@
+// Copyright Epic Games, Inc. All Rights Reserved.
+#pragma once
+
+#include "CoreMinimal.h"
+#include "RigidPhysics/RigidFwd.h"
+#include "RigidPhysics/RigidLog.h"
+
+#if UE_RIGIDPHYSICS_API_ENABLED
+
+namespace UE::Physics
+{
+	// The game callback phase.
+	enum class EContextGamePhase
+	{
+		None,
+		PreSimulate,
+		PostSimulate,
+		//...
+	};
+
+	// The simulation callback phase.
+	enum class EContextSimPhase
+	{
+		None,
+		PreTick,
+		PostTick,
+		//...
+	};
+
+	enum class ERigidMovementType
+	{
+		Static,
+		Kinematic,
+		Dynamic
+	};
+
+	class IRigidContext
+	{
+	public:
+		IRigidContext() {}
+		virtual ~IRigidContext() {}
+	};
+
+
+
+	template<EContextLockType LockTypeParam>
+	class RIGIDPHYSICS_API TGameContext : public IRigidContext
+	{
+	public:
+		static const EContextEnv ContextEnv = EContextEnv::Game;
+		static const EContextLockType LockType = LockTypeParam;
+		static const EContextSimPhase SimPhase = EContextSimPhase::None;
+		static const bool bWriteEnabled = (LockType == EContextLockType::ReadWrite);
+
+		TGameContext(EContextGamePhase InGamePhase)
+			: GamePhase(InGamePhase)
+		{
+		}
+
+		EContextGamePhase GetGamePhase() const
+		{
+			return GamePhase;
+		}
+
+	private:
+		EContextGamePhase GamePhase;
+	};
+
+	using FGameContextRO = TGameContext<EContextLockType::ReadOnly>;
+	using FGameContextRW = TGameContext<EContextLockType::ReadWrite>;
+
+
+	class RIGIDPHYSICS_API FSimContextRW : public IRigidContext
+	{
+	public:
+		static const EContextEnv ContextEnv = EContextEnv::Sim;
+		static const EContextLockType LockType = EContextLockType::None;
+		static const bool bWriteEnabled = true;
+
+		FSimContextRW()
+			: SimPhase(EContextSimPhase::None)
+		{
+		}
+
+		EContextSimPhase GetSimPhase() const
+		{
+			return SimPhase;
+		}
+
+		void SetSimPhase(EContextSimPhase InSimPhase)
+		{
+			SimPhase = InSimPhase;
+		}
+
+	private:
+		EContextSimPhase SimPhase = EContextSimPhase::None;
+	};
+}
+
+#endif // UE_RIGIDPHYSICS_API_ENABLED
\ No newline at end of file
diff --git a/Engine/Source/Runtime/Experimental/RigidPhysics/Public/RigidPhysics/RigidFwd.h b/Engine/Source/Runtime/Experimental/RigidPhysics/Public/RigidPhysics/RigidFwd.h
new file mode 100644
--- /dev/null
+++ b/Engine/Source/Runtime/Experimental/RigidPhysics/Public/RigidPhysics/RigidFwd.h
@@ -0,0 +1,116 @@
+// Copyright Epic Games, Inc. All Rights Reserved.
+#pragma once
+
+// A define to allow us to toggle between old and new APIs in BodyInstance etc
+// Also can be searched to identify all engine changes required to change to the new API
+//
+// NOTE: Do not check this in set to 1 until we are done!!
+// For now it is set only in Test Target.cs files like this:
+//   GlobalDefinitions.Add("UE_RIGIDPHYSICS_API_ENABLED=1");
+//
+#ifndef UE_RIGIDPHYSICS_API_ENABLED
+#define UE_RIGIDPHYSICS_API_ENABLED 0
+#endif
+
+#if UE_RIGIDPHYSICS_API_ENABLED
+
+#include "Templates/SharedPointer.h"
+
+#if DO_CHECK
+
+// Enable checks in the Rigid accessors. This can add significant overhead.
+#define UE_RIGIDPHYSICS_CHECK(expr) check(expr)
+#define UE_RIGIDPHYSICS_CHECKF(expr, format, ...) checkf(expr, format, ##__VA_ARGS__)
+
+// Enable scene handle safety where we assert that the handle is still valid when accessed.
+// NOTE: When disabled, the scene handle wraps a raw pointer and we have no way of knowing
+// whether the referenced object still exists or not. It is therefre an error to ever use a 
+// scene handle after the scene is destroyed, even if UE_RIGIDPHYSICS_HANDLE_SAFTEY_ENABLED=1
+// Safety comes at the cost of pinning a weak pointer whenever the scene is accessed, which
+// can add up if the scene is pinned many times. Shipping builds should probably disable
+// handle safety unless you are trying to track down a handle lifetime issue.
+#define UE_RIGIDPHYSICS_HANDLE_SAFTEY_ENABLED 1
+
+#else
+#define UE_RIGIDPHYSICS_CHECK(X)
+#define UE_RIGIDPHYSICS_CHECKF(expr, format, ...)
+#define UE_RIGIDPHYSICS_HANDLE_SAFTEY_ENABLED 0
+#endif
+
+
+// NOTES
+// 
+// Using SharedPtr and WeakPtr in a few places where they are technically required, 
+// but in practice may not be (e.g., FBodyHandle). However this is hidden as
+// an implementation detail so we have the option to switch to raw pointers if we
+// can guarantee lifetime, or retain pointer tracking in non-ship builds, etc.
+//
+
+// Math/Bounds.h
+template<typename T> struct TBounds;
+using FBounds3f = TBounds< float >;
+using FBounds3d = TBounds< double >;
+
+namespace UE::Physics
+{
+	class FRigidBodyHandle;
+	class FRigidBodyId;
+	class FRigidBodyContainerHandle;
+	class FRigidBodyContainerId;
+	class FRigidQueryDesc;
+	class FRigidQueryHit;
+	class FRigidSceneHandle;
+
+	class IRigidBody;
+	class IRigidBodyContainer;
+	class IRigidQuery;
+	class IRigidScene;
+	class IRigidSceneModifier;
+
+	template<typename ContextType> class TConstLockedRigidSceneHandle;
+	template<typename ContextType> class TRigidScenePtr;
+	template<typename ContextType> class TRigidScenePtrImpl;
+	template<typename ContextType> class TRigidBodyPtr;
+	template<typename ContextType> class TRigidBodyPtrImpl;
+	template<typename ContextType> class TRigidBodyContainerPtr;
+	template<typename ContextType> class TRigidBodyContainerPtrImpl;
+
+	using FRigidSceneSharedPtr = TSharedPtr<IRigidScene, ESPMode::ThreadSafe>;
+	using FRigidSceneWeakPtr = TWeakPtr<IRigidScene, ESPMode::ThreadSafe>;
+	using FRigidBodyContainerSharedPtr = TSharedPtr<IRigidBodyContainer, ESPMode::ThreadSafe>;
+	using FRigidBodyContainerWeakPtr = TWeakPtr<IRigidBodyContainer, ESPMode::ThreadSafe>;
+
+	using FRigidNamePtr = TSharedPtr<FString, ESPMode::ThreadSafe>;
+}
+
+namespace UE::Physics
+{
+	// Physics execution is split into Game and Simulation contexts. In async
+	// modes, these are on different threads.
+	enum class EContextEnv
+	{
+		Game,
+		Sim,
+	};
+
+	// The context must be locked before reading or writing to any objects in physics.
+	// The lock may be read-only or read-write.
+	enum class EContextLockType
+	{
+		None,
+		ReadOnly,
+		ReadWrite,
+	};
+
+	class FSimContextRW;
+	template<EContextLockType LockTypeParam> class TGameContext;
+	using FGameContextRO = TGameContext<EContextLockType::ReadOnly>;
+	using FGameContextRW = TGameContext<EContextLockType::ReadWrite>;
+}
+
+namespace UE::Physics::Private
+{
+	class UE_INTERNAL FRigidAccessor;
+}
+
+#endif // UE_RIGIDPHYSICS_API_ENABLED
\ No newline at end of file
diff --git a/Engine/Source/Runtime/Experimental/RigidPhysics/Public/RigidPhysics/RigidLog.h b/Engine/Source/Runtime/Experimental/RigidPhysics/Public/RigidPhysics/RigidLog.h
new file mode 100644
--- /dev/null
+++ b/Engine/Source/Runtime/Experimental/RigidPhysics/Public/RigidPhysics/RigidLog.h
@@ -0,0 +1,25 @@
+// Copyright Epic Games, Inc. All Rights Reserved.
+#pragma once
+
+#include "RigidPhysics/RigidFwd.h"
+
+#if UE_RIGIDPHYSICS_API_ENABLED
+
+#include "Logging/StructuredLog.h"
+
+RIGIDPHYSICS_API DECLARE_LOG_CATEGORY_EXTERN(LogChaosRB, Log, All);
+
+namespace UE::Physics::Private
+{
+	template<typename FuncNameType, typename TypeNameType>
+	UE_INTERNAL inline void LogUnsupportedAPI(const FuncNameType& FuncName, const TypeNameType& TypeName)
+	{
+		// This error means a function is being called in a context where it is not supported.
+		// E.g., Calling SetMass from the physics thread.
+		UE_LOGFMT(LogChaosRB, Warning, "{0} not supported on {1}", FuncName, TypeName);
+	}
+}
+
+#define RIGIDPHYSICS_API_UNSUPPORTED() UE::Physics::Private::LogUnsupportedAPI(__FUNCTION__, GetTypeName());
+
+#endif // UE_RIGIDPHYSICS_API_ENABLED
\ No newline at end of file
diff --git a/Engine/Source/Runtime/Experimental/RigidPhysics/Public/RigidPhysics/RigidModifier.h b/Engine/Source/Runtime/Experimental/RigidPhysics/Public/RigidPhysics/RigidModifier.h
new file mode 100644
--- /dev/null
+++ b/Engine/Source/Runtime/Experimental/RigidPhysics/Public/RigidPhysics/RigidModifier.h
@@ -0,0 +1,50 @@
+// Copyright Epic Games, Inc. All Rights Reserved.
+#pragma once
+
+#include "RigidPhysics/RigidContext.h"
+#include "Misc/Crc.h"
+#include "Templates/SharedPointer.h"
+
+#if UE_RIGIDPHYSICS_API_ENABLED
+
+namespace UE::Physics
+{
+	class IRigidSceneModifier
+	{
+	public:
+		IRigidSceneModifier()
+		{
+		}
+
+		virtual ~IRigidSceneModifier()
+		{
+		}
+
+		// Called from the game context before physics is ticked for this frame. Can be used to
+		// update inputs to the objects used in the physics context.
+		virtual void PreSimulate(TRigidScenePtr<FGameContextRW>& ScenePtr)
+		{
+		}
+
+		// Call from the game context after physics is ticked for this frame, but before the
+		// rigid body results have been sent out to the game framework objects (Actors, Components, etc.)
+		virtual void PostSimulate(TRigidScenePtr<FGameContextRW>& ScenePtr)
+		{
+		}
+
+		// Called from the simulation context before each simulation tick. At this point you can add forces, move bodies,
+		// or do pretty much anything to the scene or objects in it, and the results will impact the next tick.
+		virtual void SimPreTick(TRigidScenePtr<FSimContextRW>& ScenePtr)
+		{
+		}
+
+		// Called from the simulation context at the end of each simulation tick. At this point you can modify physics
+		// object state and the altered results will be reported back to the game, and be used for the next tick.
+		virtual void SimPostTick(TRigidScenePtr<FSimContextRW>& ScenePtr)
+		{
+		}
+	};
+
+}
+
+#endif // UE_RIGIDPHYSICS_API_ENABLED
\ No newline at end of file
diff --git a/Engine/Source/Runtime/Experimental/RigidPhysics/Public/RigidPhysics/RigidScene.h b/Engine/Source/Runtime/Experimental/RigidPhysics/Public/RigidPhysics/RigidScene.h
new file mode 100644
--- /dev/null
+++ b/Engine/Source/Runtime/Experimental/RigidPhysics/Public/RigidPhysics/RigidScene.h
@@ -0,0 +1,422 @@
+// Copyright Epic Games, Inc. All Rights Reserved.
+#pragma once
+
+#include "RigidPhysics/RigidFwd.h"
+
+#if UE_RIGIDPHYSICS_API_ENABLED
+
+#include "RigidPhysics/RigidContext.h"
+#include "Templates/SharedPointer.h"
+
+namespace UE::Physics
+{
+	inline FRigidNamePtr MakeRigidName(const TCHAR* InName)
+	{
+		return MakeShared<FString, ESPMode::ThreadSafe>(FString(InName));
+	}
+
+	template<typename... ArgTypes>
+	inline FRigidNamePtr MakeRigidName(const TCHAR* InFmt, ArgTypes... Args)
+	{
+		// TODO: Optimize
+		FStringFormatOrderedArguments ArgsArray = { std::forward<ArgTypes>(Args)... };
+		return MakeShared<FString, ESPMode::ThreadSafe>(FString::Format(InFmt, ArgsArray));
+	}
+
+	// The interface to a rigid body scene. The rigid body scene owns the physics simulations and
+	// manages a set of rigid body containers which in turn manage sets of rigid bodies.
+	class RIGIDPHYSICS_API IRigidScene
+	{
+	public:
+		IRigidScene();
+		virtual ~IRigidScene();
+
+		// Lock the scene for read or write if supported (GameContext only), otherwise do nothing
+		virtual void Lock(EContextLockType InLockType)
+		{
+			RIGIDPHYSICS_API_UNSUPPORTED();
+		}
+
+		// Unlock the scene if supported (GameContext only), otherwise do nothing
+		virtual void Unlock(EContextLockType InLockType)
+		{
+			RIGIDPHYSICS_API_UNSUPPORTED();
+		}
+
+		// Create a rigid body. It will take part in the simulation until RegisterRigidBody is called.
+		virtual IRigidBody* CreateRigidBody(const FRigidNamePtr& InName, ERigidMovementType InMovementType)
+		{
+			RIGIDPHYSICS_API_UNSUPPORTED();
+			return nullptr;
+		}
+
+		// Destroy a rigid body. It will no longer take part in the simulation.
+		virtual void DestroyRigidBody(const FRigidBodyHandle& RigidBodyHandle)
+		{
+			RIGIDPHYSICS_API_UNSUPPORTED();
+		}
+
+		// Add a rigid body to the scene. It will now be colideable and simulated as appropriate.
+		virtual void RegisterRigidBody(const FRigidBodyHandle& BodyHandle)
+		{
+			RIGIDPHYSICS_API_UNSUPPORTED();
+		}
+
+		// Remove a rigid body from the scene.
+		virtual void UnregisterRigidBody(const FRigidBodyHandle& BodyHandle)
+		{
+			RIGIDPHYSICS_API_UNSUPPORTED();
+		}
+
+		// Get the specified rigid body container.
+		virtual IRigidBodyContainer* GetRigidBodyContainer(const FRigidBodyContainerId& InRigidBodyContainerId)
+		{
+			RIGIDPHYSICS_API_UNSUPPORTED();
+			return nullptr;
+		}
+
+		virtual void RegisterModifier(IRigidSceneModifier* Modifier)
+		{
+			RIGIDPHYSICS_API_UNSUPPORTED();
+		}
+		
+		virtual void UnregisterModifier(IRigidSceneModifier* Modifier)
+		{
+			RIGIDPHYSICS_API_UNSUPPORTED();
+		}
+
+		virtual bool Query(const FRigidQueryDesc& QueryDesc, TArray<FRigidQueryHit>& OutHits)
+		{
+			RIGIDPHYSICS_API_UNSUPPORTED();
+			return false;
+		}
+
+		// Advance the simulation by Dt. If the implementation is asynchronous, this will return immediately
+		virtual void Tick(float Dt)
+		{
+			RIGIDPHYSICS_API_UNSUPPORTED();
+		}
+
+	protected:
+
+		virtual const char* GetTypeName() const
+		{
+			return "Unknown";
+		}
+	};
+
+
+	// Exists only to give pointer semantics to TRigidScenePtr.
+	// See TRigidScenePtr for comments.
+	template<typename ContextType>
+	class TRigidScenePtrImpl
+	{
+	public:
+		using FContext = ContextType;
+
+		TRigidScenePtrImpl()
+			: SceneRaw(nullptr)
+		{
+		}
+
+		~TRigidScenePtrImpl()
+		{
+			Reset();
+		}
+
+		// Impl in RigidBody.h
+		TRigidBodyPtr<FContext> CreateRigidBody(const FRigidNamePtr& InName, ERigidMovementType InMovementType) const;
+
+		// Impl in RigidBody.h
+		void DestroyRigidBody(const FRigidBodyHandle& BodyHandle) const;
+
+		// Impl in RigidBody.h
+		TRigidBodyContainerPtr<FContext> GetRigidBodyContainer(const FRigidBodyContainerId& InBodyContainerId) const;
+
+		// Impl in RigidBody.h
+		void RegisterRigidBody(const FRigidBodyHandle& BodyHandle);
+
+		// Impl in RigidBody.h
+		void UnregisterRigidBody(const FRigidBodyHandle& BodyHandle);
+
+		void RegisterModifier(IRigidSceneModifier* Modifier)
+		{
+			SceneRaw->RegisterModifier(Modifier);
+		}
+
+		void UnregisterModifier(IRigidSceneModifier* Modifier)
+		{
+			SceneRaw->UnregisterModifier(Modifier);
+		}
+
+		bool Query(const FRigidQueryDesc& QueryDesc, TArray<FRigidQueryHit>& OutHits)  const
+		{
+			return SceneRaw->Query(QueryDesc, OutHits);
+		}
+
+		void Tick(float Dt)
+		{
+			SceneRaw->Tick(Dt);
+		}
+
+	private:
+		friend class TRigidScenePtr<FContext>;
+		friend class Private::FRigidAccessor;
+
+		TRigidScenePtrImpl(IRigidScene* InScene)
+			: SceneRaw(nullptr)
+		{
+			Init(InScene);
+		}
+
+		TRigidScenePtrImpl(const TRigidScenePtrImpl& InImpl)
+			: SceneRaw(nullptr)
+		{
+			*this = InImpl;
+		}
+
+		TRigidScenePtrImpl(TRigidScenePtrImpl&& InImpl)
+			: SceneRaw(nullptr)
+		{
+			// Take over incoming scene lock - do not add a ref
+			SceneRaw = InImpl.SceneRaw;
+			InImpl.SceneRaw = nullptr;
+		}
+
+		TRigidScenePtrImpl& operator=(const TRigidScenePtrImpl& InImpl)
+		{
+			if (InImpl.SceneRaw != SceneRaw)
+			{
+				// Undo current lock if there is one
+				Reset();
+
+				// Scene will already be lock - add a ref
+				Init(InImpl.SceneRaw);
+			}
+			return *this;
+		}
+
+		TRigidScenePtrImpl& operator=(TRigidScenePtrImpl&& InImpl)
+		{
+			// Undo current lock if there is one
+			Reset();
+
+			// Take over incoming scene lock - do not add a ref
+			SceneRaw = InImpl.SceneRaw;
+			InImpl.SceneRaw = nullptr;
+
+			return *this;
+		}
+
+		bool IsValid() const
+		{
+			return (SceneRaw != nullptr);
+		}
+
+		void Init(IRigidScene* InScene)
+		{
+			SceneRaw = InScene;
+			if (SceneRaw != nullptr)
+			{
+				SceneRaw->Lock(FContext::LockType);
+			}
+		}
+
+		void Reset()
+		{
+			if (SceneRaw != nullptr)
+			{
+				SceneRaw->Unlock(FContext::LockType);
+				SceneRaw = nullptr;
+			}
+		}
+
+		IRigidScene* SceneRaw;
+	};
+
+
+	// A context-aware IRigidScene pointer that simply forwards all context-legal operations to IRigidScene. 
+	// Provides access to a restricted Scene API where the restrictions depend on the context 
+	// (template-parameter), and attempts to call disallowed functions (in the context) produce compiler errors.
+	// 
+	// Uses pointer semantics and may be "null". Check IsValid() before dereferencing. Also implicitly casts 
+	// to boolean for validity checks.
+	// 
+	// Cannot be copied or stored - use FRigidSceneHandle for that.
+	//
+	// Where appropriate, the Scene Ptr holds a lock on the scene. The type of lock depends on the context.
+	// E.g., in the async implementation, for GameContexts it may be read-only or write-enabled, but the
+	// PhysicsContext does not support multi-threaded access, and therefore there is no lock.
+	template<typename ContextType>
+	class TRigidScenePtr
+	{
+	public:
+		using FContext = ContextType;
+		using FRigidScenePtrImpl = TRigidScenePtrImpl<FContext>;
+
+		TRigidScenePtr() = default;
+		~TRigidScenePtr() = default;
+
+		void Reset()
+		{
+			SceneImpl.Reset();
+		}
+
+		bool IsValid() const
+		{
+			return SceneImpl.IsValid();
+		}
+
+		operator bool() const
+		{
+			return IsValid();
+		}
+
+		// Pointer Semantics
+		const FRigidScenePtrImpl* operator->() const requires (!FContext::bWriteEnabled)
+		{
+			UE_RIGIDPHYSICS_CHECKF(IsValid(), TEXT("Attempt to dereference null scene handle. Missing IsValid() check"));
+			return &SceneImpl;
+		}
+
+		// Pointer Semantics
+		FRigidScenePtrImpl* operator->() const requires (FContext::bWriteEnabled)
+		{
+			UE_RIGIDPHYSICS_CHECKF(IsValid(), TEXT("Attempt to dereference null scene handle. Missing IsValid() check"));
+			return const_cast<FRigidScenePtrImpl*>(&SceneImpl);
+		}
+
+	private:
+		friend class FRigidSceneHandle;
+		friend class Private::FRigidAccessor;
+		friend class TRigidScenePtrImpl<FContext>;
+		friend class TRigidBodyContainerPtrImpl<FContext>;
+		friend class TRigidBodyPtrImpl<FContext>;
+
+		TRigidScenePtr(IRigidScene* InScene)
+			: SceneImpl(InScene)
+		{
+		}
+
+		TRigidScenePtr(const TRigidScenePtr<ContextType>& InScene)
+			: SceneImpl(InScene.SceneImpl)
+		{
+		}
+
+		TRigidScenePtr(TRigidScenePtr<ContextType>&& InScene)
+			: SceneImpl(MoveTemp(InScene.SceneImpl))
+		{
+		}
+
+		TRigidScenePtr(const TRigidScenePtrImpl<ContextType>& InSceneImpl)
+			: SceneImpl(InSceneImpl)
+		{
+		}
+
+		TRigidScenePtr& operator=(const TRigidScenePtr& R) = default;
+		TRigidScenePtr& operator=(TRigidScenePtr&& R) = default;
+
+		FRigidScenePtrImpl SceneImpl;
+	};
+
+
+	class FRigidObjectId
+	{
+	public:
+		FRigidObjectId()
+		{
+			IntId = -1;
+		}
+
+		explicit FRigidObjectId(uint64 InId)
+		{
+			IntId = InId;
+		}
+
+		friend bool operator==(const FRigidObjectId& L, const FRigidObjectId& R) = default;
+		friend bool operator<(const FRigidObjectId& L, const FRigidObjectId& R)
+		{
+			return L.IntId < R.IntId;
+		}
+
+	private:
+		friend class Private::FRigidAccessor;
+
+		uint64 IntId;
+	};
+
+
+	// A handle to a rigid body scene. FRigidSceneHandle is generally only used in the GameContext
+	// because all PhysicsContect use cases (callbacks, modifiers, etc.) will be passed their
+	// scene as a TRigidScenePtr.
+	class FRigidSceneHandle
+	{
+	public:
+		FRigidSceneHandle()
+			: SceneRaw(nullptr)
+		{
+		}
+
+		FRigidSceneHandle(const FRigidSceneSharedPtr& InScene)
+			: SceneRaw(InScene.Get())
+#if UE_RIGIDPHYSICS_HANDLE_SAFTEY_ENABLED
+			, SceneWeak(InScene)
+#endif
+		{
+		}
+
+		void Reset()
+		{
+			SceneRaw = nullptr;
+
+#if UE_RIGIDPHYSICS_HANDLE_SAFTEY_ENABLED
+			SceneWeak.Reset();
+#endif
+		}
+
+		// Game-thread read-only scene lock
+		TRigidScenePtr<FGameContextRO> LockRO() const
+		{
+			return Lock<FGameContextRO>();
+		}
+
+		// Game-thread read-write scene lock
+		TRigidScenePtr<FGameContextRW> LockRW() const
+		{
+			return Lock<FGameContextRW>();
+		}
+
+		friend bool operator==(const FRigidSceneHandle& L, const FRigidSceneHandle& R) = default;
+		friend bool operator!=(const FRigidSceneHandle& L, const FRigidSceneHandle& R) = default;
+
+	private:
+		friend class Private::FRigidAccessor;
+
+		template<typename ContextType>
+		TRigidScenePtr<ContextType> Lock() const
+		{
+			// NOTE: SceneHandles are not used in PhysicsContext callbacks which pass
+			// TRigidScenePtr objects directly to the APIs. 
+			static_assert(ContextType::ContextEnv == EContextEnv::Game, "FRigidSceneHandle::Lock() is for the GameContext");
+
+#if UE_RIGIDPHYSICS_HANDLE_SAFTEY_ENABLED
+			// NOTE: In order to support the CHAOSRB_HANDLE_SAFTEY_ENABLED=0 optimization
+			// SceneHandles cannot be used after the scene has expired. The onus is on the 
+			// user to ensure that this is the case, but when we have handle safety
+			// enabled we can assert on this. Without the weak pointer, there's no way 
+			// to know that the raw pointer held in SceneWrapper points to a deleted object.
+			check(SceneRaw == SceneWeak.Pin().Get());
+#endif
+
+			return TRigidScenePtr<ContextType>(SceneRaw);
+		}
+
+		IRigidScene* SceneRaw;
+#if UE_RIGIDPHYSICS_HANDLE_SAFTEY_ENABLED
+		FRigidSceneWeakPtr SceneWeak;
+#endif
+	};
+
+}
+
+#endif // UE_RIGIDPHYSICS_API_ENABLED
\ No newline at end of file
diff --git a/Engine/Source/Runtime/Experimental/RigidPhysics/RigidPhysics.Build.cs b/Engine/Source/Runtime/Experimental/RigidPhysics/RigidPhysics.Build.cs
new file mode 100644
--- /dev/null
+++ b/Engine/Source/Runtime/Experimental/RigidPhysics/RigidPhysics.Build.cs
@@ -0,0 +1,22 @@
+// Copyright Epic Games, Inc. All Rights Reserved.
+
+using System.Linq;
+using System.Collections.Generic;
+
+namespace UnrealBuildTool.Rules
+{
+    public class RigidPhysics : ModuleRules
+    {
+        public RigidPhysics(ReadOnlyTargetRules Target) : base(Target)
+        {
+            PublicDependencyModuleNames.AddRange(
+                new string[] {
+                "Core",
+                "CoreUObject",
+				}
+			);
+
+			CppCompileWarningSettings.UnsafeTypeCastWarningLevel = WarningLevel.Error;
+		}
+    }
+}
