Add UToc hash to I/O store on-demand TOC JSON output when using -ListTocs command.

#rnx

diff --git a/Engine/Source/Runtime/Experimental/IoStore/OnDemand/Private/Tool/ListTocs.cpp b/Engine/Source/Runtime/Experimental/IoStore/OnDemand/Private/Tool/ListTocs.cpp
--- a/Engine/Source/Runtime/Experimental/IoStore/OnDemand/Private/Tool/ListTocs.cpp
+++ b/Engine/Source/Runtime/Experimental/IoStore/OnDemand/Private/Tool/ListTocs.cpp
@@ -243,6 +243,7 @@ static void ToJson(FJsonWriter JsonWriter, const FOnDemandToc& Toc, EOnDemandToc
 			JsonWriter->WriteObjectStart();
 			JsonWriter->WriteValue(TEXT("Name"), Container.ContainerName);
 			JsonWriter->WriteValue(TEXT("EncryptionKeyGuid"), Container.EncryptionKeyGuid);
+			JsonWriter->WriteValue(TEXT("UTocHash"), LexToString(Container.UTocHash));
 
 			if (EnumHasAnyFlags(Options, EOnDemandTocJsonOptions::TocEntries))
 			{
