Fix crash if an empty parameter if provided.

#rnx

- Use ::StartWith rather than accessing the first character as the former checks for string length as well and will just return false for an empty string.

diff --git a/Engine/Source/Runtime/Experimental/IoStore/OnDemand/Private/Tool/Command.h b/Engine/Source/Runtime/Experimental/IoStore/OnDemand/Private/Tool/Command.h
--- a/Engine/Source/Runtime/Experimental/IoStore/OnDemand/Private/Tool/Command.h
+++ b/Engine/Source/Runtime/Experimental/IoStore/OnDemand/Private/Tool/Command.h
@@ -86,7 +86,7 @@ private:
 template <CArgTypes Type>
 Type FContext::Get(FStringView Name) const
 {
-	if (Name.GetData()[0] == '-')
+	if (Name.StartsWith(TEXT('-')))
 	{
 		return Get(Name, Type{});
 	}
