Fix spelling mistake in the output when running 'IASTool ListTocs'

#rnx

diff --git a/Engine/Source/Runtime/Experimental/IoStore/OnDemand/Private/Tool/ListTocs.cpp b/Engine/Source/Runtime/Experimental/IoStore/OnDemand/Private/Tool/ListTocs.cpp
--- a/Engine/Source/Runtime/Experimental/IoStore/OnDemand/Private/Tool/ListTocs.cpp
+++ b/Engine/Source/Runtime/Experimental/IoStore/OnDemand/Private/Tool/ListTocs.cpp
@@ -224,7 +224,7 @@ static void ToJson(FJsonWriter JsonWriter, const FOnDemandToc& Toc, EOnDemandToc
 		JsonWriter->WriteValue(TEXT("Magic"), Toc.Header.Magic);
 		JsonWriter->WriteValue(TEXT("Version"), uint64(Toc.Header.Version));
 		JsonWriter->WriteValue(TEXT("BlockSize"), uint64(Toc.Header.BlockSize));
-		JsonWriter->WriteValue(TEXT("CompressonFormat"), Toc.Header.CompressionFormat);
+		JsonWriter->WriteValue(TEXT("CompressionFormat"), Toc.Header.CompressionFormat);
 		JsonWriter->WriteValue(TEXT("ChunksDirectory"), Toc.Header.ChunksDirectory);
 	}
 	JsonWriter->WriteObjectEnd();
