Small strategic separation of reception and impl

#rnx

diff --git a/Engine/Source/Runtime/Experimental/IoStore/HttpClient/Private/Loop.inl b/Engine/Source/Runtime/Experimental/IoStore/HttpClient/Private/Loop.inl
--- a/Engine/Source/Runtime/Experimental/IoStore/HttpClient/Private/Loop.inl
+++ b/Engine/Source/Runtime/Experimental/IoStore/HttpClient/Private/Loop.inl
@@ -32,6 +32,7 @@ private:
 		std::atomic<uint64> Value;
 	};
 
+	void					UpdateRecv(uint32 NowMs, uint32 RecvSizeKiB);
 	FBucket					Buckets[8] = {};
 	FBucket::FData			Accumulator = {};
 	std::atomic<uint32>		TotalRecvKiB = 0;
@@ -41,6 +42,12 @@ private:
 
 ////////////////////////////////////////////////////////////////////////////////
 void FRecvStats::Update(uint32 NowMs, uint32 RecvSizeKiB)
+{
+	UpdateRecv(NowMs, RecvSizeKiB);
+}
+
+////////////////////////////////////////////////////////////////////////////////
+void FRecvStats::UpdateRecv(uint32 NowMs, uint32 RecvSizeKiB)
 {
 	if (RecvSizeKiB == 0)
 	{
