Moved helper free function

#rnx

diff --git a/Engine/Source/Runtime/Experimental/IoStore/HttpClient/Private/Loop.inl b/Engine/Source/Runtime/Experimental/IoStore/HttpClient/Private/Loop.inl
--- a/Engine/Source/Runtime/Experimental/IoStore/HttpClient/Private/Loop.inl
+++ b/Engine/Source/Runtime/Experimental/IoStore/HttpClient/Private/Loop.inl
@@ -1172,26 +1172,6 @@ void FHostGroup::AddActivity(FActivityNode* Activity)
 ////////////////////////////////////////////////////////////////////////////////
 static const FEventLoop::FRequestParams GDefaultParams;
 
-////////////////////////////////////////////////////////////////////////////////
-static uint32 GetNowMs()
-{
-	// 4.2MM seconds will give us 50 days of uptime.
-	static uint64 Freq = 0;
-	static uint64 Base = 0;
-	if (Freq == 0)
-	{
-		Freq = uint64(1.0 / FPlatformTime::GetSecondsPerCycle64());
-		Base = FPlatformTime::Cycles64();
-	}
-	uint64 NowBig = ((FPlatformTime::Cycles64() - Base) * 1000) / Freq;
-	uint32 NowMs = uint32(NowBig);
-	check(NowMs == NowBig);
-	return NowMs + 1;
-	// +1 bias ----^ ...so 0 can be considered "not set"
-}
-
-
-
 ////////////////////////////////////////////////////////////////////////////////
 class FEventLoop::FImpl
 {
diff --git a/Engine/Source/Runtime/Experimental/IoStore/HttpClient/Private/Misc.inl b/Engine/Source/Runtime/Experimental/IoStore/HttpClient/Private/Misc.inl
--- a/Engine/Source/Runtime/Experimental/IoStore/HttpClient/Private/Misc.inl
+++ b/Engine/Source/Runtime/Experimental/IoStore/HttpClient/Private/Misc.inl
@@ -68,6 +68,28 @@ IOSTOREHTTPCLIENT_API const void*	GetIaxTraceChannel()	{ return nullptr; }
 static IAS_CVAR(int32, RecvWorkThresholdKiB,80,		"Threshold of data remaining at which next request is sent (in KiB)");
 static IAS_CVAR(int32, IdleMs,				50'000,	"Time in milliseconds to close idle connections or fail waits");
 
+
+
+////////////////////////////////////////////////////////////////////////////////
+uint32 GetNowMs()
+{
+	// 4.2MM seconds will give us 50 days of uptime.
+	static uint64 Freq = 0;
+	static uint64 Base = 0;
+	if (Freq == 0)
+	{
+		Freq = uint64(1.0 / FPlatformTime::GetSecondsPerCycle64());
+		Base = FPlatformTime::Cycles64();
+	}
+	uint64 NowBig = ((FPlatformTime::Cycles64() - Base) * 1000) / Freq;
+	uint32 NowMs = uint32(NowBig);
+	check(NowMs == NowBig);
+	return NowMs + 1;
+	// +1 bias ----^ ...so 0 can be considered "not set"
+}
+
+
+
 ////////////////////////////////////////////////////////////////////////////////
 class FOutcome
 {
