No longer applicable as the send has moved

#rnx

diff --git a/Engine/Source/Runtime/Experimental/IoStore/HttpClient/Private/Loop.inl b/Engine/Source/Runtime/Experimental/IoStore/HttpClient/Private/Loop.inl
--- a/Engine/Source/Runtime/Experimental/IoStore/HttpClient/Private/Loop.inl
+++ b/Engine/Source/Runtime/Experimental/IoStore/HttpClient/Private/Loop.inl
@@ -822,12 +822,6 @@ void FPeerGroup::RecvInternal(FTickState& State)
 	IsKeepAlive &= uint8(Activity->GetTransaction()->IsKeepAlive() == true);
 	bWaiting |= (Outcome.IsWaiting() && Outcome.IsWaitData());
 
-	// New work may have fell foul of peer issues
-	if (!Peer.IsValid())
-	{
-		return;
-	}
-
 	// If there was no data available this is far as receiving can go
 	if (Outcome.IsWaiting())
 	{
