Changed FIoHttpRange Expand to be const and added ExpandInline.

#rnx

diff --git a/Engine/Source/Runtime/Experimental/IoStore/OnDemandCore/Public/IO/HttpIoDispatcher.h b/Engine/Source/Runtime/Experimental/IoStore/OnDemandCore/Public/IO/HttpIoDispatcher.h
--- a/Engine/Source/Runtime/Experimental/IoStore/OnDemandCore/Public/IO/HttpIoDispatcher.h
+++ b/Engine/Source/Runtime/Experimental/IoStore/OnDemandCore/Public/IO/HttpIoDispatcher.h
@@ -43,22 +43,23 @@ ENUM_CLASS_FLAGS(EIoHttpFlags);
 /** Represents a range within a resource. */
 struct FIoHttpRange
 {
-								FIoHttpRange() = default;
-	explicit inline				FIoHttpRange(uint32 InMin, uint32 InMax);
+										FIoHttpRange() = default;
+	explicit inline						FIoHttpRange(uint32 InMin, uint32 InMax);
 
-	bool						IsValid()	const	{ return Min <= Max; }
-	uint32						GetMin()	const	{ return Min; }
-	uint32						GetMax()	const	{ return Max; }
-	uint32						GetSize()	const	{ return (Max - Min); }
-	inline FIoHttpRange&		Expand(const FIoHttpRange& Other);
+	[[nodiscard]] bool					IsValid()	const	{ return Min <= Max; }
+	[[nodiscard]] uint32				GetMin()	const	{ return Min; }
+	[[nodiscard]] uint32				GetMax()	const	{ return Max; }
+	[[nodiscard]] uint32				GetSize()	const	{ return (Max - Min); }
+	[[nodiscard]] inline FIoHttpRange	Expand(const FIoHttpRange& Other) const;
+	inline FIoHttpRange&				ExpandInline(const FIoHttpRange& Other);
 
-	inline FIoOffsetAndLength	ToOffsetAndLength() const;
-	static inline FIoHttpRange	FromOffsetAndLength(const FIoOffsetAndLength& OffsetAndLength);
+	[[nodiscard]] inline FIoOffsetAndLength ToOffsetAndLength() const;
+	static inline FIoHttpRange			FromOffsetAndLength(const FIoOffsetAndLength& OffsetAndLength);
 
-	bool						operator==(const FIoHttpRange& Other)	{ return Min == Other.Min && Max == Other.Max; }
-	bool						operator!=(const FIoHttpRange& Other)	{ return Min != Other.Min || Max != Other.Max; }
-	FIoHttpRange&				operator+(const FIoHttpRange& Other)	{ return Expand(Other); }
-	FIoHttpRange&				operator+=(const FIoHttpRange& Other)	{ return Expand(Other); }
+	[[nodiscard]] bool					operator==(const FIoHttpRange& Other)	{ return Min == Other.Min && Max == Other.Max; }
+	[[nodiscard]] bool					operator!=(const FIoHttpRange& Other)	{ return Min != Other.Min || Max != Other.Max; }
+	[[nodiscard]] FIoHttpRange			operator+(const FIoHttpRange& Other)	{ return Expand(Other); }
+	FIoHttpRange&						operator+=(const FIoHttpRange& Other)	{ return ExpandInline(Other); }
 
 private:
 	uint32 Min = MAX_uint32;
@@ -71,7 +72,7 @@ FIoHttpRange::FIoHttpRange(uint32 InMin, uint32 InMax)
 {
 }
 
-FIoHttpRange& FIoHttpRange::Expand(const FIoHttpRange& Other)
+FIoHttpRange& FIoHttpRange::ExpandInline(const FIoHttpRange& Other)
 {
 	if (IsValid())
 	{
@@ -87,6 +88,13 @@ FIoHttpRange& FIoHttpRange::Expand(const FIoHttpRange& Other)
 	return *this;
 }
 
+FIoHttpRange FIoHttpRange::Expand(const FIoHttpRange& Other) const
+{
+	FIoHttpRange Expanded(*this);
+	Expanded.ExpandInline(Other);
+	return Expanded;
+}
+
 FIoOffsetAndLength FIoHttpRange::ToOffsetAndLength() const
 {
 	return IsValid() ? FIoOffsetAndLength(Min, Max - Min) : FIoOffsetAndLength(0, 0);
