Track time to first byte stat

#rnx

diff --git a/Engine/Source/Runtime/Experimental/IoStore/HttpClient/Internal/IO/Http/Client.h b/Engine/Source/Runtime/Experimental/IoStore/HttpClient/Internal/IO/Http/Client.h
--- a/Engine/Source/Runtime/Experimental/IoStore/HttpClient/Internal/IO/Http/Client.h
+++ b/Engine/Source/Runtime/Experimental/IoStore/HttpClient/Internal/IO/Http/Client.h
@@ -234,9 +234,11 @@ public:
 	UE_API	uint32	GetRecvKiBps(uint32 Index) const;
 	UE_API	void	GetRecvKiBps(uint32 (&Out)[RecvBucketCount]) const;
 	UE_API	uint32	GetRecvKiBps() const;
+	UE_API	uint32	GetTimeToFirstByteMs() const;
 
 private:
 	friend			class FRecvStats;
+	uint64			TtfbSamples;
 	uint32			RecvKiB;
 	uint16			RecvKiBps[RecvBucketCount];
 };
diff --git a/Engine/Source/Runtime/Experimental/IoStore/HttpClient/Private/Loop.inl b/Engine/Source/Runtime/Experimental/IoStore/HttpClient/Private/Loop.inl
--- a/Engine/Source/Runtime/Experimental/IoStore/HttpClient/Private/Loop.inl
+++ b/Engine/Source/Runtime/Experimental/IoStore/HttpClient/Private/Loop.inl
@@ -12,7 +12,7 @@ class FRecvStats
 {
 public:
 							FRecvStats() = default;
-	void					Update(uint32 NowMs, uint32 RecvSizeKiB);
+	void					Update(uint32 NowMs, uint32 RecvSizeKiB, uint32 TtfbMs);
 	void					Flush();
 	void					Sample(uint32 NowMs, FLoopPerf& Out) const;
 
@@ -33,17 +33,20 @@ private:
 	};
 
 	void					UpdateRecv(uint32 NowMs, uint32 RecvSizeKiB);
+	void					UpdateTtfb(uint32 TtfbMs);
 	FBucket					Buckets[8] = {};
 	FBucket::FData			Accumulator = {};
+	std::atomic<uint64>		TtfbSamples = 0;
 	std::atomic<uint32>		TotalRecvKiB = 0;
 	uint32					LastMs = 0;
 	uint8					Cursor = 0;
 };
 
 ////////////////////////////////////////////////////////////////////////////////
-void FRecvStats::Update(uint32 NowMs, uint32 RecvSizeKiB)
+void FRecvStats::Update(uint32 NowMs, uint32 RecvSizeKiB, uint32 TtfbMs)
 {
 	UpdateRecv(NowMs, RecvSizeKiB);
+	UpdateTtfb(TtfbMs);
 }
 
 ////////////////////////////////////////////////////////////////////////////////
@@ -97,6 +100,19 @@ void FRecvStats::UpdateRecv(uint32 NowMs, uint32 RecvSizeKiB)
 	Buckets[Index].Value.store(Value, std::memory_order_release);
 }
 
+////////////////////////////////////////////////////////////////////////////////
+void FRecvStats::UpdateTtfb(uint32 TtfbMs)
+{
+	if (TtfbMs == 0)
+	{
+		return;
+	}
+
+	uint64 Current = TtfbSamples.load(std::memory_order_relaxed);
+	Current = (Current << 8ull) | uint64(TtfbMs & 0xff);
+	TtfbSamples.store(Current, std::memory_order_relaxed);
+}
+
 ////////////////////////////////////////////////////////////////////////////////
 void FRecvStats::Flush()
 {
@@ -162,6 +178,9 @@ void FRecvStats::Sample(uint32 NowMs, FLoopPerf& Out) const
 	}
 
 	Out.RecvKiB = TotalRecvKiB.load(std::memory_order_relaxed);
+
+	// Time to first byte
+	Out.TtfbSamples = TtfbSamples.load(std::memory_order_relaxed);
 }
 
 
@@ -217,6 +236,25 @@ uint32 FLoopPerf::GetRecvKiBps() const
 	return Numer / Denom;
 }
 
+////////////////////////////////////////////////////////////////////////////////
+uint32 FLoopPerf::GetTimeToFirstByteMs() const
+{
+	uint64 Samples = TtfbSamples;
+	uint32 Accumulator = 0;
+	uint32 Prev = 0;
+	for (uint32 i = 0, n = 8; i < n; ++i, Samples >>= 8)
+	{
+		uint32 Value = Samples & 0xff;
+		if (Value == 0)
+		{
+			Value = Prev;
+		}
+		Accumulator += Value;
+		Prev = Value;
+	}
+	return Accumulator >> 3;
+}
+
 
 
 // {{{1 work-queue .............................................................
@@ -537,6 +575,7 @@ struct FTickState
 	int32						FailTimeoutMs;
 	uint32						NowMs;
 	uint32						RecvSizeKiB = 0;
+	uint32						TtfbMs = 0;
 };
 
 
@@ -828,6 +867,15 @@ void FPeerGroup::RecvInternal(FTickState& State)
 		return;
 	}
 
+	// If we have a time here we are timing time-to-first-byte
+	if (LastUseMs != 0)
+	{
+		uint32 TtfbMs = GetNowMs() - LastUseMs;
+		TtfbMs = (TtfbMs <= 0xff) ? TtfbMs : 0xff;
+		State.TtfbMs = (State.TtfbMs << 8) | TtfbMs;
+		LastUseMs = 0;
+	}
+
 	// If an okay result is -1, then message recv is done and content is next
 	// and it is likely that we have it.
 	if (Outcome.GetResult() == uint32(FActivity::EStage::Response))
@@ -975,6 +1023,12 @@ void FPeerGroup::TickSend(FTickState& State, FHost& Host, FPoller& Poller)
 		Peer = FHttpPeer(MoveTemp(Params));
 	}
 
+	// If we are idle we can time the time-to-first-byte
+	if ((InflightNum == 0) && Recv.IsEmpty())
+	{
+		LastUseMs = GetNowMs();
+	}
+
 	Send.Append(Pending);
 	InflightNum++;
 
@@ -1368,7 +1422,7 @@ uint32 FEventLoop::FImpl::Tick(int32 PollTimeoutMs)
 	}
 
 	// Track some stats
-	RecvStats.Update(TickState.NowMs, TickState.RecvSizeKiB);
+	RecvStats.Update(TickState.NowMs, TickState.RecvSizeKiB, TickState.TtfbMs);
 
 	// Remove groups that are idle
 	for (uint32 i = 0, n = Groups.Num(); i < n; ++i)
