Fix package chunkids printed backwards in errors.

diff --git a/Engine/Source/Runtime/Experimental/IoStore/OnDemandCore/Public/IO/OnDemandError.h b/Engine/Source/Runtime/Experimental/IoStore/OnDemandCore/Public/IO/OnDemandError.h
--- a/Engine/Source/Runtime/Experimental/IoStore/OnDemandCore/Public/IO/OnDemandError.h
+++ b/Engine/Source/Runtime/Experimental/IoStore/OnDemandCore/Public/IO/OnDemandError.h
@@ -193,6 +193,16 @@ inline void SerializeForLog(FCbWriter& Writer, const UE::UnifiedError::IoStoreOn
 	Writer.AddString(ANSITEXTVIEW("ErrorMessage"), Ctx.ErrorMessage); 
 	Writer.EndObject();
 }
+
+
+inline void SerializeForLog(FCbWriter& Writer, const FIoChunkId& ChunkId)
+{
+	Writer.BeginObject();
+	Writer.AddString(ANSITEXTVIEW("$type"), ANSITEXTVIEW("FIoChunkId"));
+	Writer.AddString(ANSITEXTVIEW("$text"), LexToString(ChunkId));
+	Writer.EndObject();
+}
+
 ////////////////////////////////////////////////////////////////////////////////
 UE_DECLARE_ERRORSTRUCT_FEATURES(IoStoreOnDemand, FChunkMissingErrorContext);
 inline void SerializeForLog(FCbWriter& Writer, const UE::UnifiedError::IoStoreOnDemand::FChunkMissingErrorContext& Ctx)
@@ -202,13 +212,15 @@ inline void SerializeForLog(FCbWriter& Writer, const UE::UnifiedError::IoStoreOn
 	Writer.AddString(ANSITEXTVIEW("$format"), TEXT("(MissingCount: {MissingCount}, ChunkIds: [{ChunkIds}])"));
 
 	Writer.AddInteger(ANSITEXTVIEW("MissingCount"), Ctx.ChunkIds.Num());
-
-	TUtf8StringBuilder<1024> Builder;
-	Builder.Join(Ctx.ChunkIds, UTF8TEXTVIEW(", "));
-	Writer.AddString(ANSITEXTVIEW("ChunkIds"), Builder);
-	
+	Writer.BeginArray(ANSITEXTVIEW("ChunkIds"));
+	for (const auto& ChunkId : Ctx.ChunkIds)
+	{
+		SerializeForLog(Writer, ChunkId);
+	}
+	Writer.EndArray();
 	Writer.EndObject();
 }
+
 ////////////////////////////////////////////////////////////////////////////////
 UE_DECLARE_ERRORSTRUCT_FEATURES(IoStoreOnDemand, FChunkHashMismatchErrorContext);
 inline void SerializeForLog(FCbWriter& Writer, const UE::UnifiedError::IoStoreOnDemand::FChunkHashMismatchErrorContext& Ctx)
