IAD installer trace markup.

#FailedValidations Preflight

diff --git a/Engine/Source/Runtime/Experimental/IoStore/OnDemand/Private/OnDemandContentInstaller.cpp b/Engine/Source/Runtime/Experimental/IoStore/OnDemand/Private/OnDemandContentInstaller.cpp
--- a/Engine/Source/Runtime/Experimental/IoStore/OnDemand/Private/OnDemandContentInstaller.cpp
+++ b/Engine/Source/Runtime/Experimental/IoStore/OnDemand/Private/OnDemandContentInstaller.cpp
@@ -23,6 +23,12 @@
 #include "Templates/Overload.h"
 #include <atomic>
 
+TRACE_DECLARE_INT_COUNTER(OnDemandContentInstaller_HttpRequestsIssued, TEXT("OnDemandContentInstaller/HttpRequestsIssued"));
+TRACE_DECLARE_INT_COUNTER(OnDemandContentInstaller_HttpRequestsTryCanceled, TEXT("OnDemandContentInstaller/HttpRequestsTryCanceled"));
+TRACE_DECLARE_INT_COUNTER(OnDemandContentInstaller_HttpRequestsSuccess, TEXT("OnDemandContentInstaller/HttpRequestsSuccess"));
+TRACE_DECLARE_INT_COUNTER(OnDemandContentInstaller_HttpRequestsFail, TEXT("OnDemandContentInstaller/HttpRequestsFail"));
+TRACE_DECLARE_INT_COUNTER(OnDemandContentInstaller_HttpRequestsActuallyCanceled, TEXT("OnDemandContentInstaller/HttpRequestsActuallyCanceled"));
+
 namespace UE::IoStore
 {
 	namespace CVars
@@ -363,6 +369,8 @@ void FOnDemandContentInstaller::CancelInstallRequest(FSharedInternalInstallReque
 		TEXT("CancelIoStoreOnDemandInstallRequest"),
 		[this, InstallRequest]
 		{
+			TRACE_CPUPROFILER_EVENT_SCOPE(FOnDemandContentInstaller::CancelIoStoreOnDemandInstallRequest);
+
 			FRequest* ToComplete = nullptr;
 			{
 				if (InstallRequest->InstallerRequest == 0)
@@ -420,6 +428,8 @@ void FOnDemandContentInstaller::UpdateInstallRequestPriority(FSharedInternalInst
 		TEXT("UpdateIoStoreOnDemandInstallRequestPriority"),
 		[this, InstallRequest, NewPriority]
 		{
+			TRACE_CPUPROFILER_EVENT_SCOPE(FOnDemandContentInstaller::UpdateInstallRequestPriority);
+
 			if (InstallRequest->InstallerRequest == 0)
 			{
 				return;
@@ -525,6 +535,8 @@ void FOnDemandContentInstaller::TryExecuteNextRequest()
 		return;
 	}
 
+	TRACE_CPUPROFILER_EVENT_SCOPE(FOnDemandContentInstaller::TryExecuteNextRequest);
+
 	bool bRequestQueueIsEmpty = false;
 	while (bRequestQueueIsEmpty == false)
 	{
@@ -846,6 +858,8 @@ void FOnDemandContentInstaller::ExecuteInstallRequest(FRequest& Request)
 	check(Request.IsInstall());
 	check(RunningRequests.Contains(&Request));
 
+	TRACE_CPUPROFILER_EVENT_SCOPE(FOnDemandContentInstaller::ExecuteInstallRequest);
+
 	FRequest::FInstall& InstallRequest = Request.AsInstall();
 	check(InstallRequest.HttpRequestHandles.IsEmpty());
 
@@ -978,6 +992,8 @@ void FOnDemandContentInstaller::ExecuteInstallRequest(FRequest& Request)
 			UE_LOG(LogIoStoreOnDemand, VeryVerbose, TEXT("Created request for chunk: ChunkHash=%s, RequestId=%u"),
 				*TStringBuilder<64>(InPlace, ChunkHash), PendingChunkRequest->RequestId);
 
+			TRACE_COUNTER_INCREMENT(OnDemandContentInstaller_HttpRequestsIssued);
+
 			if (bUseHttpIoDispatcher)
 			{
 				PendingChunkRequest->HttpHandle.Emplace<FIoHttpRequest>(Batch->Get(
@@ -1100,6 +1116,8 @@ void FOnDemandContentInstaller::ExectuteFlushLastAccessRequest(FRequest& Request
 	check(Request.IsFlushLastAccess());
 	check(RunningRequests.Num() == 1 && RunningRequests.Contains(&Request));
 
+	TRACE_CPUPROFILER_EVENT_SCOPE(FOnDemandContentInstaller::ExectuteFlushLastAccessRequest);
+
 	UE_LOG(LogIoStoreOnDemand, Log, TEXT("Executing flush last access request"));
 
 	if (Request.IsCancelled() == false)
@@ -1123,6 +1141,8 @@ void FOnDemandContentInstaller::ProcessDownloadedChunk(
 		const FIoHash VerifyChunkHash = FIoHash::HashBuffer(Chunk.GetView());
 		if (ChunkHash == VerifyChunkHash)
 		{
+			TRACE_COUNTER_INCREMENT(OnDemandContentInstaller_HttpRequestsSuccess);
+
 			Result = IoStore.InstallCache->PutChunk(MoveTemp(Chunk), ChunkHash);
 			if (CVars::EnableConcurentInstallRequests())
 			{
@@ -1132,6 +1152,8 @@ void FOnDemandContentInstaller::ProcessDownloadedChunk(
 		}
 		else
 		{
+			TRACE_COUNTER_INCREMENT(OnDemandContentInstaller_HttpRequestsFail);
+
 			Result = MakeError(UnifiedError::IoStoreOnDemand::ChunkHashError::MakeError(
 				UnifiedError::IoStoreOnDemand::FChunkHashMismatchErrorContext
 				{
@@ -1143,6 +1165,17 @@ void FOnDemandContentInstaller::ProcessDownloadedChunk(
 	}
 	else
 	{
+#if COUNTERSTRACE_ENABLED
+		if (InErrorCode == EIoErrorCode::Cancelled)
+		{
+			TRACE_COUNTER_INCREMENT(OnDemandContentInstaller_HttpRequestsActuallyCanceled);
+		}
+		else
+		{
+			TRACE_COUNTER_INCREMENT(OnDemandContentInstaller_HttpRequestsFail);
+		}
+#endif
+
 		Result = MakeError(UnifiedError::IoStoreOnDemand::HttpError::MakeError(HttpStatusCode));
 	}
 
@@ -1343,6 +1376,8 @@ void FOnDemandContentInstaller::TryCancelHttpRequestsForInstallRequest(FRequest:
 		}
 	}
 
+	TRACE_COUNTER_ADD(OnDemandContentInstaller_HttpRequestsTryCanceled, NumCancelled);
+
 	OutNumCancelled = NumCancelled;
 }
 
@@ -1395,6 +1430,8 @@ void FOnDemandContentInstaller::DestroyRequest(FRequest& Request)
 
 void FOnDemandContentInstaller::CompleteInstallRequest(FRequest& Request)
 {
+	TRACE_CPUPROFILER_EVENT_SCOPE(FOnDemandContentInstaller::CompleteInstallRequest);
+
 	using namespace UE::IoStore::Private;
 
 	FRequest::FInstall& InstallRequest = Request.AsInstall();
