Chaos: Enforce ordered scatter in Tasks mode to fix race condition.

The issue happens because of constraint scatter reading body position/rotation while body scatter may be writing to it, and depending on which one happens first the simulation may diverge. Scatter Bodies writes to P and Q in ApplyCorrections while Scatter Constraints reads them in ApplyPlasticityLimits.

This change makes the scatter bodies task depend on scatter constraints being completed, addressing both indeterminism and the race condition.

#rnx

#DisabledValidations Preflight

diff --git a/Engine/Source/Runtime/Experimental/Chaos/Private/Chaos/Island/IslandGroupManager.cpp b/Engine/Source/Runtime/Experimental/Chaos/Private/Chaos/Island/IslandGroupManager.cpp
--- a/Engine/Source/Runtime/Experimental/Chaos/Private/Chaos/Island/IslandGroupManager.cpp
+++ b/Engine/Source/Runtime/Experimental/Chaos/Private/Chaos/Island/IslandGroupManager.cpp
@@ -487,12 +487,11 @@ namespace Chaos
 							, TStatId()
 							, GatherConstraintsCompletionEvent);
 
+						// Scatter Constraints after the solve completes.
 
-						// Scatter constraints and bodies after the solve completes. No dependencies between any tasks at this point
-						FGraphEventArray ScatterEvents;
-						ScatterEvents.Reserve(NumBodyGatherTasks + NumConstraintGatherTasks);
+						FGraphEventArray ScatterConstraintEvents;
+						ScatterConstraintEvents.Reserve(NumConstraintGatherTasks);
 
-						// Scatter Constraints
 						for (int32 ConstraintBeginIndex = 0; ConstraintBeginIndex < NumConstraints; ConstraintBeginIndex += NumConstraintsPerTask)
 						{
 							const int32 ConstraintEndIndex = FMath::Min(ConstraintBeginIndex + NumConstraintsPerTask, NumConstraints);
@@ -507,10 +506,18 @@ namespace Chaos
 								, GET_STATID(STAT_Evolution_Scatter_Constraints)
 								, SolveConstraintsCompletionEvent);
 
-							ScatterEvents.Add(ScatterConstraintsEvent);
+							ScatterConstraintEvents.Add(ScatterConstraintsEvent);
 						}
 
-						// Scatter Bodies
+						// Scatter Bodies.
+						// 
+						// Note: Executed after ScatterConstraints to avoid race condition given that ScatterConstraints reads FSolverBody P/Q and ScatterBodies writes to them.
+
+						FGraphEventRef BodyScatterDependency = TGraphTask<FNullGraphTask>::CreateTask(&ScatterConstraintEvents).ConstructAndDispatchWhenReady(TStatId(), ENamedThreads::AnyThread);
+
+						FGraphEventArray ScatterBodyEvents;
+						ScatterBodyEvents.Reserve(NumBodyGatherTasks);
+
 						for (int32 BodyBeginIndex = 0; BodyBeginIndex < NumBodies; BodyBeginIndex += NumBodiesPerTask)
 						{
 							const int32 BodyEndIndex = FMath::Min(BodyBeginIndex + NumBodiesPerTask, NumBodies);
@@ -523,13 +530,13 @@ namespace Chaos
 									IslandGroup->ScatterBodies(Dt, BodyBeginIndex, BodyEndIndex);
 								}
 								, GET_STATID(STAT_Evolution_Gather_Bodies)
-								, SolveConstraintsCompletionEvent);
+								, BodyScatterDependency);
 
-							ScatterEvents.Add(ScatterBodyEvent);
+							ScatterBodyEvents.Add(ScatterBodyEvent);
 						}
 
 						// Completion event
-						FGraphEventRef ScatterCompletionEvent = TGraphTask<FNullGraphTask>::CreateTask(&ScatterEvents).ConstructAndDispatchWhenReady(TStatId(), ENamedThreads::AnyThread);
+						FGraphEventRef ScatterCompletionEvent = TGraphTask<FNullGraphTask>::CreateTask(&ScatterBodyEvents).ConstructAndDispatchWhenReady(TStatId(), ENamedThreads::AnyThread);
 						MyCompletionGraphEvent->DontCompleteUntil(ScatterCompletionEvent);
 					});
 
