Enable setting HTTP I/O dispatcher max request range with cvar iax.HttpIoDispatcherMaxRangeKiB=<value>, default value = 4MiB
- this change address an issue where the HTTP I/O dispatcher under certain case can issue too large HTTP range requests
- a negative value disables the max range and issues all currently queued request for a resource as a single HTTP range request
- a value of zero will combine sequential requests for the same resource and range into one HTTP range request

#rnx

#tests Manual testing browsing the locker and stress test using ias.RunRequestTest

diff --git a/Engine/Source/Runtime/Experimental/IoStore/OnDemand/Private/OnDemandContentInstaller.cpp b/Engine/Source/Runtime/Experimental/IoStore/OnDemand/Private/OnDemandContentInstaller.cpp
--- a/Engine/Source/Runtime/Experimental/IoStore/OnDemand/Private/OnDemandContentInstaller.cpp
+++ b/Engine/Source/Runtime/Experimental/IoStore/OnDemand/Private/OnDemandContentInstaller.cpp
@@ -988,6 +988,7 @@ void FOnDemandContentInstaller::ExecuteInstallRequest(FRequest& Request)
 			// Setting HTTP flags to None to not use the HTTP cache
 			FIoHttpOptions HttpOptions = FIoHttpOptions(Request.Priority, HttpRetryCount, EIoHttpFlags::None);
 			HttpOptions.SetCategory(HttpCategory);
+			HttpOptions.SetRange(FIoHttpRange(0, ChunkInfo.Entry.GetDiskSize()));
 
 			UE_LOG(LogIoStoreOnDemand, VeryVerbose, TEXT("Created request for chunk: ChunkHash=%s, RequestId=%u"),
 				*TStringBuilder<64>(InPlace, ChunkHash), PendingChunkRequest->RequestId);
diff --git a/Engine/Source/Runtime/Experimental/IoStore/OnDemand/Private/OnDemandHttpIoDispacher.cpp b/Engine/Source/Runtime/Experimental/IoStore/OnDemand/Private/OnDemandHttpIoDispacher.cpp
--- a/Engine/Source/Runtime/Experimental/IoStore/OnDemand/Private/OnDemandHttpIoDispacher.cpp
+++ b/Engine/Source/Runtime/Experimental/IoStore/OnDemand/Private/OnDemandHttpIoDispacher.cpp
@@ -39,7 +39,17 @@ extern const TCHAR* LexToString(EThreadPriority Priority);
 namespace UE::IoStore
 {
 ////////////////////////////////////////////////////////////////////////////////
+int32 GHttpIoDispatcherMaxRangeKiB = 4096;
+static FAutoConsoleVariableRef CVar_HttpIoDispatcherMaxRangeKiB(
+	TEXT("iax.HttpIoDispatcherMaxRangeKiB"),
+	GHttpIoDispatcherMaxRangeKiB,
+	TEXT("Sets the max HTTP range request range according to:\n")
+	TEXT("	- a negative value merges all currently queued request(s) for a given resource\n")
+	TEXT("	- a zero or positive value will merge requests for a given resource up until the total range goes above this value."),
+	ECVF_Default
+);
 
+////////////////////////////////////////////////////////////////////////////////
 DECLARE_MULTICAST_DELEGATE(FOnRecreateHttpClient);
 extern FOnRecreateHttpClient OnRecreateHttpClient;
 
@@ -227,6 +237,46 @@ FIoHttpRange GetTotalRange(FHttpRequest* Request)
 	return TotalRange;
 }
 
+////////////////////////////////////////////////////////////////////////////////
+FHttpRequestList CreateBatch(
+	FHttpRequest& Request,
+	uint32 MaxRangeSize,
+	FHttpRequestList&& List,
+	FHttpRequestList& OutRemaining)
+{
+	// If any request didn't specify a range we fetch the whole resource 
+	for (const FHttpRequest& Existing : List)
+	{
+		if (Existing.Options.GetRange().IsValid() == false)
+		{
+			return MoveTemp(List);
+		}
+	}
+
+	// The request is expected to be in the list
+	const bool bRemoved = List.Remove(&Request);
+	check(bRemoved);
+
+	FIoHttpRange		Range = Request.Options.GetRange();
+	FHttpRequestList	Batch(&Request);
+
+	while (FHttpRequest* NextRequest = List.PopHead())
+	{
+		FIoHttpRange NextRange = Range + NextRequest->Options.GetRange();
+		if (NextRange.GetSize() > MaxRangeSize)
+		{
+			OutRemaining.AddTail(NextRequest);
+		}
+		else
+		{
+			Range = NextRange;
+			Batch.AddTail(NextRequest);
+		}
+	}
+
+	return Batch;
+}
+
 ////////////////////////////////////////////////////////////////////////////////
 void DebugPrintRange(FHttpRequest* Requests, const UE::FIoHttpRange& TotalRange)
 {
@@ -447,10 +497,8 @@ void FHttpQueue::EnqueueInternal(FHttpRequest* Request)
 		{
 			Request->Options.SetPriority(Head.Options.GetPriority());
 		}
-		List->AddOrInsertBefore(Request, [](const FHttpRequest& New, const FHttpRequest& Existing)
-		{
-			return New.Options.GetRange().GetMin() < Existing.Options.GetRange().GetMin();
-		});
+
+		List->AddTail(Request);
 	}
 	else
 	{
@@ -490,18 +538,46 @@ FHttpRequest* FHttpQueue::Dequeue()
 		return Next;
 	}
 
-	//TODO: Return N request up to some range limit
-	if (FHttpRequestList List; ByChunkKey.RemoveAndCopyValue(Next->ChunkHash, List))
+	if (FHttpRequestList* List = ByChunkKey.Find(Next->ChunkHash))
 	{
-		// The first request that was enqueued for the chunk is not necessarily the first request in the list
-		Next = List.PeekHead();
+		// The request on the heap is expected to be the first in the list
+		ensure(Next == List->PeekHead());
 
-		if (ByChunkKey.IsEmpty())
+		if (GHttpIoDispatcherMaxRangeKiB < 0)
 		{
-			ByChunkKey.Empty(ContainerSlackSize);
+			// Fetch the entire batch
+			Next = List->PeekHead();
+			ByChunkKey.Remove(Next->ChunkHash);
+		}
+		else
+		{
+			FHttpRequestList	Requests(MoveTemp(*List));
+			*List								= FHttpRequestList();
+			const uint32		MaxRangeSize	= uint32(GHttpIoDispatcherMaxRangeKiB) << 10;
+			FHttpRequestList	Batch			= CreateBatch(*Next, MaxRangeSize, MoveTemp(Requests), *List);
+
+			if (List->IsEmpty())
+			{
+				ByChunkKey.Remove(Next->ChunkHash);
+			}
+			else
+			{
+				// Push the first request in the remaining list to the heap when there's more requests for the same chunk hash
+				FHttpRequest* Head = List->PeekHead();
+				ensure(Heap.Find(Head) == INDEX_NONE);
+				Heap.HeapPush(Head, HttpRequestSortPredicate);
+				FOnDemandIoBackendStats::Get()->OnHttpEnqueue(Head->GetHttpRequestType());
+			}
+
+			Next = Batch.PeekHead();
 		}
 	}
 
+	if (ByChunkKey.IsEmpty())
+	{
+		ByChunkKey.Empty(ContainerSlackSize);
+	}
+
 	for (FHttpRequest* It = Next; It != nullptr; It = It->Next)
 	{
 		It->Flags.Add(EHttpRequestFlags::HttpInflight);
diff --git a/Engine/Source/Runtime/Experimental/IoStore/OnDemand/Private/Statistics.cpp b/Engine/Source/Runtime/Experimental/IoStore/OnDemand/Private/Statistics.cpp
--- a/Engine/Source/Runtime/Experimental/IoStore/OnDemand/Private/Statistics.cpp
+++ b/Engine/Source/Runtime/Experimental/IoStore/OnDemand/Private/Statistics.cpp
@@ -28,6 +28,7 @@ namespace UE::IoStore
 {
 
 extern int32 GIaxHttpVersion;
+extern int32 GHttpIoDispatcherMaxRangeKiB;
 
 float GIasStatisticsLogInterval = 60.f;
 static FAutoConsoleVariableRef CVar_StatisticsLogInterval(
@@ -766,6 +767,25 @@ void FOnDemandIoBackendStats::PrintOnScreenStatistics(FCoreDelegates::FSeverityM
 		BackendStatus.ToString(BackendStatusText);
 		BackendStatusText.Appendf(TEXT(" | HTTP I/O dispatcher: %s"), GHttpIoDispatcherEnabled ? TEXT("Enabled") : TEXT("Disabled"));
 
+		if (GHttpIoDispatcherEnabled)
+		{
+			BackendStatusText.Append(" | MaxRange: ");
+
+			const int32 MaxRange = GHttpIoDispatcherMaxRangeKiB;
+			if (MaxRange < 0)
+			{
+				BackendStatusText.Append("Unlimited");
+			}
+			else if (MaxRange == 0)
+			{
+				BackendStatusText.Append("Single");
+			}
+			else
+			{
+				BackendStatusText.Appendf(TEXT("%d KiB"), MaxRange);
+			}
+		}
+
 		OutMessages.Add(FCoreDelegates::EOnScreenMessageSeverity::Info, FText::FromStringView(BackendStatusText.ToView()));
 	}
 
