Surface time-to-first-byte + small tweak to IAX debug text overlay.

#rnx

diff --git a/Engine/Source/Runtime/Experimental/IoStore/OnDemand/Private/OnDemandHttpClient.cpp b/Engine/Source/Runtime/Experimental/IoStore/OnDemand/Private/OnDemandHttpClient.cpp
--- a/Engine/Source/Runtime/Experimental/IoStore/OnDemand/Private/OnDemandHttpClient.cpp
+++ b/Engine/Source/Runtime/Experimental/IoStore/OnDemand/Private/OnDemandHttpClient.cpp
@@ -368,6 +368,7 @@ void FMultiEndpointHttpClient::GetHttpStats(FOnDemandHttpStats& Out) const
 		virtual uint32 GetRecvKiBps() const { return Perf.GetRecvKiBps(); }
 		virtual void GetRecvKiBps(uint32 (&Out)[SAMPLE_COUNT]) const { return Perf.GetRecvKiBps(Out); }
 		virtual uint32 GetTotalRecvKiB() const { return Perf.GetTotalRecvKiB(); }
+		virtual uint32 GetTimeToFirstByteMs() const { return Perf.GetTimeToFirstByteMs(); }
 		FLoopPerf Perf;
 	};
 
diff --git a/Engine/Source/Runtime/Experimental/IoStore/OnDemand/Private/Statistics.cpp b/Engine/Source/Runtime/Experimental/IoStore/OnDemand/Private/Statistics.cpp
--- a/Engine/Source/Runtime/Experimental/IoStore/OnDemand/Private/Statistics.cpp
+++ b/Engine/Source/Runtime/Experimental/IoStore/OnDemand/Private/Statistics.cpp
@@ -840,13 +840,15 @@ void FOnDemandIoBackendStats::PrintOnScreenStatistics(FCoreDelegates::FSeverityM
 
 		{
 			static uint32 RateKiBps = 0;
+			static uint32 DownKiB = 0;
+			static uint32 TtfbMs = 0;
 			if (IOnDemandIoStore* Store = TryGetOnDemandIoStore(); Store != nullptr)
 			{
 				double Time = FPlatformTime::Seconds();
-				if (static double LastSampleTime = Time; Time - LastSampleTime > 1.0)
+				if (static double LastSampleTime = Time; Time - LastSampleTime > 0.5)
 				{
-					LastSampleTime -= 1.0;
-					if (LastSampleTime > 1.0)
+					LastSampleTime -= 0.5;
+					if (LastSampleTime > 0.5)
 					{
 						LastSampleTime = Time;
 					}
@@ -854,6 +856,8 @@ void FOnDemandIoBackendStats::PrintOnScreenStatistics(FCoreDelegates::FSeverityM
 					static FOnDemandHttpStats Stats;
 					Store->GetHttpStats(Stats);
 					RateKiBps = Stats.GetRecvKiBps();
+					DownKiB = Stats.GetTotalRecvKiB();
+					TtfbMs = Stats.GetTimeToFirstByteMs();
 				}
 			}
 
@@ -862,10 +866,12 @@ void FOnDemandIoBackendStats::PrintOnScreenStatistics(FCoreDelegates::FSeverityM
 			TStringBuilder<256> Builder;
 			Builder << TEXT("IAX Http:");
 			Builder << TEXT(" v") << GIaxHttpVersion;
-			Builder << TEXT(" DataRate: ") << RateKiBps << TEXT(" / ") << RateLimit << TEXT(" KiB/s");
+			Builder << TEXT(" - Downloaded: ") << FText::AsMemory(uint64(DownKiB) << 10).ToString();
+			Builder << TEXT(" - Rate/limit: ") << RateKiBps << TEXT(" / ") << RateLimit << TEXT(" KiB/s");
+			Builder << TEXT(" - TTFB: ") << TtfbMs << TEXT("ms");
 
 			const int32 NumConcurrentRequests = GHttpNumConcurrentRequests;
-			Builder << TEXT(" ConcurrentRequests: ") << NumConcurrentRequests << TEXT("/") << GIasHttpConcurrentRequests;
+			Builder << TEXT(" - Queue: ") << NumConcurrentRequests << TEXT("/") << GIasHttpConcurrentRequests;
 			if (NumConcurrentRequests >= GIasHttpConcurrentRequests)
 			{
 				Builder << TEXT("(saturated)");
diff --git a/Engine/Source/Runtime/Experimental/IoStore/OnDemandCore/Internal/IO/IoStoreOnDemandInternals.h b/Engine/Source/Runtime/Experimental/IoStore/OnDemandCore/Internal/IO/IoStoreOnDemandInternals.h
--- a/Engine/Source/Runtime/Experimental/IoStore/OnDemandCore/Internal/IO/IoStoreOnDemandInternals.h
+++ b/Engine/Source/Runtime/Experimental/IoStore/OnDemandCore/Internal/IO/IoStoreOnDemandInternals.h
@@ -80,6 +80,7 @@ public:
 	virtual uint32 GetRecvKiBps() const = 0;
 	virtual void GetRecvKiBps(uint32 (&Out)[SAMPLE_COUNT]) const = 0;
 	virtual uint32 GetTotalRecvKiB() const = 0;
+	virtual uint32 GetTimeToFirstByteMs() const = 0;
 };
 
 } // namespace UE::IoStore
diff --git a/Engine/Source/Runtime/Experimental/IoStore/OnDemandCore/Private/OnDemandHttpStats.cpp b/Engine/Source/Runtime/Experimental/IoStore/OnDemandCore/Private/OnDemandHttpStats.cpp
--- a/Engine/Source/Runtime/Experimental/IoStore/OnDemandCore/Private/OnDemandHttpStats.cpp
+++ b/Engine/Source/Runtime/Experimental/IoStore/OnDemandCore/Private/OnDemandHttpStats.cpp
@@ -30,5 +30,10 @@ uint32 FOnDemandHttpStats::GetTotalRecvKiB() const
 	return Internal.IsValid() ? Internal->GetTotalRecvKiB() : 0;
 }
 
+uint32 FOnDemandHttpStats::GetTimeToFirstByteMs() const
+{
+	return Internal.IsValid() ? Internal->GetTimeToFirstByteMs() : 0;
+}
+
 } // namespace IoStore
 } // namespace UE
diff --git a/Engine/Source/Runtime/Experimental/IoStore/OnDemandCore/Public/IO/IoStoreOnDemand.h b/Engine/Source/Runtime/Experimental/IoStore/OnDemandCore/Public/IO/IoStoreOnDemand.h
--- a/Engine/Source/Runtime/Experimental/IoStore/OnDemandCore/Public/IO/IoStoreOnDemand.h
+++ b/Engine/Source/Runtime/Experimental/IoStore/OnDemandCore/Public/IO/IoStoreOnDemand.h
@@ -537,6 +537,7 @@ public:
 	UE_API uint32 GetRecvKiBps() const;
 	UE_API void GetRecvKiBps(uint32 (&Out)[SAMPLE_COUNT]) const;
 	UE_API uint32 GetTotalRecvKiB() const;
+	UE_API uint32 GetTimeToFirstByteMs() const;
 
 private:
 	friend class FMultiEndpointHttpClient;
