Rounding error

#rnx

diff --git a/Engine/Source/Runtime/Experimental/IoStore/HttpClient/Private/Loop.inl b/Engine/Source/Runtime/Experimental/IoStore/HttpClient/Private/Loop.inl
--- a/Engine/Source/Runtime/Experimental/IoStore/HttpClient/Private/Loop.inl
+++ b/Engine/Source/Runtime/Experimental/IoStore/HttpClient/Private/Loop.inl
@@ -706,7 +706,7 @@ uint32 FPeerGroup::CollectRecvSizeKiB()
 	uint64 Size = Peer.GetRecvSize();
 	uint32 Ret = uint32(Size - LastRecvSize) >> 10;
 
-	LastRecvSize = Size;
+	LastRecvSize = Size & ~1023;
 	return Ret;
 }
 
