Allow the cvars 'iax.HttpVersion' and 'iax.HttpMaxInfligh' to be hotfixed and/or dynamically changed at runtime.

#rnx
#prelfight

- Since the http client is created before hotfixes are commonly applied we need to make sure that the cvads triggers the OnHttpClientCVarChanged callback so that the client will be recreated with the desired change.

diff --git a/Engine/Source/Runtime/Experimental/IoStore/OnDemand/Private/OnDemandHttpThread.cpp b/Engine/Source/Runtime/Experimental/IoStore/OnDemand/Private/OnDemandHttpThread.cpp
--- a/Engine/Source/Runtime/Experimental/IoStore/OnDemand/Private/OnDemandHttpThread.cpp
+++ b/Engine/Source/Runtime/Experimental/IoStore/OnDemand/Private/OnDemandHttpThread.cpp
@@ -98,7 +98,8 @@ int32 GIaxHttpVersion = 1;
 static FAutoConsoleVariableRef CVar_IaxHttpVersion(
 	TEXT("iax.HttpVersion"),
 	GIaxHttpVersion,
-	TEXT("Protocol version to use, either '1' or '2'")
+	TEXT("Protocol version to use, either '1' or '2'"),
+	FConsoleVariableDelegate::CreateStatic(&OnHttpClientCVarChanged)
 );
 
 int32 GIasHttpConnectionCount = 4;
@@ -113,7 +114,8 @@ int32 GIaxHttpMaxInflight = 1;
 static FAutoConsoleVariableRef CVar_IaxHttpMaxInflight(
 	TEXT("iax.HttpMaxInflight"),
 	GIaxHttpMaxInflight,
-	TEXT("Number of requests to issue at once per connection")
+	TEXT("Number of requests to issue at once per connection"),
+	FConsoleVariableDelegate::CreateStatic(&OnHttpClientCVarChanged)
 );
 
 int32 GIasHttpRecvBufKiB = -1;
