Fixed bug where filtered CDNs were still being logged as part of a host group.

#rnx

- When logging the CDN urls present in the newly created FIASHostGroup we are using the list of urls that we pass to FIASHostGroup::Resolve. However one or more of these urls might be filtered out during the resolve so we should be logging the urls in the final hostgroup instead.

diff --git a/Engine/Source/Runtime/Experimental/IoStore/OnDemand/Private/OnDemandHttpIoDispacher.cpp b/Engine/Source/Runtime/Experimental/IoStore/OnDemand/Private/OnDemandHttpIoDispacher.cpp
--- a/Engine/Source/Runtime/Experimental/IoStore/OnDemand/Private/OnDemandHttpIoDispacher.cpp
+++ b/Engine/Source/Runtime/Experimental/IoStore/OnDemand/Private/OnDemandHttpIoDispacher.cpp
@@ -939,12 +939,14 @@ FIoStatus FHttpDispatcher::RegisterHostGroup(const FName& HostGroupName, TConstA
 		return FIoStatus::Ok;
 	}
 
-	TIoStatusOr<FIASHostGroup> NewGroup = FHostGroupManager::Get().Register(HostGroupName, TestUrl);
-	if (NewGroup.IsOk() == false)
+	TIoStatusOr<FIASHostGroup> RegisterResult = FHostGroupManager::Get().Register(HostGroupName, TestUrl);
+	if (RegisterResult.IsOk() == false)
 	{
-		return NewGroup.Status();
+		return RegisterResult.Status();
 	}
 
+	FIASHostGroup NewGroup = RegisterResult.ConsumeValueOrDie();
+
 #if !UE_BUILD_SHIPPING
 	TArray<FAnsiString> HostsOverride;
 	if (HostGroupName == FOnDemandHostGroup::DefaultName)
@@ -968,14 +970,14 @@ FIoStatus FHttpDispatcher::RegisterHostGroup(const FName& HostGroupName, TConstA
 	}
 #endif
 
-	if (FIoStatus Status = NewGroup.ConsumeValueOrDie().Resolve(HostNames); !Status.IsOk())
+	if (FIoStatus Status = NewGroup.Resolve(HostNames); !Status.IsOk())
 	{
 		UE_LOG(LogHttpIoDispatcher, Warning, TEXT("Failed to create host group '%s'"), *HostGroupName.ToString());
 		return Status;
 	}
 
 	UE_LOG(LogHttpIoDispatcher, Log, TEXT("Registered new host group '%s'"), *HostGroupName.ToString());
-	for (const FAnsiString& HostName : HostNames) 
+	for (const FAnsiString& HostName : NewGroup.GetHostUrls())
 	{
 		UE_LOG(LogHttpIoDispatcher, Log, TEXT("\t* %s"), *FString(HostName));
 	}
