Fix unsafe use of a string view's pointer

#rnx

- Currently we know that the FStringViews return by FContext are null terminated but stricktly speaking there if no guarentte of this when calling FStringView::GetData.
- Since speed doesn't matter here we might as well store the views in FString/FAnsiString so we don't have to worry that future changes to FContext might make this code unsafe.

diff --git a/Engine/Source/Runtime/Experimental/IoStore/OnDemand/Private/Tool/Test.cpp b/Engine/Source/Runtime/Experimental/IoStore/OnDemand/Private/Tool/Test.cpp
--- a/Engine/Source/Runtime/Experimental/IoStore/OnDemand/Private/Tool/Test.cpp
+++ b/Engine/Source/Runtime/Experimental/IoStore/OnDemand/Private/Tool/Test.cpp
@@ -4,10 +4,11 @@
 
 #include "Command.h"
 
-#include <Async/TaskGraphInterfaces.h>
-#include <Containers/StringConv.h>
-#include <Containers/UnrealString.h>
-#include <Misc/CommandLine.h>
+#include "Async/TaskGraphInterfaces.h"
+#include "Containers/AnsiString.h"
+#include "Containers/StringConv.h"
+#include "Containers/UnrealString.h"
+#include "Misc/CommandLine.h"
 
 namespace UE::IoStore {
 
@@ -23,19 +24,18 @@ void CommandTest();
 ////////////////////////////////////////////////////////////////////////////////
 static void HttpTests(const FContext& Context)
 {
-	auto TestHost = Context.Get<FStringView>(TEXT("-Host"), TEXT("localhost"));
+	FAnsiString TestHost = FAnsiString(Context.Get<FStringView>(TEXT("-Host"), TEXT("localhost")));
 
-	auto TestHostAnsi = StringCast<ANSICHAR>(TestHost.GetData());
 	uint32 Seed = Context.Get<uint32>(TEXT("-HttpSeed"), 493);
-	HTTP::IasHttpTest(TestHostAnsi.Get(), Seed);
+	HTTP::IasHttpTest(*TestHost, Seed);
 }
 
 ////////////////////////////////////////////////////////////////////////////////
 static void CacheTests(const FContext& Context)
 {
-	FStringView CacheDirStr = Context.Get<FStringView>(TEXT("-Dir"));
+	FString CacheDirStr = FString(Context.Get<FStringView>(TEXT("-Dir")));
 
-	const TCHAR* CacheDir = CacheDirStr.IsEmpty() ? nullptr : CacheDirStr.GetData();
+	const TCHAR* CacheDir = CacheDirStr.IsEmpty() ? nullptr : *CacheDirStr;
 	IasJournaledFileCacheTest::Tests(CacheDir);
 }
 
