Collision Filter Refactor Part 3.11: Updating GeometryCollection to use new API and structures.

diff --git a/Engine/Source/Programs/HeadlessChaos/Private/GeometryCollection/GeometryCollectionTestFramework.cpp b/Engine/Source/Programs/HeadlessChaos/Private/GeometryCollection/GeometryCollectionTestFramework.cpp
--- a/Engine/Source/Programs/HeadlessChaos/Private/GeometryCollection/GeometryCollectionTestFramework.cpp
+++ b/Engine/Source/Programs/HeadlessChaos/Private/GeometryCollection/GeometryCollectionTestFramework.cpp
@@ -203,22 +203,18 @@ namespace GeometryCollectionTest
 			FGeometryCollectionPhysicsProxy::InitializeDynamicCollection(*DynamicCollection, *RestCollection, SimulationParams);
 		}
 
-		FCollisionFilterData SimFilterData;
-		FCollisionFilterData QueryFilterData;
-
 		PRAGMA_DISABLE_INTERNAL_WARNINGS
 		const Chaos::Filter::FShapeFilterData ShapeFilterData = Chaos::Filter::FShapeFilterBuilder::BuildBlockAll(EFilterFlags::All);
-		Chaos::Filter::FCombinedShapeFilterData CombinedShapeFilter(ShapeFilterData, Chaos::Filter::FInstanceData());
-		Chaos::Filter::FShapeFilterBuilder::GetLegacyShapeFilter(CombinedShapeFilter, QueryFilterData, SimFilterData);
 		PRAGMA_ENABLE_INTERNAL_WARNINGS
+		Chaos::Filter::FInstanceData ShapeInstanceData;
 
 		FGeometryCollectionPhysicsProxy* PhysObject =
 			new FGeometryCollectionPhysicsProxy(
 				nullptr,			// UObject owner
 				DynamicCollection.ToSharedRef(), // Game thread collection
 				SimulationParams,
-				SimFilterData,
-				QueryFilterData);
+				ShapeFilterData,
+				ShapeInstanceData);
 
 		Wrapper->RestCollection = RestCollection;
 		Wrapper->DynamicCollection = DynamicCollection;
diff --git a/Engine/Source/Runtime/Experimental/Chaos/Private/PhysicsProxy/GeometryCollectionPhysicsProxy.cpp b/Engine/Source/Runtime/Experimental/Chaos/Private/PhysicsProxy/GeometryCollectionPhysicsProxy.cpp
--- a/Engine/Source/Runtime/Experimental/Chaos/Private/PhysicsProxy/GeometryCollectionPhysicsProxy.cpp
+++ b/Engine/Source/Runtime/Experimental/Chaos/Private/PhysicsProxy/GeometryCollectionPhysicsProxy.cpp
@@ -290,8 +290,8 @@ void SetImplicitToPTParticles(Chaos::TPBDRigidParticleHandle<Chaos::FReal, 3>* H
 	const FSharedSimulationParameters& SharedParams,
 	const FCollisionStructureManager::FSimplicial* Simplicial,
 	Chaos::FImplicitObjectPtr Implicit,
-	const FCollisionFilterData SimFilterIn,
-	const FCollisionFilterData QueryFilterIn,
+	const Chaos::Filter::FShapeFilterData& ShapeFilterDataIn,
+	const Chaos::Filter::FInstanceData& ShapeInstanceDataIn,
 	const FTransform& WorldTransform,
 	float CollisionParticlesPerObjectFraction)
 {
@@ -390,10 +390,10 @@ void SetImplicitToPTParticles(Chaos::TPBDRigidParticleHandle<Chaos::FReal, 3>* H
 	}
 	else
 	{
-		const Chaos::Filter::FCombinedShapeFilterData CombinedFilterData = Chaos::Filter::FShapeFilterBuilder::BuildFromLegacyShapeFilter(QueryFilterIn, SimFilterIn);
 		for (const TUniquePtr<Chaos::FPerShapeData>& Shape : Handle->ShapesArray())
 		{
-			Shape->SetCombinedShapeFilterData(CombinedFilterData);
+			Shape->SetShapeFilterData(ShapeFilterDataIn);
+			Shape->SetFilterInstanceData(ShapeInstanceDataIn);
 		}
 	}
 }
@@ -404,8 +404,8 @@ void PopulateSimulatedParticle(
 	const FSharedSimulationParameters& SharedParams,
 	const FCollisionStructureManager::FSimplicial* Simplicial,
 	Chaos::FImplicitObjectPtr Implicit,
-	const FCollisionFilterData SimFilterIn,
-	const FCollisionFilterData QueryFilterIn,
+	const Chaos::Filter::FShapeFilterData& ShapeFilterDataIn,
+	const Chaos::Filter::FInstanceData& ShapeInstanceDataIn,
 	Chaos::FReal MassIn,
 	Chaos::TVec3<Chaos::FRealSingle> InertiaTensorVec,
 	const FTransform& WorldTransform, 
@@ -463,7 +463,7 @@ void PopulateSimulatedParticle(
 
 	Handle->SetCollisionGroup(CollisionGroup);
 
-	SetImplicitToPTParticles(Handle, SharedParams, Simplicial, Implicit, SimFilterIn, QueryFilterIn, WorldTransform, CollisionParticlesPerObjectFraction);
+	SetImplicitToPTParticles(Handle, SharedParams, Simplicial, Implicit, ShapeFilterDataIn, ShapeInstanceDataIn, WorldTransform, CollisionParticlesPerObjectFraction);
 
 	//
 	//  Manage Object State
@@ -516,14 +516,16 @@ FGeometryCollectionPhysicsProxy::FGeometryCollectionPhysicsProxy(
 	, IsObjectLoading(true)
 	, IsObjectDeleting(false)
 	, bSkipChildToParentUpdateWhenInClusterUnion(false)
-	, SimFilter(InSimFilter)
-	, QueryFilter(InQueryFilter)
 	, PhysicsThreadCollection(Parameters.RestCollectionShared)
 	, GameThreadCollection(GameThreadCollectionIn)
 #if WITH_EDITORONLY_DATA
 	, CollectorGuid(InCollectorGuid)
 #endif
 {
+	Chaos::Filter::FCombinedShapeFilterData CombinedShapeFilter = Chaos::Filter::FShapeFilterBuilder::BuildFromLegacyShapeFilter(InQueryFilter, InSimFilter);
+	ShapeFilterData = CombinedShapeFilter.GetShapeFilterData();
+	ShapeInstanceData = CombinedShapeFilter.GetInstanceData();
+
 	// We rely on a guarded buffer.
 	check(BufferMode == Chaos::EMultiBufferMode::TripleGuarded);
 	InterpolationData = MakeUnique<FProxyInterpolationBase>();
@@ -552,8 +554,46 @@ FGeometryCollectionPhysicsProxy::FGeometryCollectionPhysicsProxy(
 	, IsObjectLoading(true)
 	, IsObjectDeleting(false)
 	, bSkipChildToParentUpdateWhenInClusterUnion(false)
-	, SimFilter(InSimFilter)
-	, QueryFilter(InQueryFilter)
+	, PhysicsThreadCollection(Parameters.RestCollectionShared)
+	, GameThreadCollection(GameThreadCollectionSharedRefIn.Get())
+	, GameThreadCollectionSharedPtr(GameThreadCollectionSharedRefIn.ToSharedPtr())
+#if WITH_EDITORONLY_DATA
+	, CollectorGuid(InCollectorGuid)
+#endif
+{
+	Chaos::Filter::FCombinedShapeFilterData CombinedShapeFilter = Chaos::Filter::FShapeFilterBuilder::BuildFromLegacyShapeFilter(InQueryFilter, InSimFilter);
+	ShapeFilterData = CombinedShapeFilter.GetShapeFilterData();
+	ShapeInstanceData = CombinedShapeFilter.GetInstanceData();
+
+	// We rely on a guarded buffer.
+	check(BufferMode == Chaos::EMultiBufferMode::TripleGuarded);
+	InterpolationData = MakeUnique<FProxyInterpolationBase>();
+}
+
+FGeometryCollectionPhysicsProxy::FGeometryCollectionPhysicsProxy(
+	UObject* InOwner,
+	TSharedRef<FGeometryDynamicCollection> GameThreadCollectionSharedRefIn,
+	const FSimulationParameters& SimulationParameters,
+	const Chaos::Filter::FShapeFilterData& InShapeFilterData,
+	const Chaos::Filter::FInstanceData& InShapeInstanceData,
+	FGuid InCollectorGuid,
+	const Chaos::EMultiBufferMode BufferMode)
+	: Base(InOwner)
+	, Parameters(SimulationParameters)
+	, WorldTransform_External(FTransform::Identity)
+	, NumTransforms(INDEX_NONE)
+	, NumEffectiveParticles(INDEX_NONE)
+	, BaseParticleIndex(INDEX_NONE)
+	, CollisionParticlesPerObjectFraction(CollisionParticlesPerObjectFractionDefault)
+	, bIsGameThreadWorldTransformDirty(false)
+	, bHasBuiltGeometryOnPT(false)
+	, bHasBuiltGeometryOnGT(false)
+	, IsObjectDynamic(false)
+	, IsObjectLoading(true)
+	, IsObjectDeleting(false)
+	, bSkipChildToParentUpdateWhenInClusterUnion(false)
+	, ShapeFilterData(InShapeFilterData)
+	, ShapeInstanceData(InShapeInstanceData)
 	, PhysicsThreadCollection(Parameters.RestCollectionShared)
 	, GameThreadCollection(GameThreadCollectionSharedRefIn.Get())
 	, GameThreadCollectionSharedPtr(GameThreadCollectionSharedRefIn.ToSharedPtr())
@@ -966,8 +1006,8 @@ void FGeometryCollectionPhysicsProxy::CreateGTParticles(TManagedArray<Chaos::FIm
 			for (int32 ShapeIndex = 0; ShapeIndex < NumShapes; ++ShapeIndex)
 			{
 				Chaos::FPerShapeData* Shape = Shapes[ShapeIndex].Get();
-				const Chaos::Filter::FCombinedShapeFilterData CombinedShapeFilter = Chaos::Filter::FShapeFilterBuilder::BuildFromLegacyShapeFilter(QueryFilter, SimFilter);
-				Shape->SetCombinedShapeFilterData(CombinedShapeFilter);
+				Shape->SetShapeFilterData(ShapeFilterData);
+				Shape->SetFilterInstanceData(ShapeInstanceData);
 				Shape->SetProxy(this);
 				Shape->SetMaterial(Parameters.PhysicalMaterialHandle);
 			}
@@ -1330,8 +1370,8 @@ void FGeometryCollectionPhysicsProxy::InitializeBodiesPT(Chaos::FPBDRigidsSolver
 					Parameters.Shared,
 					Simplicials? (*Simplicials)[TransformGroupIndex].Get(): nullptr,
 					bUseImplicit? (*Implicits)[TransformGroupIndex] : nullptr,
-					SimFilter,
-					QueryFilter,
+					ShapeFilterData,
+					ShapeInstanceData,
 					ScaledMass,
 					ScaledInertia,
 					Transforms[TransformGroupIndex],
@@ -1775,7 +1815,7 @@ void FGeometryCollectionPhysicsProxy::CreateChildrenGeometry_Internal()
 					{
 						const FTransform WorldTransform = MassToLocal[TransformIndex] * Transform[TransformIndex] * Parameters.WorldTransform;
 						const FCollisionStructureManager::FSimplicial* Simplicial = Simplicials ? (*Simplicials)[TransformIndex].Get() : nullptr;
-						SetImplicitToPTParticles(Handle, Parameters.Shared, Simplicial, Implicits? (*Implicits)[TransformIndex]: nullptr, SimFilter, QueryFilter, WorldTransform, CollisionParticlesPerObjectFraction);
+						SetImplicitToPTParticles(Handle, Parameters.Shared, Simplicial, Implicits? (*Implicits)[TransformIndex]: nullptr, ShapeFilterData, ShapeInstanceData, WorldTransform, CollisionParticlesPerObjectFraction);
 						check(Handle->GetGeometry() != nullptr);
 					}
 				}
@@ -2014,8 +2054,8 @@ Chaos::TPBDGeometryCollectionParticleHandle<Chaos::FReal, 3>* FGeometryCollectio
 		Parameters.Shared,
 		Simplicials? (*Simplicials)[CollectionClusterIndex].Get(): nullptr,
 		bUseImplicit? (*Implicits)[CollectionClusterIndex] : nullptr,
-		SimFilter,
-		QueryFilter,
+		ShapeFilterData, 
+		ShapeInstanceData, 
 		Mass,
 		Inertia,
 		ParticleTM,
@@ -2152,8 +2192,8 @@ FGeometryCollectionPhysicsProxy::BuildClusters_Internal(
 		Parameters.Shared, 
 		nullptr, // CollisionParticles is optionally set from CreateClusterParticle()
 		nullptr, // Parent->Geometry() ? Parent->Geometry() : Implicits[CollectionClusterIndex], 
-		SimFilter,
-		QueryFilter,
+		ShapeFilterData, 
+		ShapeInstanceData, 
 		ScaledMass,
 		ScaledInertia,
 		ParticleTM, 
@@ -4762,7 +4802,8 @@ void FGeometryCollectionPhysicsProxy::SetFilterData_Internal(const Chaos::Filter
 {
 	using namespace Chaos;
 
-	Chaos::Filter::FShapeFilterBuilder::GetLegacyShapeFilter(NewShapeFilter, Parameters.QueryFilterData, Parameters.SimulationFilterData);
+	Parameters.ShapeFilterData = NewShapeFilter.GetShapeFilterData();
+	Parameters.ShapeInstanceData = NewShapeFilter.GetInstanceData();
 
 	TArray<FClusterHandle*> ProcessedInternalClusters;
 	if (Chaos::FPhysicsSolver* RigidSolver = GetSolver<Chaos::FPhysicsSolver>())
@@ -4829,7 +4870,7 @@ void FGeometryCollectionPhysicsProxy::SetPerParticleFilterData_Internal(const TA
 					check(Object != nullptr);
 					TArrayView<Chaos::FPhysicsObjectHandle> ParticleView{ &Object, 1 };
 					Interface.UpdateShapeCollisionFlags(ParticleView, Data.bSimEnabled, Data.bQueryEnabled);
-					Interface.UpdateCombinedShapeFilterData(ParticleView, Chaos::Filter::FShapeFilterBuilder::BuildFromLegacyShapeFilter(Data.QueryFilter, Data.SimFilter));
+					Interface.UpdateCombinedShapeFilterData(ParticleView, Chaos::Filter::FCombinedShapeFilterData(Data.ShapeFilterData, Data.FilterInstanceData));
 					// todo(chaos): It's not ideal but the geometry collection needs to request the cluster union update its
 					// cached shape data if the particle is in a cluster union. This is because we don't share shape
 					// data between the GC shapes and the cluster union shapes.
@@ -5567,8 +5608,8 @@ void FGeometryCollectionPhysicsProxy::InitializeSharedCollisionStructures(
  					SharedParams, 
  					CollectionSimplicials[TransformGroupIndex].Get(),
  					CollectionImplicits[TransformGroupIndex],
- 					FCollisionFilterData(),		// SimFilter
- 					FCollisionFilterData(),		// QueryFilter
+ 					Chaos::Filter::FShapeFilterData(),
+ 					Chaos::Filter::FInstanceData(),
  					CollectionMass[TransformGroupIndex],
  					CollectionInertiaTensor[TransformGroupIndex], 
 					CollectionSpaceParticleTransform,
diff --git a/Engine/Source/Runtime/Experimental/Chaos/Public/GeometryCollection/GeometryCollectionSimulationCoreTypes.h b/Engine/Source/Runtime/Experimental/Chaos/Public/GeometryCollection/GeometryCollectionSimulationCoreTypes.h
--- a/Engine/Source/Runtime/Experimental/Chaos/Public/GeometryCollection/GeometryCollectionSimulationCoreTypes.h
+++ b/Engine/Source/Runtime/Experimental/Chaos/Public/GeometryCollection/GeometryCollectionSimulationCoreTypes.h
@@ -230,8 +230,12 @@ PRAGMA_ENABLE_DEPRECATION_WARNINGS
 		, SleepThresholdMultiplier(1.0f)
 		, GravityGroupIndex(0)
 		, OneWayInteractionLevel(INDEX_NONE)
+PRAGMA_DISABLE_DEPRECATION_WARNINGS
 		, SimulationFilterData()
 		, QueryFilterData()
+PRAGMA_ENABLE_DEPRECATION_WARNINGS
+		, ShapeFilterData()
+		, ShapeInstanceData()
 		, UserData(nullptr)
 #if SIMULATIONPARAMETERS_CACHE_PARAMETERS
 PRAGMA_DISABLE_DEPRECATION_WARNINGS
@@ -309,8 +313,12 @@ PRAGMA_ENABLE_DEPRECATION_WARNINGS
 		, SleepThresholdMultiplier(Other.SleepThresholdMultiplier)
 		, GravityGroupIndex(Other.GravityGroupIndex)
 		, OneWayInteractionLevel(Other.OneWayInteractionLevel)
+PRAGMA_DISABLE_DEPRECATION_WARNINGS
 		, SimulationFilterData(Other.SimulationFilterData)
 		, QueryFilterData(Other.QueryFilterData)
+PRAGMA_ENABLE_DEPRECATION_WARNINGS
+		, ShapeFilterData(Other.ShapeFilterData)
+		, ShapeInstanceData(Other.ShapeInstanceData)
 		, UserData(Other.UserData)
 #if SIMULATIONPARAMETERS_CACHE_PARAMETERS
 PRAGMA_DISABLE_DEPRECATION_WARNINGS
@@ -428,8 +436,12 @@ PRAGMA_ENABLE_DEPRECATION_WARNINGS
 	int32 GravityGroupIndex;
 	int32 OneWayInteractionLevel;
 
+	UE_DEPRECATED(5.8, "Use ShapeFilterData and ShapeInstanceData instead")
 	FCollisionFilterData SimulationFilterData;
+	UE_DEPRECATED(5.8, "Use ShapeFilterData and ShapeInstanceData instead")
 	FCollisionFilterData QueryFilterData;
+	Chaos::Filter::FShapeFilterData ShapeFilterData;
+	Chaos::Filter::FInstanceData ShapeInstanceData;
 
 	void* UserData;
 
diff --git a/Engine/Source/Runtime/Experimental/Chaos/Public/PhysicsProxy/GeometryCollectionPhysicsProxy.h b/Engine/Source/Runtime/Experimental/Chaos/Public/PhysicsProxy/GeometryCollectionPhysicsProxy.h
--- a/Engine/Source/Runtime/Experimental/Chaos/Public/PhysicsProxy/GeometryCollectionPhysicsProxy.h
+++ b/Engine/Source/Runtime/Experimental/Chaos/Public/PhysicsProxy/GeometryCollectionPhysicsProxy.h
@@ -169,6 +169,7 @@ public:
 		FGuid InCollectorGuid = FGuid::NewGuid(),
 		const Chaos::EMultiBufferMode BufferMode=Chaos::EMultiBufferMode::TripleGuarded);
 
+	UE_DEPRECATED(5.8, "Use the constructor that takes FShapeFilterData and FInstanceData instead")
 	CHAOS_API FGeometryCollectionPhysicsProxy(
 		UObject* InOwner,
 		TSharedRef<FGeometryDynamicCollection> GameThreadCollection,
@@ -178,6 +179,15 @@ public:
 		FGuid InCollectorGuid = FGuid::NewGuid(),
 		const Chaos::EMultiBufferMode BufferMode = Chaos::EMultiBufferMode::TripleGuarded);
 
+	CHAOS_API FGeometryCollectionPhysicsProxy(
+		UObject* InOwner,
+		TSharedRef<FGeometryDynamicCollection> GameThreadCollection,
+		const FSimulationParameters& SimulationParameters,
+		const Chaos::Filter::FShapeFilterData& InShapeFilterData,
+		const Chaos::Filter::FInstanceData& InShapeInstanceData,
+		FGuid InCollectorGuid = FGuid::NewGuid(),
+		const Chaos::EMultiBufferMode BufferMode = Chaos::EMultiBufferMode::TripleGuarded);
+
 	CHAOS_API virtual ~FGeometryCollectionPhysicsProxy();
 
 	/**
@@ -521,12 +531,26 @@ public:
 
 	struct FParticleCollisionFilterData
 	{
+		// These must be explicitly declared to get rid of deprecation warnings.
+		PRAGMA_DISABLE_DEPRECATION_WARNINGS
+		FParticleCollisionFilterData() = default;
+		FParticleCollisionFilterData(const FParticleCollisionFilterData&) = default;
+		FParticleCollisionFilterData(FParticleCollisionFilterData&&) = default;
+
+		FParticleCollisionFilterData& operator=(const FParticleCollisionFilterData&) = default;
+		FParticleCollisionFilterData& operator=(FParticleCollisionFilterData&&) = default;
+		PRAGMA_ENABLE_DEPRECATION_WARNINGS
+
 		int32 ParticleIndex = INDEX_NONE;
 		bool bIsValid = false;
 		bool bQueryEnabled = false;
 		bool bSimEnabled = false;
+		UE_DEPRECATED(5.8, "Use ShapeFilterData and FilterInstanceData")
 		FCollisionFilterData QueryFilter;
+		UE_DEPRECATED(5.8, "Use ShapeFilterData and FilterInstanceData")
 		FCollisionFilterData SimFilter;
+		UE_INTERNAL Chaos::Filter::FShapeFilterData ShapeFilterData;
+		UE_INTERNAL Chaos::Filter::FInstanceData FilterInstanceData;
 	};
 
 	CHAOS_API void UpdatePerParticleFilterData_External(const TArray<FParticleCollisionFilterData>& Data);
@@ -752,8 +776,8 @@ private:
 	TMap<int32, TArray<int32>> InternalClusterUniqueIdxToChildrenTransformIndices;
 
 	// These are read on both threads and should not be changed
-	const FCollisionFilterData SimFilter;
-	const FCollisionFilterData QueryFilter;
+	Chaos::Filter::FShapeFilterData ShapeFilterData;
+	Chaos::Filter::FInstanceData ShapeInstanceData;
 
 	// todo(chaos): Remove this and move to a cook time approach of the SM data based on the GC property
 	FCreateTraceCollisionGeometryCallback CreateTraceCollisionGeometryCallback;
diff --git a/Engine/Source/Runtime/Experimental/GeometryCollectionEngine/Private/GeometryCollection/GeometryCollectionComponent.cpp b/Engine/Source/Runtime/Experimental/GeometryCollectionEngine/Private/GeometryCollection/GeometryCollectionComponent.cpp
--- a/Engine/Source/Runtime/Experimental/GeometryCollectionEngine/Private/GeometryCollection/GeometryCollectionComponent.cpp
+++ b/Engine/Source/Runtime/Experimental/GeometryCollectionEngine/Private/GeometryCollection/GeometryCollectionComponent.cpp
@@ -3986,7 +3986,7 @@ PRAGMA_ENABLE_DEPRECATION_WARNINGS
 		FRuntimeDataCollector::GetInstance().RemoveCollector(CollectorGuid);
 	}
 #endif
-	PhysicsProxy = new FGeometryCollectionPhysicsProxy(this, DynamicCollection.ToSharedRef(), SimulationParameters, InitialSimFilter, InitialQueryFilter, CollectorGuid);
+	PhysicsProxy = new FGeometryCollectionPhysicsProxy(this, DynamicCollection.ToSharedRef(), SimulationParameters, InitialShapeFilterData, InitialFilterInstanceData, CollectorGuid);
 
 	PhysicsProxy->SetPostPhysicsSyncCallback(
 		[WeakThis = MakeWeakObjectPtr(this)]()
@@ -4803,8 +4803,8 @@ void UGeometryCollectionComponent::LoadCollisionProfiles()
 
 		if (ProfileName == DefaultCollisionProfileName)
 		{
-			Cache.QueryFilter = InitialQueryFilter;
-			Cache.SimFilter = InitialSimFilter;
+			Cache.ShapeFilterData = InitialShapeFilterData;
+			Cache.FilterInstanceData = InitialFilterInstanceData;
 
 			Cache.bQueryEnabled = true;
 			Cache.bSimEnabled = true;
@@ -4826,8 +4826,8 @@ void UGeometryCollectionComponent::LoadCollisionProfiles()
 			FilterBuilder.SetFlags(Chaos::EFilterFlags::CCD, BodyInstance.bUseCCD);
 			FilterBuilder.SetFlags(Chaos::EFilterFlags::ContactNotify, bNotifyCollisions || bNotifyGlobalCollisions);
 
-			const Chaos::Filter::FCombinedShapeFilterData CombinedFilter(FilterBuilder.BuildShapeFilterData(), FilterBuilder.BuildInstanceData());
-			Chaos::Filter::FShapeFilterBuilder::GetLegacyShapeFilter(CombinedFilter, Cache.QueryFilter, Cache.SimFilter);
+			Cache.FilterInstanceData = FilterBuilder.BuildInstanceData();
+			Cache.ShapeFilterData = FilterBuilder.BuildShapeFilterData();
 		
 			Cache.bQueryEnabled = CollisionEnabledHasQuery(Template.CollisionEnabled);
 			Cache.bSimEnabled = CollisionEnabledHasPhysics(Template.CollisionEnabled);
@@ -4881,7 +4881,7 @@ void UGeometryCollectionComponent::LoadCollisionProfiles()
 
 				// Need to update on the GT too (physics proxy will just enqueue to do it on the PT) or else the changes won't take effect for GT SQ queries.
 				Interface->UpdateShapeCollisionFlags(ParticleView, ParticleData.bSimEnabled, ParticleData.bQueryEnabled);
-				Interface->UpdateCombinedShapeFilterData(ParticleView, Chaos::Filter::FShapeFilterBuilder::BuildFromLegacyShapeFilter(ParticleData.QueryFilter, ParticleData.SimFilter));
+				Interface->UpdateCombinedShapeFilterData(ParticleView, Chaos::Filter::FCombinedShapeFilterData(ParticleData.ShapeFilterData, ParticleData.FilterInstanceData));
 			}
 		}
 	}
@@ -4977,7 +4977,7 @@ void UGeometryCollectionComponent::SetCollisionObjectType(ECollisionChannel Chan
 	// Update filters stored on proxy
 	if (PhysicsProxy)
 	{
-		PhysicsProxy->UpdateFilterData_External(Chaos::Filter::FShapeFilterBuilder::BuildFromLegacyShapeFilter(InitialQueryFilter, InitialSimFilter));
+		PhysicsProxy->UpdateFilterData_External(Chaos::Filter::FCombinedShapeFilterData(InitialShapeFilterData, InitialFilterInstanceData));
 	}
 }
 
@@ -4986,13 +4986,13 @@ void UGeometryCollectionComponent::OnActorEnableCollisionChanged()
 	// Update filters on BI
 	BodyInstance.UpdatePhysicsFilterData();
 
-	// Update InitialSimFilter and InitialQueryFilter
+	// Update Shape Filter and Instance data
 	BuildInitialFilterData();
 
 	// Update filters stored on proxy
 	if (PhysicsProxy)
 	{
-		PhysicsProxy->UpdateFilterData_External(Chaos::Filter::FShapeFilterBuilder::BuildFromLegacyShapeFilter(InitialQueryFilter, InitialSimFilter));
+		PhysicsProxy->UpdateFilterData_External(Chaos::Filter::FCombinedShapeFilterData(InitialShapeFilterData, InitialFilterInstanceData));
 	}
 
 }
@@ -5000,21 +5000,19 @@ void UGeometryCollectionComponent::OnActorEnableCollisionChanged()
 void UGeometryCollectionComponent::BuildInitialFilterData()
 {
 	FBodyCollisionFilterData FilterData;
-	FMaskFilter FilterMask = BodyInstance.GetMaskFilter();
 	BodyInstance.BuildBodyFilterData(FilterData);
-
-	Chaos::Filter::FCombinedShapeFilterData CombinedFilter(FilterData.SimpleShapeFilterData, FilterData.FilterInstanceData);
-	Chaos::Filter::FShapeFilterBuilder::GetLegacyShapeFilter(CombinedFilter, InitialQueryFilter, InitialSimFilter);
-
+	
+	FPhysicsFilterBuilder Builder(GetInnerChaosScene(), FilterData.SimpleShapeFilterData, FilterData.FilterInstanceData);
 	// Enable for complex and simple (no dual representation currently like other meshes)
-	InitialQueryFilter.Word3 |= (EPDF_SimpleCollision | EPDF_ComplexCollision);
-	InitialSimFilter.Word3 |= (EPDF_SimpleCollision | EPDF_ComplexCollision);
-
+	Builder.SetFlags(Chaos::EFilterFlags::SimpleCollision | Chaos::EFilterFlags::ComplexCollision, true);	
+	// Should this actually clear the flag if the body is true but the geometry component is false?
 	if (bNotifyCollisions || bNotifyGlobalCollisions)
 	{
-		InitialQueryFilter.Word3 |= EPDF_ContactNotify;
-		InitialSimFilter.Word3 |= EPDF_ContactNotify;
+		Builder.SetFlags(Chaos::EFilterFlags::ContactNotify, true);
 	}
+
+	InitialFilterInstanceData = Builder.BuildInstanceData();
+	InitialShapeFilterData = Builder.BuildShapeFilterData();
 }
 
 void UGeometryCollectionComponent::ApplyAssetDefaults()
diff --git a/Engine/Source/Runtime/Experimental/GeometryCollectionEngine/Public/GeometryCollection/GeometryCollectionComponent.h b/Engine/Source/Runtime/Experimental/GeometryCollectionEngine/Public/GeometryCollection/GeometryCollectionComponent.h
--- a/Engine/Source/Runtime/Experimental/GeometryCollectionEngine/Public/GeometryCollection/GeometryCollectionComponent.h
+++ b/Engine/Source/Runtime/Experimental/GeometryCollectionEngine/Public/GeometryCollection/GeometryCollectionComponent.h
@@ -1689,8 +1689,8 @@ private:
 
 	uint32 NavmeshInvalidationTimeSliceIndex;
 
-	FCollisionFilterData InitialSimFilter;
-	FCollisionFilterData InitialQueryFilter;
+	Chaos::Filter::FShapeFilterData InitialShapeFilterData;
+	Chaos::Filter::FInstanceData InitialFilterInstanceData;
 	float CurrentCacheTime;
 	FChaosUserData PhysicsUserData;
 
