Wait for inflight request before shutting down HTTP I/O dispatcher.

#rnx

diff --git a/Engine/Source/Runtime/Experimental/IoStore/OnDemand/Private/OnDemandHttpIoDispacher.cpp b/Engine/Source/Runtime/Experimental/IoStore/OnDemand/Private/OnDemandHttpIoDispacher.cpp
--- a/Engine/Source/Runtime/Experimental/IoStore/OnDemand/Private/OnDemandHttpIoDispacher.cpp
+++ b/Engine/Source/Runtime/Experimental/IoStore/OnDemand/Private/OnDemandHttpIoDispacher.cpp
@@ -49,6 +49,14 @@ static FAutoConsoleVariableRef CVar_HttpIoDispatcherMaxRangeKiB(
 	ECVF_Default
 );
 
+float GHttpIoDispatcherShutdownTimeoutSeconds = 5.0;
+static FAutoConsoleVariableRef CVar_HttpIoDispatcherShutdownTimeoutSeconds(
+	TEXT("iax.HttpIoDispatcherShutdownTimeoutSeconds"),
+	GHttpIoDispatcherShutdownTimeoutSeconds,
+	TEXT("Sets the maximum time the dispatcher will wait for outstanding request(s) when shutting down."),
+	ECVF_Default
+);
+
 ////////////////////////////////////////////////////////////////////////////////
 DECLARE_MULTICAST_DELEGATE(FOnRecreateHttpClient);
 extern FOnRecreateHttpClient OnRecreateHttpClient;
@@ -764,7 +772,8 @@ private:
 
 	bool							TryReadFromCache(FHttpRequest* Request);
 	bool							ProcessHttpRequests();
-	FHttpRequest*					GetNextRequest();
+	void							CancelRemainingHttpRequests();
+	FHttpRequest*					DequeueHttpRequests();
 	void							CompleteHttpRequest(FHttpRequest& Request, FMultiEndpointHttpClientResponse&& HttpResponse, FIASHostGroup& HostGroup, const FIoHttpRange& TotalRange);
 	void							CompleteRequest(FHttpRequest& Request, TArray<FAnsiString>&& Headers, FIoBuffer& Body, const FIoHttpRange& TotalRange, uint32 StatusCode, bool bCached);
 	void							CompleteRequest(FHttpRequest& Request, FIoBuffer& Body, const FIoHttpRange& TotalRange, uint32 StatusCode, bool bCached)
@@ -788,6 +797,7 @@ private:
 	FDelegateHandle							OnRecreateHttpClientHandle;
 	std::atomic_bool						bRecreateHttpClient = false;
 	std::atomic_bool						bStopRequested = false;
+	std::atomic_int32_t						InflightRequestCount{0};
 
 #if UE_ALLOW_HTTP_PAUSE
 	FConsoleCommandPtr						PauseCommand;
@@ -889,6 +899,29 @@ uint32 FHttpDispatcher::Run()
 		WakeUp->Wait(FHostGroupManager::Get().GetNumDisconnctedHosts() > 0 ? GIasHttpHealthCheckWaitTime : MAX_uint32);
 	}
 
+	CancelRemainingHttpRequests();
+
+	{
+		const float SleepTimeSeconds	= 0.5f;
+		const double WaitTimeoutSeconds	= GHttpIoDispatcherShutdownTimeoutSeconds;
+		const double StartTimeSeconds	= FPlatformTime::Seconds();
+		int32 Remaining					= InflightRequestCount.load(std::memory_order_seq_cst);
+
+		UE_CLOG(Remaining > 0, LogHttpIoDispatcher, Log, TEXT("Waiting on %d request(s) before shutting down..."), Remaining);
+
+		while (Remaining > 0)
+		{
+			if ((FPlatformTime::Seconds() - StartTimeSeconds) > WaitTimeoutSeconds)
+			{
+				UE_LOG(LogIoStoreOnDemand, Warning, TEXT("Shutting down with %d inflight request(s) after waiting %.2lf second(s)"), Remaining, WaitTimeoutSeconds);
+				break;
+			}
+
+			FPlatformProcess::SleepNoStats(SleepTimeSeconds);
+			Remaining = InflightRequestCount.load(std::memory_order_seq_cst);
+		}
+	}
+
 	return 0;
 }
 
@@ -1106,6 +1139,7 @@ void FHttpDispatcher::CancelRequest(FIoHttpRequestHandle Handle)
 
 	if (HttpQueue.Cancel(Request))
 	{
+		InflightRequestCount.fetch_add(1, std::memory_order_relaxed);
 		UE::Tasks::Launch(UE_SOURCE_LOCATION, [this, Request]
 		{
 			CompleteRequest(*Request, Request->Buffer, FIoHttpRange(), 0, false);
@@ -1257,6 +1291,7 @@ bool FHttpDispatcher::TryReadFromCache(FHttpRequest* Request)
 	const EIoErrorCode CacheStatus = Cache->Get(Request->CacheKey, Request->Buffer);
 	if (CacheStatus == EIoErrorCode::Ok)
 	{
+		InflightRequestCount.fetch_add(1, std::memory_order_relaxed);
 		FOnDemandIoBackendStats::Get()->OnCacheGet(Request->Buffer.GetSize());
 		UE::Tasks::Launch(UE_SOURCE_LOCATION, [this, Request]
 		{
@@ -1268,6 +1303,7 @@ bool FHttpDispatcher::TryReadFromCache(FHttpRequest* Request)
 
 	if (CacheStatus == EIoErrorCode::FileNotOpen)
 	{
+		InflightRequestCount.fetch_add(1, std::memory_order_relaxed);
 		UE::Tasks::FTaskEvent OnReadyEvent(TEXT("HttpCacheReadComplete"));
 		Launch(UE_SOURCE_LOCATION, [this, Request]
 		{
@@ -1276,7 +1312,7 @@ bool FHttpDispatcher::TryReadFromCache(FHttpRequest* Request)
 				FOnDemandIoBackendStats::Get()->OnCacheGet(Request->Buffer.GetSize());
 				CompleteRequest(*Request, Request->Buffer, FIoHttpRange(), 200, true);
 			}
-			else if ((Request->CacheStatus == EIoErrorCode::Cancelled) || Request->Flags.HasAny(EHttpRequestFlags::CancelRequested))
+			else if ((Request->CacheStatus == EIoErrorCode::Cancelled) || Request->Flags.HasAny(EHttpRequestFlags::CancelRequested) || bStopRequested.load(std::memory_order_relaxed))
 			{
 				CompleteRequest(*Request, Request->Buffer, FIoHttpRange(), 0, false);
 			}
@@ -1287,6 +1323,7 @@ bool FHttpDispatcher::TryReadFromCache(FHttpRequest* Request)
 					FOnDemandIoBackendStats::Get()->OnCacheError();
 				}
 
+				ensure(InflightRequestCount.fetch_sub(1, std::memory_order_relaxed) > 0);
 				HttpQueue.Enqueue(Request);
 				WakeUp->Trigger();
 			}
@@ -1328,7 +1365,7 @@ bool FHttpDispatcher::ProcessHttpRequests()
 	const int32 MaxConcurrentRequests = FMath::Min(GIasHttpConcurrentRequests, 32);
 	int32		NumConcurrentRequests = 0;
 
-	FHttpRequest* NextHttpRequest = HttpQueue.Dequeue();
+	FHttpRequest* NextHttpRequest = DequeueHttpRequests();
 	if (NextHttpRequest == nullptr)
 	{
 		return false;
@@ -1357,7 +1394,7 @@ bool FHttpDispatcher::ProcessHttpRequests()
 				CompleteRequest(*HttpRequest, HttpRequest->Buffer, FIoHttpRange(), 0, false);
 
 				check(NextHttpRequest == nullptr);
-				NextHttpRequest = GetNextRequest();
+				NextHttpRequest = DequeueHttpRequests();
 				continue;
 			}
 
@@ -1441,7 +1478,7 @@ bool FHttpDispatcher::ProcessHttpRequests()
 			}
 
 			check(NextHttpRequest == nullptr);
-			NextHttpRequest = GetNextRequest();
+			NextHttpRequest = DequeueHttpRequests();
 		} // Inner
 
 		{
@@ -1451,7 +1488,7 @@ bool FHttpDispatcher::ProcessHttpRequests()
 			TRACE_CPUPROFILER_EVENT_SCOPE(FHttpDispatcher::TickHttp);
 			while (NextHttpRequest == nullptr && HttpClient->Tick(GIasHttpPollTimeoutMs, GIasHttpRateLimitKiBPerSecond))
 			{
-				NextHttpRequest = GetNextRequest();
+				NextHttpRequest = DequeueHttpRequests();
 			}
 		}
 	} // Outer
@@ -1462,16 +1499,36 @@ bool FHttpDispatcher::ProcessHttpRequests()
 	return true;
 }
 
-FHttpRequest* FHttpDispatcher::GetNextRequest()
+void FHttpDispatcher::CancelRemainingHttpRequests()
 {
-	if (!bRecreateHttpClient)
+	while (FHttpRequest* Requests = HttpQueue.Dequeue())
 	{
-		return HttpQueue.Dequeue();
+		for (FHttpRequest* It = Requests; It != nullptr; It = It->Next)
+		{
+			It->Flags.Add(EHttpRequestFlags::CancelRequested);
+		}
+
+		InflightRequestCount.fetch_add(1, std::memory_order_relaxed);
+		FIoBuffer Buffer;
+		CompleteRequest(*Requests, Buffer, FIoHttpRange(), 0, false);
 	}
-	else
+}
+
+FHttpRequest* FHttpDispatcher::DequeueHttpRequests()
+{
+	// Return nullptr to break the process loop if we need to re-create the HTTP client or we are shutting down
+	if (bRecreateHttpClient.load(std::memory_order_relaxed) || bStopRequested.load(std::memory_order_relaxed))
 	{
 		return nullptr;
 	}
+
+	if (FHttpRequest* Requests = HttpQueue.Dequeue())
+	{
+		InflightRequestCount.fetch_add(1, std::memory_order_relaxed);
+		return Requests;
+	}
+
+	return nullptr;
 }
 
 void FHttpDispatcher::CompleteHttpRequest(FHttpRequest& Request, FMultiEndpointHttpClientResponse&& HttpResponse, FIASHostGroup& HostGroup, const FIoHttpRange& TotalRange)
@@ -1565,6 +1622,8 @@ void FHttpDispatcher::CompleteRequest(FHttpRequest& Request, TArray<FAnsiString>
 		ToComplete->CompletionStatus.store(CompletionStatus, std::memory_order_seq_cst);
 		ToComplete->Release();
 	}
+
+	ensure(InflightRequestCount.fetch_sub(1, std::memory_order_relaxed) > 0);
 }
 
 #if UE_ALLOW_HTTP_PAUSE
