IAD
- Fix block timestamp update on content handle release logic running on the game thread by making it always run in the content installer.
- Simplify logic to only check for chunk decref and not full release.  The extra operations to check for full release are not worth it becuase timestamps updates are throttled to 1hr granularity anyway.

diff --git a/Engine/Source/Runtime/Experimental/IoStore/OnDemand/Private/OnDemandContentInstaller.cpp b/Engine/Source/Runtime/Experimental/IoStore/OnDemand/Private/OnDemandContentInstaller.cpp
--- a/Engine/Source/Runtime/Experimental/IoStore/OnDemand/Private/OnDemandContentInstaller.cpp
+++ b/Engine/Source/Runtime/Experimental/IoStore/OnDemand/Private/OnDemandContentInstaller.cpp
@@ -353,11 +353,11 @@ void FOnDemandContentInstaller::EnqueueVerifyRequest(FOnDemandVerifyCacheComplet
 	TryExecuteNextRequest();
 }
 
-void FOnDemandContentInstaller::EnqueueFlushLastAccessRequest(FOnDemandFlushLastAccessCompleted&& OnCompleted)
+void FOnDemandContentInstaller::EnqueueFlushLastAccessRequest(FOnDemandFlushLastAccessArgs&& Args, FOnDemandFlushLastAccessCompleted&& OnCompleted)
 {
 	{
 		TUniqueLock Lock(Mutex);
-		RequestQueue.HeapPush(RequestAllocator.Construct(MoveTemp(OnCompleted)), RequestSortPredicate);
+		RequestQueue.HeapPush(RequestAllocator.Construct(MoveTemp(Args), MoveTemp(OnCompleted)), RequestSortPredicate);
 	}
 
 	TryExecuteNextRequest();
@@ -1119,9 +1119,37 @@ void FOnDemandContentInstaller::ExectuteFlushLastAccessRequest(FRequest& Request
 
 	TRACE_CPUPROFILER_EVENT_SCOPE(FOnDemandContentInstaller::ExectuteFlushLastAccessRequest);
 
+	const FRequest::FFlushLastAccess& FlushLastAccessReqest = Request.AsFlushLastAccess();
+
 	UE_LOG(LogIoStoreOnDemand, Log, TEXT("Executing flush last access request"));
 
-	if (Request.IsCancelled() == false)
+	if (Request.IsCancelled())
+	{
+		CompleteFlushLastAccessRequest(Request);
+		return;
+	}
+
+	TArray<FIoHash> TouchedChunks;
+	for (const TPair<FSharedOnDemandContainer, TBitArray<>>& Pair : FlushLastAccessReqest.Args.RemovedChunkEntryReferences)
+	{
+		const FSharedOnDemandContainer& Container = Pair.Key;
+		const TBitArray<>& RemovedRefs = Pair.Value;
+		check(RemovedRefs.Num() == Container->ChunkEntries.Num());
+
+		for (TBitArray<>::FConstIterator It(RemovedRefs); It; ++It)
+		{
+			if (It.GetValue())
+			{
+				TouchedChunks.Add(Container->ChunkEntries[It.GetIndex()].Hash);
+			}
+		}
+	}
+
+	bool bLastAccessDirty = false;
+	IoStore.InstallCache->UpdateLastAccess(TouchedChunks, bLastAccessDirty);
+	bLastAccessDirty = bLastAccessDirty || FlushLastAccessReqest.Args.bForceLastAccessDirty;
+
+	if (bLastAccessDirty)
 	{
 		Request.Result = IoStore.InstallCache->FlushLastAccess();
 	}
diff --git a/Engine/Source/Runtime/Experimental/IoStore/OnDemand/Private/OnDemandContentInstaller.h b/Engine/Source/Runtime/Experimental/IoStore/OnDemand/Private/OnDemandContentInstaller.h
--- a/Engine/Source/Runtime/Experimental/IoStore/OnDemand/Private/OnDemandContentInstaller.h
+++ b/Engine/Source/Runtime/Experimental/IoStore/OnDemand/Private/OnDemandContentInstaller.h
@@ -127,9 +127,11 @@ class FOnDemandContentInstaller
 
 		struct FFlushLastAccess
 		{
-			FFlushLastAccess(FOnDemandFlushLastAccessCompleted&& InOnCompleted)
-				: OnCompleted(MoveTemp(InOnCompleted)) { }
+			FFlushLastAccess(FOnDemandFlushLastAccessArgs&& InArgs, FOnDemandFlushLastAccessCompleted&& InOnCompleted)
+				: Args(MoveTemp(InArgs))
+				, OnCompleted(MoveTemp(InOnCompleted)) { }
 
+			FOnDemandFlushLastAccessArgs		Args;
 			FOnDemandFlushLastAccessCompleted	OnCompleted;
 		};
 
@@ -155,9 +157,9 @@ class FOnDemandContentInstaller
 			Variant.Emplace<FVerify>(MoveTemp(OnCompleted));
 		}
 
-		FRequest(FOnDemandFlushLastAccessCompleted&& OnCompleted)
+		FRequest(FOnDemandFlushLastAccessArgs&& Args, FOnDemandFlushLastAccessCompleted&& OnCompleted)
 		{
-			Variant.Emplace<FFlushLastAccess>(MoveTemp(OnCompleted));
+			Variant.Emplace<FFlushLastAccess>(MoveTemp(Args), MoveTemp(OnCompleted));
 		}
 
 		bool				IsCancelled() const			{ return bCancelled; }
@@ -232,7 +234,7 @@ public:
 	void							EnqueuePurgeRequest(FOnDemandPurgeArgs&& Args, FOnDemandPurgeCompleted&& OnCompleted);
 	void							EnqueueDefragRequest(FOnDemandDefragArgs&& Args, FOnDemandDefragCompleted&& OnCompleted);
 	void							EnqueueVerifyRequest(FOnDemandVerifyCacheCompleted&& OnCompleted);
-	void							EnqueueFlushLastAccessRequest(FOnDemandFlushLastAccessCompleted&& OnCompleted);
+	void							EnqueueFlushLastAccessRequest(FOnDemandFlushLastAccessArgs&& Args, FOnDemandFlushLastAccessCompleted&& OnCompleted);
 	void							CancelInstallRequest(FSharedInternalInstallRequest InstallRequest);
 	void							UpdateInstallRequestPriority(FSharedInternalInstallRequest InstallRequest, int32 NewPriority);
 	void							ReportAnalytics(TArray<FAnalyticsEventAttribute>& OutAnalyticsArray) const;
diff --git a/Engine/Source/Runtime/Experimental/IoStore/OnDemand/Private/OnDemandInstallCache.cpp b/Engine/Source/Runtime/Experimental/IoStore/OnDemand/Private/OnDemandInstallCache.cpp
--- a/Engine/Source/Runtime/Experimental/IoStore/OnDemand/Private/OnDemandInstallCache.cpp
+++ b/Engine/Source/Runtime/Experimental/IoStore/OnDemand/Private/OnDemandInstallCache.cpp
@@ -1745,7 +1745,7 @@ public:
 	virtual FResult								Verify() override;
 	virtual FResult								Flush() override;
 	virtual FResult								FlushLastAccess() override;
-	virtual void								UpdateLastAccess(TConstArrayView<FIoHash> ChunkHashes) override;
+	virtual void								UpdateLastAccess(TConstArrayView<FIoHash> ChunkHashes, bool& bInOutLastAccessDirty) override;
 	virtual FOnDemandInstallCacheUsage			GetCacheUsage() override;
 
 private:
@@ -2157,7 +2157,7 @@ bool FOnDemandInstallCache::Resolve(FIoRequestImpl* Request)
 	constexpr const bool bDirty = true;
 	if (Cas.TrackAccessIf(ECasTrackAccessType::Granular, CasLoc.BlockId, bDirty))
 	{
-		IoStore.FlushLastAccess(nullptr);
+		IoStore.FlushLastAccess({}, nullptr);
 	}
 
 #if UE_ONDEMANDINSTALLCACHE_EXCLUSIVE_WRITE
@@ -3380,32 +3380,29 @@ FResult FOnDemandInstallCache::FlushLastAccess()
 	return MakeValue();
 }
 
-void FOnDemandInstallCache::UpdateLastAccess(TConstArrayView<FIoHash> ChunkHashes)
+void FOnDemandInstallCache::UpdateLastAccess(TConstArrayView<FIoHash> ChunkHashes, bool& bInOutLastAccessDirty)
 {
-	bool bLastAccessDirty = false;
+	if (ChunkHashes.IsEmpty())
 	{
-		const int64 Now = FDateTime::UtcNow().GetTicks();
-		constexpr const bool bDirty = true;
+		return;
+	}
 
-		TUniqueLock Lock(Cas);
+	const int64 Now = FDateTime::UtcNow().GetTicks();
+	constexpr const bool bDirty = true;
 
-		for (const FIoHash& ChunkHash : ChunkHashes)
-		{
-			const FCasAddr& Addr = AsCasAddr(ChunkHash);
-			if (FCasLocation* CasLoc = Cas.Lookup.Find(Addr))
-			{
-				const FCasBlockId BlockId = CasLoc->BlockId;
-				const uint32 BlockIdHash = GetTypeHash(BlockId);
-				const bool bUpdatedLastAccess = Cas.UnlockedTrackAccessIf(
-					ECasTrackAccessType::Granular, BlockIdHash, BlockId, Now, bDirty);
-				bLastAccessDirty = bLastAccessDirty || bUpdatedLastAccess;
-			}
-		}
-	}
+	TUniqueLock Lock(Cas);
 
-	if (bLastAccessDirty)
+	for (const FIoHash& ChunkHash : ChunkHashes)
 	{
-		IoStore.FlushLastAccess(nullptr);
+		const FCasAddr& Addr = AsCasAddr(ChunkHash);
+		if (FCasLocation* CasLoc = Cas.Lookup.Find(Addr))
+		{
+			const FCasBlockId BlockId = CasLoc->BlockId;
+			const uint32 BlockIdHash = GetTypeHash(BlockId);
+			const bool bUpdatedLastAccess = Cas.UnlockedTrackAccessIf(
+				ECasTrackAccessType::Granular, BlockIdHash, BlockId, Now, bDirty);
+			bInOutLastAccessDirty = bInOutLastAccessDirty || bUpdatedLastAccess;
+		}
 	}
 }
 
diff --git a/Engine/Source/Runtime/Experimental/IoStore/OnDemand/Private/OnDemandInstallCache.h b/Engine/Source/Runtime/Experimental/IoStore/OnDemand/Private/OnDemandInstallCache.h
--- a/Engine/Source/Runtime/Experimental/IoStore/OnDemand/Private/OnDemandInstallCache.h
+++ b/Engine/Source/Runtime/Experimental/IoStore/OnDemand/Private/OnDemandInstallCache.h
@@ -42,7 +42,7 @@ public:
 	virtual FResult								Verify() = 0;
 	virtual FResult								Flush() = 0;
 	virtual FResult								FlushLastAccess() = 0;
-	virtual void								UpdateLastAccess(TConstArrayView<FIoHash> ChunkHashes) = 0;
+	virtual void								UpdateLastAccess(TConstArrayView<FIoHash> ChunkHashes, bool& bInOutLastAccessDirty) = 0;
 	virtual FOnDemandInstallCacheUsage			GetCacheUsage() = 0;
 };
 
diff --git a/Engine/Source/Runtime/Experimental/IoStore/OnDemand/Private/OnDemandIoStore.cpp b/Engine/Source/Runtime/Experimental/IoStore/OnDemand/Private/OnDemandIoStore.cpp
--- a/Engine/Source/Runtime/Experimental/IoStore/OnDemand/Private/OnDemandIoStore.cpp
+++ b/Engine/Source/Runtime/Experimental/IoStore/OnDemand/Private/OnDemandIoStore.cpp
@@ -830,7 +830,7 @@ TIoStatusOr<FOnDemandInstallSizeResult> FOnDemandIoStore::GetInstallSize(const F
 		Unreferenced.Reset();
 		{
 			TUniqueLock Lock(ContainerMutex);
-			FOnDemandChunkEntryReferences& Refs = Resolved.Container->FindOrAddChunkEntryReferences(*Args.ContentHandle.Handle);
+			const FOnDemandChunkEntryReferences& Refs = Resolved.Container->FindOrAddChunkEntryReferences(*Args.ContentHandle.Handle);
 			for (int32 EntryIndex : Resolved.EntryIndices)
 			{
 				const bool bReferenced = Refs.Indices[EntryIndex];
@@ -1844,10 +1844,10 @@ TArray<FSharedOnDemandContainer> FOnDemandIoStore::GetContainers(EOnDemandContai
 	return Out;
 }
 
-void FOnDemandIoStore::FlushLastAccess(FOnDemandFlushLastAccessCompleted&& OnCompleted)
+void FOnDemandIoStore::FlushLastAccess(FOnDemandFlushLastAccessArgs&& Args, FOnDemandFlushLastAccessCompleted&& OnCompleted)
 {
 	LLM_SCOPE_BYNAME(TEXT("FileSystem/OnDemandIoStore"));
-	Installer->EnqueueFlushLastAccessRequest(MoveTemp(OnCompleted));
+	Installer->EnqueueFlushLastAccessRequest(MoveTemp(Args), MoveTemp(OnCompleted));
 }
 
 FIoStatus FOnDemandIoStore::GetContainersAndPackagesForInstall(
@@ -1958,7 +1958,8 @@ void FOnDemandIoStore::ReleaseContent(FOnDemandInternalContentHandle& ContentHan
 	LLM_SCOPE_BYNAME(TEXT("FileSystem/OnDemandIoStore"));
 	UE_LOG(LogIoStoreOnDemand, Verbose, TEXT("Releasing content handle, %s"), *LexToString(ContentHandle));
 
-	TArray<FIoHash> ReleasedChunks;
+	FOnDemandFlushLastAccessArgs FlushLastAccessArgs;
+	FlushLastAccessArgs.bForceLastAccessDirty = false;
 
 	{
 		UE::TUniqueLock Lock(ContainerMutex);
@@ -1986,19 +1987,11 @@ void FOnDemandIoStore::ReleaseContent(FOnDemandInternalContentHandle& ContentHan
 			// referencer is not added the same container more than once. So once the refrencer is removed,
 			// we only need to inspect the current container for unreffed chunks.
 
-			FOnDemandChunkEntryReferences RemovedRefs = MoveTemp(Container->ChunkEntryReferences[RemoveIndex]);
+			FlushLastAccessArgs.RemovedChunkEntryReferences.Add(MakeTuple(
+				Container,
+				MoveTemp(Container->ChunkEntryReferences[RemoveIndex].Indices)
+			));
 			Container->ChunkEntryReferences.RemoveAtSwap(RemoveIndex);
-			check(RemovedRefs.Indices.Num() == Container->ChunkEntries.Num());
-
-			TBitArray<> ReffedChunks = Container->GetReferencedChunkEntries();
-
-			for (int32 i = 0; i < RemovedRefs.Indices.Num(); ++i)
-			{
-				if (RemovedRefs.Indices[i] && (ReffedChunks.IsEmpty() || !ReffedChunks[i]))
-				{
-					ReleasedChunks.Add(Container->ChunkEntries[i].Hash);
-				}
-			}
 		}
 	}
 
@@ -2008,9 +2001,9 @@ void FOnDemandIoStore::ReleaseContent(FOnDemandInternalContentHandle& ContentHan
 		PackageStoreBackend->NeedsUpdate(EOnDemandPackageStoreUpdateMode::ReferencedPackages);
 	}
 
-	if (ensureMsgf(InstallCache, TEXT("InstallCache is null, is this a prefork server?")))
+	if (FlushLastAccessArgs.RemovedChunkEntryReferences.IsEmpty() == false)
 	{
-		InstallCache->UpdateLastAccess(ReleasedChunks);
+		FlushLastAccess(MoveTemp(FlushLastAccessArgs), nullptr);
 	}
 }
 
diff --git a/Engine/Source/Runtime/Experimental/IoStore/OnDemand/Private/OnDemandIoStore.h b/Engine/Source/Runtime/Experimental/IoStore/OnDemand/Private/OnDemandIoStore.h
--- a/Engine/Source/Runtime/Experimental/IoStore/OnDemand/Private/OnDemandIoStore.h
+++ b/Engine/Source/Runtime/Experimental/IoStore/OnDemand/Private/OnDemandIoStore.h
@@ -311,6 +311,13 @@ void FOnDemandChunkInfo::GetUrl(FAnsiStringBuilderBase& Url) const
 }
 
 ///////////////////////////////////////////////////////////////////////////////
+/** Flush last access args. */
+struct FOnDemandFlushLastAccessArgs
+{
+	TArray<TPair<FSharedOnDemandContainer, TBitArray<>>> RemovedChunkEntryReferences;
+	bool bForceLastAccessDirty = true;
+};
+
 /** Result from flushing cache last access times. */
 struct FOnDemandFlushLastAccessResult
 {
@@ -386,7 +393,7 @@ public:
 	virtual void						ReleaseContent(FOnDemandInternalContentHandle& ContentHandle) override;
 
 	TArray<FSharedOnDemandContainer>	GetContainers(EOnDemandContainerFlags ContainerFlags = EOnDemandContainerFlags::None) const;
-	void								FlushLastAccess(FOnDemandFlushLastAccessCompleted&& OnCompleted);
+	void								FlushLastAccess(FOnDemandFlushLastAccessArgs&& Args, FOnDemandFlushLastAccessCompleted&& OnCompleted);
 
 private:
 	friend class FOnDemandContentInstaller;
