Maintain request order explicitly as cutting corners has undesirable edge cases.

#rnx

diff --git a/Engine/Source/Runtime/Experimental/IoStore/HttpClient/Private/Loop.inl b/Engine/Source/Runtime/Experimental/IoStore/HttpClient/Private/Loop.inl
--- a/Engine/Source/Runtime/Experimental/IoStore/HttpClient/Private/Loop.inl
+++ b/Engine/Source/Runtime/Experimental/IoStore/HttpClient/Private/Loop.inl
@@ -600,11 +600,8 @@ private:
 ////////////////////////////////////////////////////////////////////////////////
 void FWorkQueue::AddActivity(FActivityNode* Activity)
 {
-	// This prepends and thus reverses order. It is assumed this this is working
-	// in conjunction with the loop, negating the reversal what happens there.
-
 	check(Activity->Next == nullptr);
-	List.Prepend(Activity);
+	List.Append(Activity);
 
 	ActiveSlots |= (1ull << Activity->Slot);
 }
@@ -1352,9 +1349,7 @@ void FEventLoop::FImpl::ReceiveWork()
 		FScopeLock _(&Lock);
 		Swap(Collected, Pending);
 	}
-
-	// Pending is in the reverse of the order that requests were made. Adding
-	// activities to their corresponding group will reverse this reversal.
+	Collected.Reverse();
 
 	// Group activities by their host.
 	FActivityNode* Activity = Collected.Detach();
