Calculate TTFB after negotiation has happened.

#rnx

diff --git a/Engine/Source/Runtime/Experimental/IoStore/HttpClient/Private/Loop.inl b/Engine/Source/Runtime/Experimental/IoStore/HttpClient/Private/Loop.inl
--- a/Engine/Source/Runtime/Experimental/IoStore/HttpClient/Private/Loop.inl
+++ b/Engine/Source/Runtime/Experimental/IoStore/HttpClient/Private/Loop.inl
@@ -807,6 +807,7 @@ void FPeerGroup::Negotiate(FTickState& State)
 		return;
 	}
 
+	LastUseMs = GetNowMs(); // exclude connect/neg. from ttfb timing.
 	bNegotiating = false;
 	return SendInternal(State);
 }
@@ -1021,7 +1022,7 @@ void FPeerGroup::TickSend(FTickState& State, FHost& Host, FPoller& Poller)
 	}
 
 	// If we are idle we can time the time-to-first-byte
-	if ((InflightNum == 0) && Recv.IsEmpty())
+	if (!bNegotiating && (InflightNum == 0) && Recv.IsEmpty())
 	{
 		LastUseMs = GetNowMs();
 	}
