As add appends, we don't need to reverse when processing cancels.

#rnx

diff --git a/Engine/Source/Runtime/Experimental/IoStore/HttpClient/Private/Loop.inl b/Engine/Source/Runtime/Experimental/IoStore/HttpClient/Private/Loop.inl
--- a/Engine/Source/Runtime/Experimental/IoStore/HttpClient/Private/Loop.inl
+++ b/Engine/Source/Runtime/Experimental/IoStore/HttpClient/Private/Loop.inl
@@ -636,9 +636,7 @@ void FWorkQueue::TickCancels(FTickState& State)
 	check(!List.IsEmpty());
 	ActiveSlots = 0;
 
-	List.Reverse();
 	FActivityNode* Activity = List.Detach();
-
 	for (FActivityNode* Next; Activity != nullptr; Activity = Next)
 	{
 		Next = Activity->Next;
