[HWRT] Don't recreate VertexFactory in FGeometryCollectionSceneProxyBase::GetDynamicRayTracingInstances(...)
- simply reuse the main VertexFactory.

diff --git a/Engine/Source/Runtime/Experimental/GeometryCollectionEngine/Private/GeometryCollection/GeometryCollectionSceneProxy.cpp b/Engine/Source/Runtime/Experimental/GeometryCollectionEngine/Private/GeometryCollection/GeometryCollectionSceneProxy.cpp
--- a/Engine/Source/Runtime/Experimental/GeometryCollectionEngine/Private/GeometryCollection/GeometryCollectionSceneProxy.cpp
+++ b/Engine/Source/Runtime/Experimental/GeometryCollectionEngine/Private/GeometryCollection/GeometryCollectionSceneProxy.cpp
@@ -136,26 +136,6 @@ static void UpdateLooseParameter(
 	VertexFactory.LooseParameterUniformBuffer = FGCBoneLooseParametersRef::CreateUniformBufferImmediate(LooseParameters, UniformBufferUsage);
 }
 
-
-class FGeometryCollectionMeshCollectorResources : public FOneFrameResource
-{
-public:
-	FGeometryCollectionVertexFactory VertexFactory;
-	
-	FGeometryCollectionMeshCollectorResources(ERHIFeatureLevel::Type InFeatureLevel)
-		: VertexFactory(InFeatureLevel,true)
-	{
-	}
-
-	virtual ~FGeometryCollectionMeshCollectorResources()
-	{
-		VertexFactory.ReleaseResource();
-	}
-
-	virtual FGeometryCollectionVertexFactory& GetVertexFactory() { return VertexFactory; }
-};
-
-
 FGeometryCollectionSceneProxyBase::FGeometryCollectionSceneProxyBase(UGeometryCollectionComponent* Component, bool bInIsNanite)
 	: bIsNanite(bInIsNanite)
 	, FeatureLevel(Component->GetScene()->GetFeatureLevel())
@@ -581,13 +561,6 @@ void FGeometryCollectionSceneProxyBase::GetDynamicRayTracingInstances(FRayTracin
 	const uint32 LODIndex = 0;
 	const bool bWireframe = false;
 
-	//Loose parameter needs to be updated every frame
-	FGeometryCollectionMeshCollectorResources* CollectorResources;
-	CollectorResources = &Collector.AllocateOneFrameResource<FGeometryCollectionMeshCollectorResources>(FeatureLevel);
-	FGeometryCollectionVertexFactory& GeometryCollectionVertexFactory = CollectorResources->GetVertexFactory();
-
-	SetupVertexFactory(Collector.GetRHICommandList(), GeometryCollectionVertexFactory);
-
 	const bool bUseSubSections = MeshDescription.SubSections.Num() && !bAnyMaterialHasWorldPositionOffset && CVarRayTracingGeometryCollectionCombinedBLAS.GetValueOnRenderThread() == 0;
 
 	TArray<FGeometryCollectionMeshElement> const& SectionArray = bUseSubSections ? MeshDescription.SubSections : MeshDescription.Sections;
@@ -615,7 +588,7 @@ void FGeometryCollectionSceneProxyBase::GetDynamicRayTracingInstances(FRayTracin
 	if (RayTracingGeometry.IsValid())
 	{
 		// Render dynamic objects
-		if (!GeometryCollectionVertexFactory.GetType()->SupportsRayTracingDynamicGeometry())
+		if (!VertexFactory.GetType()->SupportsRayTracingDynamicGeometry())
 		{
 			return;
 		}
@@ -632,7 +605,7 @@ void FGeometryCollectionSceneProxyBase::GetDynamicRayTracingInstances(FRayTracin
 			FMeshBatch& Mesh = RayTracingInstance.Materials.AddDefaulted_GetRef();
 			Mesh.bWireframe = bWireframe;
 			Mesh.SegmentIndex = SectionIndex;
-			Mesh.VertexFactory = &GeometryCollectionVertexFactory;
+			Mesh.VertexFactory = &VertexFactory;
 			Mesh.MaterialRenderProxy = MaterialProxies[SectionIndex];
 			Mesh.LODIndex = LODIndex;
 			Mesh.bDisableBackfaceCulling = true;
@@ -711,7 +684,7 @@ void FGeometryCollectionSceneProxyBase::GetDynamicRayTracingInstances(FRayTracin
 			FMeshBatch& Mesh = RayTracingInstance.Materials.AddDefaulted_GetRef();
 			Mesh.bWireframe = bWireframe;
 			Mesh.SegmentIndex = 0;
-			Mesh.VertexFactory = &GeometryCollectionVertexFactory;
+			Mesh.VertexFactory = &VertexFactory;
 			Mesh.MaterialRenderProxy = MaterialProxies[SectionIndex];
 			Mesh.LODIndex = LODIndex;
 			Mesh.bDisableBackfaceCulling = true;
