[HWRT] Per-part ray tracing geometries in FGeometryCollectionRenderData.
- avoids duplicate ray tracing geometries per proxy.
- expose flag bSupportRayTracing in UGeometryCollection object.

diff --git a/Engine/Source/Runtime/Experimental/GeometryCollectionEngine/Private/GeometryCollection/GeometryCollectionObject.cpp b/Engine/Source/Runtime/Experimental/GeometryCollectionEngine/Private/GeometryCollection/GeometryCollectionObject.cpp
--- a/Engine/Source/Runtime/Experimental/GeometryCollectionEngine/Private/GeometryCollection/GeometryCollectionObject.cpp
+++ b/Engine/Source/Runtime/Experimental/GeometryCollectionEngine/Private/GeometryCollection/GeometryCollectionObject.cpp
@@ -113,6 +113,7 @@ UGeometryCollection::UGeometryCollection(const FObjectInitializer& ObjectInitial
 	, bUseFullPrecisionUVs(false)
 	, bStripOnCook(false)
 	, bStripRenderDataOnCook(false)
+	, bSupportRayTracing(true)
 	, EnableNanite(false)
 	, bEnableNaniteFallback(false)
 	, NaniteMinimumResidencyInKB(0)
@@ -1945,6 +1946,11 @@ void UGeometryCollection::PostEditChangeProperty(struct FPropertyChangedEvent& P
 			bDoInvalidateCollection = true;
 			bRebuildRenderData = true;
 		}
+		else if (PropertyChangedEvent.Property->GetFName() == GET_MEMBER_NAME_CHECKED(UGeometryCollection, bSupportRayTracing))
+		{
+			bDoInvalidateCollection = true;
+			bRebuildRenderData = true;
+		}
 		else if (PropertyChangedEvent.Property->GetFName() == GET_MEMBER_NAME_CHECKED(UGeometryCollection, bUseFullPrecisionUVs))
 		{
 			bDoInvalidateCollection = true;
diff --git a/Engine/Source/Runtime/Experimental/GeometryCollectionEngine/Private/GeometryCollection/GeometryCollectionRenderData.cpp b/Engine/Source/Runtime/Experimental/GeometryCollectionEngine/Private/GeometryCollection/GeometryCollectionRenderData.cpp
--- a/Engine/Source/Runtime/Experimental/GeometryCollectionEngine/Private/GeometryCollection/GeometryCollectionRenderData.cpp
+++ b/Engine/Source/Runtime/Experimental/GeometryCollectionEngine/Private/GeometryCollection/GeometryCollectionRenderData.cpp
@@ -790,6 +790,50 @@ void FGeometryCollectionRenderData::InitResources(UGeometryCollection const& Own
 	if (bHasMeshData)
 	{
 		MeshResource.InitResources(Owner);
+
+#if RHI_RAYTRACING
+		if (IsRayTracingAllowed() && Owner.bSupportRayTracing && MeshDescription.NumVertices != 0)
+		{
+			PartRayTracingGeometries.SetNum(MeshDescription.SubSections.Num());
+
+			FName OwnerName = Owner.GetFName();
+
+			ENQUEUE_RENDER_COMMAND(InitGeometryCollectiomRayTracingGeometry)(
+				[this, OwnerName](FRHICommandListImmediate& RHICmdList)
+				{
+					// TODO: Combine sections using the same transform into a single BLAS
+
+					for (int32 SectionIndex = 0; SectionIndex < MeshDescription.SubSections.Num(); ++SectionIndex)
+					{
+						FRayTracingGeometryInitializer Initializer;
+						Initializer.DebugName = FDebugName(OwnerName, SectionIndex);
+						Initializer.GeometryType = RTGT_Triangles;
+						Initializer.bFastBuild = true;
+						Initializer.bAllowUpdate = false;
+						Initializer.TotalPrimitiveCount = 0;
+						Initializer.IndexBuffer = MeshResource.IndexBuffer.IndexBufferRHI;
+
+						{
+							const FGeometryCollectionMeshElement& Section = MeshDescription.SubSections[SectionIndex];
+
+							FRayTracingGeometrySegment Segment;
+							Segment.FirstPrimitive = Section.TriangleStart;
+							Segment.VertexBuffer = MeshResource.PositionVertexBuffer.VertexBufferRHI;
+							Segment.NumPrimitives = Section.TriangleCount;
+							Segment.MaxVertices = Section.VertexEnd;
+
+							Initializer.Segments.Add(Segment);
+
+							Initializer.TotalPrimitiveCount += Section.TriangleCount;
+						}
+
+						PartRayTracingGeometries[SectionIndex].SetInitializer(MoveTemp(Initializer));
+						PartRayTracingGeometries[SectionIndex].InitResource(RHICmdList);
+					}
+				}
+				);
+		}
+#endif // RHI_RAYTRACING
 	}
 
 	if (bHasNaniteData)
@@ -809,6 +853,12 @@ void FGeometryCollectionRenderData::ReleaseResources()
 
 	if (bHasMeshData)
 	{
+		for (FRayTracingGeometry& PartRayTracingGeometry : PartRayTracingGeometries)
+		{
+			PartRayTracingGeometry.ReleaseResource();
+		}
+		PartRayTracingGeometries.Empty();
+
 		MeshResource.ReleaseResources();
 	}
 
diff --git a/Engine/Source/Runtime/Experimental/GeometryCollectionEngine/Private/GeometryCollection/GeometryCollectionRenderData.h b/Engine/Source/Runtime/Experimental/GeometryCollectionEngine/Private/GeometryCollection/GeometryCollectionRenderData.h
--- a/Engine/Source/Runtime/Experimental/GeometryCollectionEngine/Private/GeometryCollection/GeometryCollectionRenderData.h
+++ b/Engine/Source/Runtime/Experimental/GeometryCollectionEngine/Private/GeometryCollection/GeometryCollectionRenderData.h
@@ -8,6 +8,7 @@
 #include "Rendering/NaniteResources.h"
 #include "Rendering/PositionVertexBuffer.h"
 #include "Rendering/StaticMeshVertexBuffer.h"
+#include "RayTracingGeometry.h"
 
 class FGeometryCollection;
 class UGeometryCollection;
@@ -156,6 +157,8 @@ public:
 	/** Object local bounds before any animation. */
 	FBoxSphereBounds PreSkinnedBounds = FBoxSphereBounds(ForceInitToZero);
 
+	TArray<FRayTracingGeometry> PartRayTracingGeometries;
+
 private:
 	bool bIsInitialized = false;
 };
diff --git a/Engine/Source/Runtime/Experimental/GeometryCollectionEngine/Private/GeometryCollection/GeometryCollectionSceneProxy.cpp b/Engine/Source/Runtime/Experimental/GeometryCollectionEngine/Private/GeometryCollection/GeometryCollectionSceneProxy.cpp
--- a/Engine/Source/Runtime/Experimental/GeometryCollectionEngine/Private/GeometryCollection/GeometryCollectionSceneProxy.cpp
+++ b/Engine/Source/Runtime/Experimental/GeometryCollectionEngine/Private/GeometryCollection/GeometryCollectionSceneProxy.cpp
@@ -145,7 +145,8 @@ FGeometryCollectionSceneProxyBase::FGeometryCollectionSceneProxyBase(UGeometryCo
 	, bUseShaderBoneTransform(VertexFactory.UseShaderBoneTransform(Component->GetScene()->GetShaderPlatform()))
 	, bSupportsTripleBufferVertexUpload(GRHISupportsMapWriteNoOverwrite)
 #if RHI_RAYTRACING
-	, bHasRayTracingRepresentation(IsRayTracingEnabled() && Component->GetRestCollection()->RenderData->bHasMeshData && Component->GetRestCollection()->RenderData->MeshDescription.NumVertices != 0)
+	, bHasRayTracingRepresentation(IsRayTracingEnabled() && Component->GetRestCollection()->RenderData->PartRayTracingGeometries.Num() > 0)
+
 #endif
 {
 	if (!bIsNanite || bHasRayTracingRepresentation)
@@ -169,6 +170,13 @@ FGeometryCollectionSceneProxyBase::FGeometryCollectionSceneProxyBase(UGeometryCo
 		}
 	}
 
+#if RHI_RAYTRACING
+	if (bHasRayTracingRepresentation)
+	{
+		PartRayTracingGeometries = Component->GetRestCollection()->RenderData->PartRayTracingGeometries;
+	}
+#endif
+
 	PreSkinnedBounds = Component->GetRestCollection()->RenderData->PreSkinnedBounds;
 
 	DynamicData = Component->InitDynamicData(true);
@@ -240,12 +248,6 @@ void FGeometryCollectionSceneProxyBase::DestroyRenderThreadResources()
 #if RHI_RAYTRACING
 	if (bHasRayTracingRepresentation)
 	{
-		for (FRayTracingGeometry& PartRayTracingGeometry : PartRayTracingGeometries)
-		{
-			PartRayTracingGeometry.ReleaseResource();
-		}
-		PartRayTracingGeometries.Empty();
-
 		RayTracingGeometry.ReleaseResource();
 		RayTracingDynamicVertexBuffer.Release();
 	}
@@ -443,51 +445,9 @@ void FGeometryCollectionSceneProxyBase::UpdatingRayTracingGeometry_RenderingThre
 	{
 		// release combined geometry since will use part geometries
 		RayTracingGeometry.ReleaseResource();
-
-		// TODO: Combine sections using the same transform into a single BLAS
-
-		if (PartRayTracingGeometries.IsEmpty())
-		{
-			PartRayTracingGeometries.SetNum(InSectionArray.Num());
-
-			for (int32 SectionIndex = 0; SectionIndex < InSectionArray.Num(); ++SectionIndex)
-			{
-				FRayTracingGeometryInitializer Initializer;
-				Initializer.DebugName = FDebugName(FName("GeometryCollectionPart"), SectionIndex);
-				Initializer.GeometryType = RTGT_Triangles;
-				Initializer.bFastBuild = true;
-				Initializer.bAllowUpdate = false;
-				Initializer.TotalPrimitiveCount = 0;
-				Initializer.IndexBuffer = MeshResource.IndexBuffer.IndexBufferRHI;
-
-				{
-					const FGeometryCollectionMeshElement& Section = InSectionArray[SectionIndex];
-
-					FRayTracingGeometrySegment Segment;
-					Segment.FirstPrimitive = Section.TriangleStart;
-					Segment.VertexBuffer = MeshResource.PositionVertexBuffer.VertexBufferRHI;
-					Segment.NumPrimitives = Section.TriangleCount;
-					Segment.MaxVertices = Section.VertexEnd;
-
-					Initializer.Segments.Add(Segment);
-
-					Initializer.TotalPrimitiveCount += Section.TriangleCount;
-				}
-
-				PartRayTracingGeometries[SectionIndex].SetInitializer(MoveTemp(Initializer));
-				PartRayTracingGeometries[SectionIndex].InitResource(RHICmdList);
-			}
-		}
 	}
 	else
 	{
-		// release part geometries since will use combined geometry
-		for (FRayTracingGeometry& PartRayTracingGeometry : PartRayTracingGeometries)
-		{
-			PartRayTracingGeometry.ReleaseResource();
-		}
-		PartRayTracingGeometries.Empty();
-
 		// initialize combined geometry if necessary
 		if (!RayTracingGeometry.IsInitialized())
 		{
@@ -587,6 +547,8 @@ void FGeometryCollectionSceneProxyBase::GetDynamicRayTracingInstances(FRayTracin
 
 	if (RayTracingGeometry.IsValid())
 	{
+		check(!bUseSubSections);
+
 		// Render dynamic objects
 		if (!VertexFactory.GetType()->SupportsRayTracingDynamicGeometry())
 		{
@@ -671,47 +633,51 @@ void FGeometryCollectionSceneProxyBase::GetDynamicRayTracingInstances(FRayTracin
 			Collector.AddRayTracingInstance(ViewIndex, RayTracingInstance);
 		}
 	}
-
-	for (int32 SectionIndex = 0; SectionIndex < PartRayTracingGeometries.Num(); ++SectionIndex)
+	else
 	{
-		const FGeometryCollectionMeshElement& Section = SectionArray[SectionIndex];
-
-		FRayTracingInstance RayTracingInstance;
-		RayTracingInstance.Geometry = &PartRayTracingGeometries[SectionIndex];
-		RayTracingInstance.InstanceTransforms.Emplace(FMatrix(DynamicData->Transforms[Section.TransformIndex]) * LocalToWorld);
+		check(bUseSubSections);
 
+		for (int32 SectionIndex = 0; SectionIndex < PartRayTracingGeometries.Num(); ++SectionIndex)
 		{
-			FMeshBatch& Mesh = RayTracingInstance.Materials.AddDefaulted_GetRef();
-			Mesh.bWireframe = bWireframe;
-			Mesh.SegmentIndex = 0;
-			Mesh.VertexFactory = &VertexFactory;
-			Mesh.MaterialRenderProxy = MaterialProxies[SectionIndex];
-			Mesh.LODIndex = LODIndex;
-			Mesh.bDisableBackfaceCulling = true;
-			Mesh.Type = PT_TriangleList;
-			Mesh.DepthPriorityGroup = SDPG_World;
-			Mesh.bCanApplyViewModeOverrides = true;
-
-			FMeshBatchElement& BatchElement = Mesh.Elements[0];
-			BatchElement.IndexBuffer = &MeshResource.IndexBuffer;
-			BatchElement.PrimitiveUniformBuffer = UniformBuffer;
-			BatchElement.FirstIndex = Section.TriangleStart * 3;
-			BatchElement.NumPrimitives = Section.TriangleCount;
-			BatchElement.MinVertexIndex = Section.VertexStart;
-			BatchElement.MaxVertexIndex = Section.VertexEnd;
-			BatchElement.NumInstances = 1;
+			const FGeometryCollectionMeshElement& Section = SectionArray[SectionIndex];
 
-			//#TODO: bone color, bone selection and render bound?
-		}
+			FRayTracingInstance RayTracingInstance;
+			RayTracingInstance.Geometry = &PartRayTracingGeometries[SectionIndex];
+			RayTracingInstance.InstanceTransforms.Emplace(FMatrix(DynamicData->Transforms[Section.TransformIndex]) * LocalToWorld);
 
-		for (int32 ViewIndex = 0; ViewIndex < Views.Num(); ViewIndex++)
-		{
-			if ((VisibilityMap & (1 << ViewIndex)) == 0)
 			{
-				continue;
+				FMeshBatch& Mesh = RayTracingInstance.Materials.AddDefaulted_GetRef();
+				Mesh.bWireframe = bWireframe;
+				Mesh.SegmentIndex = 0;
+				Mesh.VertexFactory = &VertexFactory;
+				Mesh.MaterialRenderProxy = MaterialProxies[SectionIndex];
+				Mesh.LODIndex = LODIndex;
+				Mesh.bDisableBackfaceCulling = true;
+				Mesh.Type = PT_TriangleList;
+				Mesh.DepthPriorityGroup = SDPG_World;
+				Mesh.bCanApplyViewModeOverrides = true;
+
+				FMeshBatchElement& BatchElement = Mesh.Elements[0];
+				BatchElement.IndexBuffer = &MeshResource.IndexBuffer;
+				BatchElement.PrimitiveUniformBuffer = UniformBuffer;
+				BatchElement.FirstIndex = Section.TriangleStart * 3;
+				BatchElement.NumPrimitives = Section.TriangleCount;
+				BatchElement.MinVertexIndex = Section.VertexStart;
+				BatchElement.MaxVertexIndex = Section.VertexEnd;
+				BatchElement.NumInstances = 1;
+
+				//#TODO: bone color, bone selection and render bound?
 			}
 
-			Collector.AddRayTracingInstance(ViewIndex, RayTracingInstance);
+			for (int32 ViewIndex = 0; ViewIndex < Views.Num(); ViewIndex++)
+			{
+				if ((VisibilityMap & (1 << ViewIndex)) == 0)
+				{
+					continue;
+				}
+
+				Collector.AddRayTracingInstance(ViewIndex, RayTracingInstance);
+			}
 		}
 	}
 }
diff --git a/Engine/Source/Runtime/Experimental/GeometryCollectionEngine/Private/GeometryCollection/GeometryCollectionSceneProxy.h b/Engine/Source/Runtime/Experimental/GeometryCollectionEngine/Private/GeometryCollection/GeometryCollectionSceneProxy.h
--- a/Engine/Source/Runtime/Experimental/GeometryCollectionEngine/Private/GeometryCollection/GeometryCollectionSceneProxy.h
+++ b/Engine/Source/Runtime/Experimental/GeometryCollectionEngine/Private/GeometryCollection/GeometryCollectionSceneProxy.h
@@ -220,7 +220,7 @@ private:
 	FRayTracingGeometry RayTracingGeometry;
 	FRWBuffer RayTracingDynamicVertexBuffer;
 
-	TArray<FRayTracingGeometry> PartRayTracingGeometries;
+	TConstArrayView<FRayTracingGeometry> PartRayTracingGeometries;
 #endif
 
 	bool bRenderResourcesCreated = false;
diff --git a/Engine/Source/Runtime/Experimental/GeometryCollectionEngine/Public/GeometryCollection/GeometryCollectionObject.h b/Engine/Source/Runtime/Experimental/GeometryCollectionEngine/Public/GeometryCollection/GeometryCollectionObject.h
--- a/Engine/Source/Runtime/Experimental/GeometryCollectionEngine/Public/GeometryCollection/GeometryCollectionObject.h
+++ b/Engine/Source/Runtime/Experimental/GeometryCollectionEngine/Public/GeometryCollection/GeometryCollectionObject.h
@@ -626,6 +626,12 @@ public:
 	UPROPERTY(EditAnywhere, BlueprintReadWrite, Category = "Rendering")
 	bool bStripRenderDataOnCook;
 
+	/**
+	 * If true, a ray tracing acceleration structure will be built for this geometry collection and it may be used in ray tracing effects
+	 */
+	UPROPERTY(EditAnywhere, Category = "Rendering")
+	uint8 bSupportRayTracing : 1;
+
 	/** Custom class type that will be used to render the geometry collection instead of using the native rendering. */
 	UPROPERTY(EditAnywhere, BlueprintReadWrite, Category = "Rendering", meta = (MustImplement = "/Script/GeometryCollectionEngine.GeometryCollectionExternalRenderInterface"))
 	TObjectPtr<UClass> CustomRendererType;
