[HWRT] Optimize geometry collection ray tracing instance gathering.
- updated FRayTracingInstanceCollector to store RayTracingInstances and RayTracingGeometriesToUpdate using TChunkedArray.
- cache ray tracing mesh batches in proxy.
- allocate PartInstanceTransforms from Collector and use InstanceTransformsView to avoid allocations/copies.

Test level with 36 GCs with 100 parts each it reduces the following costs:
- RayTracing_GatherDynamicRayTracingInstances - 2.53ms -> 0.35ms
- RayTracingScene_AddDynamicInstances - 1.5ms -> 0.95ms

diff --git a/Engine/Source/Runtime/Engine/Private/SceneManagement.cpp b/Engine/Source/Runtime/Engine/Private/SceneManagement.cpp
--- a/Engine/Source/Runtime/Engine/Private/SceneManagement.cpp
+++ b/Engine/Source/Runtime/Engine/Private/SceneManagement.cpp
@@ -375,7 +375,7 @@ void FRayTracingInstanceCollector::AddRayTracingInstance(int32 ViewIndex, FRayTr
 	Tmp.Instance = MoveTemp(Instance);
 	Tmp.ViewIndex = ViewIndex;
 
-	RayTracingInstances.Add(MoveTemp(Tmp));
+	RayTracingInstances.AddElement(MoveTemp(Tmp));
 }
 
 void FRayTracingInstanceCollector::AddReferencedGeometryGroup(RayTracing::FGeometryGroupHandle GeometryGroup)
@@ -405,7 +405,7 @@ void FRayTracingInstanceCollector::AddRayTracingGeometryUpdate(int32 ViewIndex,
 	UpdateRequest.Params = MoveTemp(Params);
 	UpdateRequest.ViewIndex = ViewIndex;
 
-	RayTracingGeometriesToUpdate.Add(MoveTemp(UpdateRequest));
+	RayTracingGeometriesToUpdate.AddElement(MoveTemp(UpdateRequest));
 }
 
 #endif
diff --git a/Engine/Source/Runtime/Engine/Public/SceneManagement.h b/Engine/Source/Runtime/Engine/Public/SceneManagement.h
--- a/Engine/Source/Runtime/Engine/Public/SceneManagement.h
+++ b/Engine/Source/Runtime/Engine/Public/SceneManagement.h
@@ -1582,9 +1582,9 @@ private:
 	};
 
 	TArray<const FSceneView*> Views;
-	TArray<FRayTracingInstanceAndViewIndex, SceneRenderingAllocator> RayTracingInstances;
 	TArray<TSet<FRDGPooledBuffer*>, TInlineAllocator<2, SceneRenderingAllocator> > RDGPooledBuffers;
-	TArray<FRayTracingDynamicGeometryUpdateRequest, SceneRenderingAllocator> RayTracingGeometriesToUpdate;
+	TChunkedArray<FRayTracingInstanceAndViewIndex, 16384, SceneRenderingAllocator> RayTracingInstances;
+	TChunkedArray<FRayTracingDynamicGeometryUpdateRequest, 16384, SceneRenderingAllocator> RayTracingGeometriesToUpdate;
 	TSet<RayTracing::FGeometryGroupHandle> ReferencedGeometryGroups;
 	TSet<RayTracing::FGeometryGroupHandle> ReferencedGeometryGroupsForDynamicUpdate;
 
diff --git a/Engine/Source/Runtime/Experimental/GeometryCollectionEngine/Private/GeometryCollection/GeometryCollectionSceneProxy.cpp b/Engine/Source/Runtime/Experimental/GeometryCollectionEngine/Private/GeometryCollection/GeometryCollectionSceneProxy.cpp
--- a/Engine/Source/Runtime/Experimental/GeometryCollectionEngine/Private/GeometryCollection/GeometryCollectionSceneProxy.cpp
+++ b/Engine/Source/Runtime/Experimental/GeometryCollectionEngine/Private/GeometryCollection/GeometryCollectionSceneProxy.cpp
@@ -637,16 +637,24 @@ void FGeometryCollectionSceneProxyBase::GetDynamicRayTracingInstances(FRayTracin
 	{
 		check(bUseSubSections);
 
-		for (int32 SectionIndex = 0; SectionIndex < PartRayTracingGeometries.Num(); ++SectionIndex)
+		TArray<FMatrix, SceneRenderingAllocator> PartInstanceTransforms = Collector.AllocateOneFrameResource<TArray<FMatrix, SceneRenderingAllocator>>();
+		PartInstanceTransforms.AddUninitialized(DynamicData->Transforms.Num());
+
+		for (int32 TransformIndex = 0; TransformIndex < DynamicData->Transforms.Num(); ++TransformIndex)
 		{
-			const FGeometryCollectionMeshElement& Section = SectionArray[SectionIndex];
+			PartInstanceTransforms[TransformIndex] = FMatrix(DynamicData->Transforms[TransformIndex]) * LocalToWorld;
+		}
 
-			FRayTracingInstance RayTracingInstance;
-			RayTracingInstance.Geometry = &PartRayTracingGeometries[SectionIndex];
-			RayTracingInstance.InstanceTransforms.Emplace(FMatrix(DynamicData->Transforms[Section.TransformIndex]) * LocalToWorld);
+		if (SectionArray.Num() != CachedPartRayTracingMaterials.Num())
+		{
+			CachedPartRayTracingMaterials.Reset();
+			CachedPartRayTracingMaterials.Reserve(SectionArray.Num());
 
+			for (int32 SectionIndex = 0; SectionIndex < SectionArray.Num(); ++SectionIndex)
 			{
-				FMeshBatch& Mesh = RayTracingInstance.Materials.AddDefaulted_GetRef();
+				const FGeometryCollectionMeshElement& Section = SectionArray[SectionIndex];
+
+				FMeshBatch& Mesh = CachedPartRayTracingMaterials.AddDefaulted_GetRef();
 				Mesh.bWireframe = bWireframe;
 				Mesh.SegmentIndex = 0;
 				Mesh.VertexFactory = &VertexFactory;
@@ -666,8 +674,20 @@ void FGeometryCollectionSceneProxyBase::GetDynamicRayTracingInstances(FRayTracin
 				BatchElement.MaxVertexIndex = Section.VertexEnd;
 				BatchElement.NumInstances = 1;
 
-				//#TODO: bone color, bone selection and render bound?
+				// #TODO: bone color, bone selection and render bound?
 			}
+		}
+
+		checkf(PartRayTracingGeometries.Num() == SectionArray.Num(), TEXT("Expected to have one part per section (%d vs %d)."), PartRayTracingGeometries.Num(), SectionArray.Num());
+
+		for (int32 SectionIndex = 0; SectionIndex < PartRayTracingGeometries.Num(); ++SectionIndex)
+		{
+			const FGeometryCollectionMeshElement& Section = SectionArray[SectionIndex];
+
+			FRayTracingInstance RayTracingInstance;
+			RayTracingInstance.Geometry = &PartRayTracingGeometries[SectionIndex];
+			RayTracingInstance.InstanceTransformsView = MakeConstArrayView(&PartInstanceTransforms[Section.TransformIndex], 1);
+			RayTracingInstance.MaterialsView = MakeConstArrayView(&CachedPartRayTracingMaterials[SectionIndex], 1);
 
 			for (int32 ViewIndex = 0; ViewIndex < Views.Num(); ViewIndex++)
 			{
diff --git a/Engine/Source/Runtime/Experimental/GeometryCollectionEngine/Private/GeometryCollection/GeometryCollectionSceneProxy.h b/Engine/Source/Runtime/Experimental/GeometryCollectionEngine/Private/GeometryCollection/GeometryCollectionSceneProxy.h
--- a/Engine/Source/Runtime/Experimental/GeometryCollectionEngine/Private/GeometryCollection/GeometryCollectionSceneProxy.h
+++ b/Engine/Source/Runtime/Experimental/GeometryCollectionEngine/Private/GeometryCollection/GeometryCollectionSceneProxy.h
@@ -221,6 +221,7 @@ private:
 	FRWBuffer RayTracingDynamicVertexBuffer;
 
 	TConstArrayView<FRayTracingGeometry> PartRayTracingGeometries;
+	TArray<FMeshBatch> CachedPartRayTracingMaterials;
 #endif
 
 	bool bRenderResourcesCreated = false;
diff --git a/Engine/Source/Runtime/Renderer/Private/RayTracing/RayTracing.cpp b/Engine/Source/Runtime/Renderer/Private/RayTracing/RayTracing.cpp
--- a/Engine/Source/Runtime/Renderer/Private/RayTracing/RayTracing.cpp
+++ b/Engine/Source/Runtime/Renderer/Private/RayTracing/RayTracing.cpp
@@ -850,14 +850,11 @@ namespace RayTracing
 			const FScene::FPrimitiveRayTracingData& PrimitiveRayTracingData = Scene.PrimitiveRayTracingDatas[PrimitiveIndex];
 			const FPersistentPrimitiveIndex PersistentPrimitiveIndex = SceneInfo->GetPersistentIndex();
 
-			TArrayView<FRayTracingInstanceCollector::FRayTracingInstanceAndViewIndex> TempRayTracingInstances = MakeArrayView(
-				RayTracingInstanceCollector.RayTracingInstances.GetData() + DynamicRayTracingPrimitive.InstancesRange.GetLowerBoundValue(),
-				DynamicRayTracingPrimitive.InstancesRange.Size<int32>());
-
-			if (TempRayTracingInstances.Num() > 0)
+			if (DynamicRayTracingPrimitive.InstancesRange.Size<int32>() > 0)
 			{
-				for (FRayTracingInstanceCollector::FRayTracingInstanceAndViewIndex& InstanceAndViewIndex : TempRayTracingInstances)
+				for (int32 DynamicInstanceIndex = DynamicRayTracingPrimitive.InstancesRange.GetLowerBoundValue(); DynamicInstanceIndex < DynamicRayTracingPrimitive.InstancesRange.GetUpperBoundValue(); ++DynamicInstanceIndex)
 				{
+					FRayTracingInstanceCollector::FRayTracingInstanceAndViewIndex& InstanceAndViewIndex = RayTracingInstanceCollector.RayTracingInstances[DynamicInstanceIndex];
 					FRayTracingInstance& Instance = InstanceAndViewIndex.Instance;
 					FViewInfo& View = *Views[InstanceAndViewIndex.ViewIndex];
 					 
