Fix PartInstanceTransforms allocated from Collector not being assigned to a reference.

diff --git a/Engine/Source/Runtime/Experimental/GeometryCollectionEngine/Private/GeometryCollection/GeometryCollectionSceneProxy.cpp b/Engine/Source/Runtime/Experimental/GeometryCollectionEngine/Private/GeometryCollection/GeometryCollectionSceneProxy.cpp
--- a/Engine/Source/Runtime/Experimental/GeometryCollectionEngine/Private/GeometryCollection/GeometryCollectionSceneProxy.cpp
+++ b/Engine/Source/Runtime/Experimental/GeometryCollectionEngine/Private/GeometryCollection/GeometryCollectionSceneProxy.cpp
@@ -637,7 +637,7 @@ void FGeometryCollectionSceneProxyBase::GetDynamicRayTracingInstances(FRayTracin
 	{
 		check(bUseSubSections);
 
-		TArray<FMatrix, SceneRenderingAllocator> PartInstanceTransforms = Collector.AllocateOneFrameResource<TArray<FMatrix, SceneRenderingAllocator>>();
+		TArray<FMatrix, SceneRenderingAllocator>& PartInstanceTransforms = Collector.AllocateOneFrameResource<TArray<FMatrix, SceneRenderingAllocator>>();
 		PartInstanceTransforms.AddUninitialized(DynamicData->Transforms.Num());
 
 		for (int32 TransformIndex = 0; TransformIndex < DynamicData->Transforms.Num(); ++TransformIndex)
