The disk cache will return EIoErrorCode::InvalidParameter if the requested offset is out of bounds, matching the other IoDispatcher backends.

#rnx

diff --git a/Engine/Source/Runtime/Experimental/IoStore/OnDemand/Private/JournaledCache.cpp b/Engine/Source/Runtime/Experimental/IoStore/OnDemand/Private/JournaledCache.cpp
--- a/Engine/Source/Runtime/Experimental/IoStore/OnDemand/Private/JournaledCache.cpp
+++ b/Engine/Source/Runtime/Experimental/IoStore/OnDemand/Private/JournaledCache.cpp
@@ -157,7 +157,7 @@ FMemCache::FMemCache(uint32 InMaxSize)
 ////////////////////////////////////////////////////////////////////////////////
 uint32 FMemCache::GetDemand() const
 {
-	return (GetUsed() * 100) / MaxSize;
+	return MaxSize > 0 ? (GetUsed() * 100) / MaxSize : 0;
 }
 
 ////////////////////////////////////////////////////////////////////////////////
@@ -820,6 +820,11 @@ EIoErrorCode FDiskCache::Materialize(uint64 Key, FIoBuffer& Out, uint32 Offset)
 			return EIoErrorCode::NotFound;
 		}
 
+		if (Offset > uint32(Entry->Size))
+		{
+			return EIoErrorCode::InvalidParameter;
+		}
+
 		ReadSize = uint32(Entry->Size) - Offset;
 		EntryDataCursor = Entry->DataCursor;
 	}
