Enable version two everywhere.

#rnx

diff --git a/Engine/Source/Runtime/Experimental/IoStore/OnDemand/Private/OnDemandHttpClient.cpp b/Engine/Source/Runtime/Experimental/IoStore/OnDemand/Private/OnDemandHttpClient.cpp
--- a/Engine/Source/Runtime/Experimental/IoStore/OnDemand/Private/OnDemandHttpClient.cpp
+++ b/Engine/Source/Runtime/Experimental/IoStore/OnDemand/Private/OnDemandHttpClient.cpp
@@ -693,6 +693,10 @@ TUniquePtr<HTTP::FConnectionPool> FMultiEndpointHttpClient::CreateConnection(FAn
 		Params.ConnectionCount = 1; // rfc9113 conformance
 		Params.VerifyCert = HTTP::FCertRoots::Default();
 	}
+	else
+	{
+		Params.MaxInflight = 1;
+	}
 
 	return MakeUnique<HTTP::FConnectionPool>(Params);
 }
diff --git a/Engine/Source/Runtime/Experimental/IoStore/OnDemand/Private/OnDemandHttpThread.cpp b/Engine/Source/Runtime/Experimental/IoStore/OnDemand/Private/OnDemandHttpThread.cpp
--- a/Engine/Source/Runtime/Experimental/IoStore/OnDemand/Private/OnDemandHttpThread.cpp
+++ b/Engine/Source/Runtime/Experimental/IoStore/OnDemand/Private/OnDemandHttpThread.cpp
@@ -94,7 +94,7 @@ FAutoConsoleVariableRef CVarOnDemandBackendThreadPriority(
 	TEXT("Note that this is switchable at runtime"),
 	ECVF_Default);
 
-int32 GIaxHttpVersion = 1;
+int32 GIaxHttpVersion = 2;
 static FAutoConsoleVariableRef CVar_IaxHttpVersion(
 	TEXT("iax.HttpVersion"),
 	GIaxHttpVersion,
@@ -102,6 +102,8 @@ static FAutoConsoleVariableRef CVar_IaxHttpVersion(
 	FConsoleVariableDelegate::CreateStatic(&OnHttpClientCVarChanged)
 );
 
+/* Note that the value of this cvar does not affect throttled scenarios. Setting it
+ * too low will start to surface the negative effects of latency and CDN misses. */
 int32 GIasHttpConnectionCount = 4;
 static FAutoConsoleVariableRef CVar_IasHttpConnectionCount(
 	TEXT("ias.HttpConnectionCount"),
@@ -110,7 +112,10 @@ static FAutoConsoleVariableRef CVar_IasHttpConnectionCount(
 	FConsoleVariableDelegate::CreateStatic(&OnHttpClientCVarChanged)
 );
 
-int32 GIaxHttpMaxInflight = 1;
+/* Note that the value of this cvar does not affect throttling. Setting it too low
+ * will prevent the CDN peer from reordering cache hits ahead of misses. Latency
+ * will also start to become apparent and negatively effect user experience. */
+int32 GIaxHttpMaxInflight = 5;
 static FAutoConsoleVariableRef CVar_IaxHttpMaxInflight(
 	TEXT("iax.HttpMaxInflight"),
 	GIaxHttpMaxInflight,
