Derive queue length that best suits version in use.

#rnx

diff --git a/Engine/Source/Runtime/Experimental/IoStore/OnDemand/Private/OnDemandHttpIoDispacher.cpp b/Engine/Source/Runtime/Experimental/IoStore/OnDemand/Private/OnDemandHttpIoDispacher.cpp
--- a/Engine/Source/Runtime/Experimental/IoStore/OnDemand/Private/OnDemandHttpIoDispacher.cpp
+++ b/Engine/Source/Runtime/Experimental/IoStore/OnDemand/Private/OnDemandHttpIoDispacher.cpp
@@ -68,6 +68,7 @@ extern float GIaxInvalidUrlChance;
 #if UE_ALLOW_CACHE_POISONING
 extern bool GIasPoisonCache;
 #endif // UE_ALLOW_CACHE_POISONING
+int32 CalculateRequestQueueLengthMax();
 
 /** Matches the priority set up found in Engine/Source/Runtime/Online/HTTP/Private/HttpThread.cpp */
 static constexpr int32 GThreadPriorities[5] =
@@ -1362,7 +1363,8 @@ bool FHttpDispatcher::ProcessHttpRequests()
 
 	RecreateHttpClientIfNeeded();
 
-	const int32 MaxConcurrentRequests = FMath::Min(GIasHttpConcurrentRequests, 32);
+	int32		QueueLength = CalculateRequestQueueLengthMax();
+	const int32 MaxConcurrentRequests = FMath::Min(QueueLength, 32);
 	int32		NumConcurrentRequests = 0;
 
 	FHttpRequest* NextHttpRequest = DequeueHttpRequests();
diff --git a/Engine/Source/Runtime/Experimental/IoStore/OnDemand/Private/OnDemandHttpThread.cpp b/Engine/Source/Runtime/Experimental/IoStore/OnDemand/Private/OnDemandHttpThread.cpp
--- a/Engine/Source/Runtime/Experimental/IoStore/OnDemand/Private/OnDemandHttpThread.cpp
+++ b/Engine/Source/Runtime/Experimental/IoStore/OnDemand/Private/OnDemandHttpThread.cpp
@@ -200,6 +200,16 @@ static FAutoConsoleVariableRef CVar_GIaxInvalidUrlChance(
 );
 #endif // UE_ALLOW_INVALID_URL_DEBUGGING
 
+int32 CalculateRequestQueueLengthMax()
+{
+	// Keeping the loop loaded with work to do is important to cover round-trip
+	// latency and the effect of CDN cache misses. Reducing what is in the queue
+	// DOES NOT have any effect on throughput throttling.
+	int32 ConnNum = (GIaxHttpVersion == 2) ? 1 : GIasHttpConnectionCount;
+	int32 InflightNum = (GIaxHttpVersion == 2) ? GIaxHttpMaxInflight : 1;
+	return (ConnNum * InflightNum) + 2;
+}
+
 const int32 GOnDemandBackendThreadPriorities[5] =
 {
 	EThreadPriority::TPri_Lowest,
@@ -474,7 +484,8 @@ void FOnDemandHttpThread::TickRequests()
 	}
 
 	// We can't limit cvars so add a hard coded cap to prevent concurrent requests from getting out of hand.
-	const int32 MaxConcurrentRequests = FMath::Min(GIasHttpConcurrentRequests, 32);
+	int32 QueueLength = CalculateRequestQueueLengthMax();
+	const int32 MaxConcurrentRequests = FMath::Min(QueueLength, 32);
 
 	int32 NumConcurrentRequests = 0;
 	FHttpRequest* NextHttpRequest = HttpRequests.Dequeue(MaxConcurrentRequests);
diff --git a/Engine/Source/Runtime/Experimental/IoStore/OnDemand/Private/Statistics.cpp b/Engine/Source/Runtime/Experimental/IoStore/OnDemand/Private/Statistics.cpp
--- a/Engine/Source/Runtime/Experimental/IoStore/OnDemand/Private/Statistics.cpp
+++ b/Engine/Source/Runtime/Experimental/IoStore/OnDemand/Private/Statistics.cpp
@@ -30,6 +30,8 @@ namespace UE::IoStore
 extern int32 GIaxHttpVersion;
 extern int32 GHttpIoDispatcherMaxRangeKiB;
 
+int32 CalculateRequestQueueLengthMax();
+
 float GIasStatisticsLogInterval = 60.f;
 static FAutoConsoleVariableRef CVar_StatisticsLogInterval(
 	TEXT("ias.StatisticsLogInterval"),
@@ -871,8 +873,9 @@ void FOnDemandIoBackendStats::PrintOnScreenStatistics(FCoreDelegates::FSeverityM
 			Builder << TEXT(" - TTFB: ") << TtfbMs << TEXT("ms");
 
 			const int32 NumConcurrentRequests = GHttpNumConcurrentRequests;
-			Builder << TEXT(" - Queue: ") << NumConcurrentRequests << TEXT("/") << GIasHttpConcurrentRequests;
-			if (NumConcurrentRequests >= GIasHttpConcurrentRequests)
+			const int32 MaxConcurrentRequests = CalculateRequestQueueLengthMax();
+			Builder << TEXT(" - Queue: ") << NumConcurrentRequests << TEXT("/") << MaxConcurrentRequests;
+			if (NumConcurrentRequests > MaxConcurrentRequests)
 			{
 				Builder << TEXT("(saturated)");
 			}
