Just in case someone sets negative cvar values

#rnx

diff --git a/Engine/Source/Runtime/Experimental/IoStore/OnDemand/Private/OnDemandHttpThread.cpp b/Engine/Source/Runtime/Experimental/IoStore/OnDemand/Private/OnDemandHttpThread.cpp
--- a/Engine/Source/Runtime/Experimental/IoStore/OnDemand/Private/OnDemandHttpThread.cpp
+++ b/Engine/Source/Runtime/Experimental/IoStore/OnDemand/Private/OnDemandHttpThread.cpp
@@ -207,7 +207,7 @@ int32 CalculateRequestQueueLengthMax()
 	// DOES NOT have any effect on throughput throttling.
 	int32 ConnNum = (GIaxHttpVersion == 2) ? 1 : GIasHttpConnectionCount;
 	int32 InflightNum = (GIaxHttpVersion == 2) ? GIaxHttpMaxInflight : 1;
-	return (ConnNum * InflightNum) + 2;
+	return FMath::Max((ConnNum * InflightNum) + 2, 1);
 }
 
 const int32 GOnDemandBackendThreadPriorities[5] =
