API simplification

#rnx

diff --git a/Engine/Source/Runtime/Experimental/IoStore/HttpClient/Internal/IO/Http/Client.h b/Engine/Source/Runtime/Experimental/IoStore/HttpClient/Internal/IO/Http/Client.h
--- a/Engine/Source/Runtime/Experimental/IoStore/HttpClient/Internal/IO/Http/Client.h
+++ b/Engine/Source/Runtime/Experimental/IoStore/HttpClient/Internal/IO/Http/Client.h
@@ -128,10 +128,13 @@ class FRequest
 public:
 	UE_API 				~FRequest();
 	UE_API 				FRequest(FRequest&& Rhs);
+						FRequest(const FRequest&) = delete;
+	UE_API FRequest&	operator = (FRequest&& Rhs);
+	FRequest&			operator = (const FRequest&) = delete;
 	bool				IsValid() const { return Ptr != nullptr; }
-	UE_API FRequest&&	Accept(EMimeType MimeType);
-	UE_API FRequest&&	Accept(FAnsiStringView MimeType);
-	UE_API FRequest&&	Header(FAnsiStringView Key, FAnsiStringView Value);
+	UE_API void			Accept(EMimeType MimeType);
+	UE_API void			Accept(FAnsiStringView MimeType);
+	UE_API void			Header(FAnsiStringView Key, FAnsiStringView Value);
 	UE_API void			Content(const void* Data, SIZE_T Size, EMimeType MimeType);
 	UE_API void			Content(const void* Data, SIZE_T Size, FAnsiStringView MimeType);
 
@@ -139,11 +142,6 @@ private:
 	friend				class FEventLoop;
 						FRequest() = default;
 	FActivityNode*		Ptr = nullptr;
-
-private:
-						FRequest(const FRequest&) = delete;
-	FRequest&			operator = (const FRequest&) = delete;
-	FRequest&			operator = (FRequest&&) = delete;
 };
 
 ////////////////////////////////////////////////////////////////////////////////
diff --git a/Engine/Source/Runtime/Experimental/IoStore/HttpClient/Private/Api.inl b/Engine/Source/Runtime/Experimental/IoStore/HttpClient/Private/Api.inl
--- a/Engine/Source/Runtime/Experimental/IoStore/HttpClient/Private/Api.inl
+++ b/Engine/Source/Runtime/Experimental/IoStore/HttpClient/Private/Api.inl
@@ -7,6 +7,19 @@ namespace UE::IoStore::HTTP
 
 // {{{1 request ................................................................
 
+////////////////////////////////////////////////////////////////////////////////
+FRequest::FRequest(FRequest&& Rhs)
+{
+	Swap(Ptr, Rhs.Ptr);
+}
+
+////////////////////////////////////////////////////////////////////////////////
+FRequest& FRequest::operator = (FRequest&& Rhs)
+{
+	Swap(Ptr, Rhs.Ptr);
+	return *this;
+}
+
 ////////////////////////////////////////////////////////////////////////////////
 FRequest::~FRequest()
 {
@@ -19,14 +32,7 @@ FRequest::~FRequest()
 }
 
 ////////////////////////////////////////////////////////////////////////////////
-FRequest::FRequest(FRequest&& Rhs)
-{
-	this->~FRequest();
-	Swap(Ptr, Rhs.Ptr);
-}
-
-////////////////////////////////////////////////////////////////////////////////
-FRequest&& FRequest::Accept(EMimeType MimeType)
+void FRequest::Accept(EMimeType MimeType)
 {
 	switch (MimeType)
 	{
@@ -38,21 +44,18 @@ FRequest&& FRequest::Accept(EMimeType MimeType)
 	case EMimeType::CbPackage:			return Header("Accept", "application/x-ue-pkg");
 	case EMimeType::CompressedBuffer:	return Header("Accept", "application/x-ue-comp");
 	}
-
-	return MoveTemp(*this);
 }
 
 ////////////////////////////////////////////////////////////////////////////////
-FRequest&& FRequest::Accept(FAnsiStringView MimeType)
+void FRequest::Accept(FAnsiStringView MimeType)
 {
-	return Header("Accept", MimeType);
+	Header("Accept", MimeType);
 }
 
 ////////////////////////////////////////////////////////////////////////////////
-FRequest&& FRequest::Header(FAnsiStringView Key, FAnsiStringView Value)
+void FRequest::Header(FAnsiStringView Key, FAnsiStringView Value)
 {
 	Ptr->AddHeader(Key, Value);
-	return MoveTemp(*this);
 }
 
 
diff --git a/Engine/Source/Runtime/Experimental/IoStore/HttpClient/Private/Test.inl b/Engine/Source/Runtime/Experimental/IoStore/HttpClient/Private/Test.inl
--- a/Engine/Source/Runtime/Experimental/IoStore/HttpClient/Private/Test.inl
+++ b/Engine/Source/Runtime/Experimental/IoStore/HttpClient/Private/Test.inl
@@ -189,7 +189,8 @@ static void ThrottleTest(FAnsiStringView TestUrl)
 		FEventLoop Loop;
 		Loop.Throttle(ThrottleKiB);
 
-		FRequest Request = Loop.Request("GET", Url).Accept("*/*");
+		FRequest Request = Loop.Request("GET", Url);
+		Request.Accept("*/*");
 		Loop.Send(MoveTemp(Request), [&] (const FTicketStatus& Status) {
 			check(Status.GetId() != FTicketStatus::EId::Error);
 			if (Status.GetId() == FTicketStatus::EId::Response)
@@ -667,24 +668,6 @@ static void HttpTest(const ANSICHAR* TestHost, FCertRootsRef VerifyCert)
 		WaitForLoopIdle();
 	}
 
-	// convenience
-	{
-		FRequest Request = Loop.Get(BuildUrl("/data"), ReqParams).Accept(EMimeType::Json);
-
-		Loop.Send(Loop.Get(BuildUrl("/data"), ReqParams).Accept(EMimeType::Json), HashSink),
-		Loop.Send(MoveTemp(Request), HashSink),
-		Loop.Send(Loop.Get("http://epicgames.com"), NoErrorSink),
-
-		WaitForLoopIdle();
-	}
-
-	// convenience
-	{
-		FRequest Request = Loop.Get(BuildUrl("/data"), ReqParams).Accept(EMimeType::Json);
-		Loop.Send(MoveTemp(Request), HashSink);
-		WaitForLoopIdle();
-	}
-
 	// pool
 	for (uint16 i = 1; i < 64; ++i)
 	{
@@ -866,7 +849,9 @@ static void HttpTest(const ANSICHAR* TestHost, FCertRootsRef VerifyCert)
 				FTicket Tickets[StressLoad];
 				for (FTicket& Ticket : Tickets)
 				{
-					Ticket = Loop.Send(Loop.Get(StressUrl, ReqParams).Header("Accept", "*/*"), Sink);
+					FRequest Request = Loop.Get(StressUrl, ReqParams);
+					Request.Header("Accept", "*/*");
+					Ticket = Loop.Send(MoveTemp(Request), Sink);
 				}
 
 				LoopTickDelay = AddDelay;
