Silence erroneous warning when running IASTool chunkplugin.

#rnx

- There is no need to log a warning when the default TocFlags are used.
- The original use case here was the idea that all containers in the plugin would be sent to the same system (IAD or IAS) but now we support both types at the same time so warning aobut it doesn't make much sense.
- In future work we will need to provide a better way for the caller to designate which container goes to which system anyway.
- For now remove the warning as the caller can't really do anything about it and it gives the impression that there is a problem when we are running expected code paths.

diff --git a/Engine/Source/Runtime/Experimental/IoStore/OnDemand/Private/Tool/ChunkPlugin.cpp b/Engine/Source/Runtime/Experimental/IoStore/OnDemand/Private/Tool/ChunkPlugin.cpp
--- a/Engine/Source/Runtime/Experimental/IoStore/OnDemand/Private/Tool/ChunkPlugin.cpp
+++ b/Engine/Source/Runtime/Experimental/IoStore/OnDemand/Private/Tool/ChunkPlugin.cpp
@@ -589,11 +589,10 @@ static int32 ChunkPluginCommandEntry(const FContext& Context)
 	}
 	else
 	{
-		const EOnDemandTocFlags FallbackFlags = EOnDemandTocFlags::InstallOnDemand;
-		UE_LOG(LogIoStoreOnDemand, Warning, TEXT("No TOC flags was specified, falling back to '%s'"), *LexToString(FallbackFlags));
-		TocFlags = FallbackFlags;
+		TocFlags = EOnDemandTocFlags::InstallOnDemand;
 	}
 
+
 	FS3Params S3Params(Context);
 
 	const bool bIncludeSigPak			= Context.Get<bool>(TEXT("-IncludeSigPak"), false);
