Say that we are able to buffer more that the defaults.

#rnx

diff --git a/Engine/Source/Runtime/Experimental/IoStore/HttpClient/Private/TransactionTwo.inl b/Engine/Source/Runtime/Experimental/IoStore/HttpClient/Private/TransactionTwo.inl
--- a/Engine/Source/Runtime/Experimental/IoStore/HttpClient/Private/TransactionTwo.inl
+++ b/Engine/Source/Runtime/Experimental/IoStore/HttpClient/Private/TransactionTwo.inl
@@ -65,6 +65,9 @@ namespace UE::IoStore::HTTP
 namespace DetailTwo
 {
 
+////////////////////////////////////////////////////////////////////////////////
+enum { RECV_WINDOW_SIZE = 1024 << 10 };
+
 ////////////////////////////////////////////////////////////////////////////////
 struct FFrame
 {
@@ -245,12 +248,7 @@ FDriverNg::FDriverNg(FTlsPeer& Peer)
 	RecvRemaining = 0 - int32(sizeof(RecvFrame));
 
 	nghttp2_session_client_new3(&Session, Statics.Instance, this, nullptr, &Mem);
-
-	// Perhaps in the future a window size can be picked - or adjusted - by
-	// monitoring latency and throughput.
-	/*
-	nghttp2_session_set_local_window_size(Session, 0, 0, 256 << 10);
-	*/
+	nghttp2_session_set_local_window_size(Session, 0, 0, RECV_WINDOW_SIZE);
 }
 
 ////////////////////////////////////////////////////////////////////////////////
@@ -347,7 +345,13 @@ FOutcome FDriverNg::PumpRecv(FTlsPeer& Peer, FMemoryView Data)
 ////////////////////////////////////////////////////////////////////////////////
 FOutcome FDriverNg::Handshake(FTlsPeer& Peer)
 {
-	nghttp2_submit_settings(Session, NGHTTP2_FLAG_NONE, nullptr, 0);
+	nghttp2_settings_entry Settings[] = {
+		{ NGHTTP2_SETTINGS_INITIAL_WINDOW_SIZE,		RECV_WINDOW_SIZE },
+		{ NGHTTP2_SETTINGS_MAX_FRAME_SIZE,			RECV_WINDOW_SIZE >> 2 },
+		{ NGHTTP2_SETTINGS_ENABLE_PUSH,				0 },
+	};
+	nghttp2_submit_settings(Session, NGHTTP2_FLAG_NONE, Settings, UE_ARRAY_COUNT(Settings));
+
 	FOutcome Outcome = PumpSend(Peer);
 	return Outcome.IsError() ? Outcome : FOutcome::Ok();
 }
