Wrong half of the string was return.

#rnx

diff --git a/Engine/Source/Runtime/Experimental/IoStore/HttpClient/Private/Api.inl b/Engine/Source/Runtime/Experimental/IoStore/HttpClient/Private/Api.inl
--- a/Engine/Source/Runtime/Experimental/IoStore/HttpClient/Private/Api.inl
+++ b/Engine/Source/Runtime/Experimental/IoStore/HttpClient/Private/Api.inl
@@ -117,10 +117,9 @@ void FResponse::GetContentType(FAnsiStringView& Out) const
 {
 	Out = GetHeader("Content-Type");
 
-	int32 SemiColon;
-	if (Out.FindChar(';', SemiColon))
+	if (int32 SemiColon; Out.FindChar(';', SemiColon))
 	{
-		Out = Out.Mid(SemiColon).TrimEnd();
+		Out = Out.Left(SemiColon).TrimEnd();
 	}
 }
 
