Chaos networking: temp fix for bug where the pre-rollback simulation time would still be in effect during the ProcessInputs phase of the first resimulated frame

#rnx

diff --git a/Engine/Source/Runtime/Experimental/Chaos/Private/PBDRigidsSolver.cpp b/Engine/Source/Runtime/Experimental/Chaos/Private/PBDRigidsSolver.cpp
--- a/Engine/Source/Runtime/Experimental/Chaos/Private/PBDRigidsSolver.cpp
+++ b/Engine/Source/Runtime/Experimental/Chaos/Private/PBDRigidsSolver.cpp
@@ -2193,6 +2193,13 @@ namespace Chaos
 								{
 									SCOPED_NAMED_EVENT(ChaosRewindAndResim_ResetCollisions, FColor::Red);
 									MTime = PushData->StartTime;	//not sure if sub-steps have proper StartTime so just do this once and let solver evolve remaining time
+
+									// @todo(chaos) : replace this temporary fix with a better method of syncing sim times prior to a resimulation
+									for (ISimCallbackObject* Callback : SimCallbackObjects)
+									{
+										Callback->SetSimAndDeltaTime_Internal(MTime, MLastDt);
+									}
+
 									GetEvolution()->ResetCollisions();
 								}
 
