BoxOverlapActors failed to catch an Overlaped actor with a particular scale, from a particular position
Gjk Large boxes unit test fix

diff --git a/Engine/Source/Programs/HeadlessChaos/Private/HeadlessChaosTestLargeScale.cpp b/Engine/Source/Programs/HeadlessChaos/Private/HeadlessChaosTestLargeScale.cpp
--- a/Engine/Source/Programs/HeadlessChaos/Private/HeadlessChaosTestLargeScale.cpp
+++ b/Engine/Source/Programs/HeadlessChaos/Private/HeadlessChaosTestLargeScale.cpp
@@ -154,4 +154,22 @@ namespace ChaosTest {
 		EXPECT_NEAR(Normal.Z, 1.0f, MaxNormalError);
 	}
 
+	GTEST_TEST(LargeScaleTests, TestSmallBoxLargeBoxGJKIntersection)
+	{
+		const FVec3 FloorMin(-250000.09155273438, -250000.01525878906, -50.000015258789062);
+		const FVec3 FloorMax(250000.00000000000, 249999.98474121094, 2.8421709430404007e-14);
+		const FVec3 BoxPos(6400.000000, 6400.000000, 95.304688);
+		const FVec3 BoxMin(5040.0000000000000, 5040.0000000000000, -4548.6328119999998);
+		const FVec3 BoxMax(7760.0000000000000, 7760.0000000000000, 4739.2421880000002);
+		const FRigidTransform3 AT(FVec3::Zero(), FRotation3::FromIdentity());
+		const FRigidTransform3 BT(FVec3::Zero(), FRotation3::FromIdentity());
+		const FRigidTransform3 BTtoAT = BT.GetRelativeTransform(AT);
+		FImplicitBox3 Box(BoxMin, BoxMax);
+		FImplicitBox3 Floor(FloorMin, FloorMax);
+		FVec3 InitialDir = -BoxPos;
+		FReal Thickness = 0;
+		const bool bResult = GJKIntersection(Floor, Box, BTtoAT, Thickness, InitialDir);
+		EXPECT_TRUE(bResult);
+	}
+
 }
\ No newline at end of file
diff --git a/Engine/Source/Runtime/Experimental/Chaos/Public/Chaos/GJK.h b/Engine/Source/Runtime/Experimental/Chaos/Public/Chaos/GJK.h
--- a/Engine/Source/Runtime/Experimental/Chaos/Public/Chaos/GJK.h
+++ b/Engine/Source/Runtime/Experimental/Chaos/Public/Chaos/GJK.h
@@ -13,7 +13,6 @@
 #include "Chaos/SimplexVectorized.h"
 #include "Chaos/EPAVectorized.h"
 
-
 #include "ChaosCheck.h"
 #include "ChaosLog.h"
 
@@ -189,11 +188,15 @@ namespace Chaos
 	template <typename T, typename TGeometryA, typename TGeometryB>
 	bool GJKIntersection(const TGeometryA& RESTRICT A, const TGeometryB& RESTRICT B, const TRigidTransform<T, 3>& BToATM, const T InThicknessA = 0, const TVector<T, 3>& InitialDir = TVector<T, 3>(1, 0, 0))
 	{
-		const FReal EpsilonScale = FMath::Max<FReal>(A.TGeometryA::BoundingBox().Extents().Max(), B.TGeometryB::BoundingBox().Extents().Max());
-		return GJKIntersection(A, B, EpsilonScale, BToATM, InThicknessA, InitialDir);
-	}
+		const T EpsilonScale = FMath::Max<T>(A.TGeometryA::BoundingBox().Extents().Max(), B.TGeometryB::BoundingBox().Extents().Max());
 
-	
+		const bool bUseDouble = EpsilonScale > 1.e5;
+		if (bUseDouble)
+		{
+			return GJKIntersectionImpl<T, FReal>(A, B, EpsilonScale, BToATM, InThicknessA, InitialDir);
+		}
+		return GJKIntersectionImpl<T, FRealSingle>(A, B, EpsilonScale, BToATM, InThicknessA, InitialDir);
+	}
 
 	template <typename TGeometryA, typename TGeometryB>
 	bool GJKIntersectionSimd(const TGeometryA& RESTRICT A, const TGeometryB& RESTRICT B, const VectorRegister4Float& Translation, const VectorRegister4Float& Rotation, FRealSingle InThicknessA, const VectorRegister4Float& InitialDir)
@@ -202,19 +205,20 @@ namespace Chaos
 		return GJKIntersectionSimd(A, B, FRealSingle(EpsilonScale), Translation, Rotation, InThicknessA, InitialDir);
 	}
 
-	inline bool GJKIntersectionSimd(const FGeomGJKHelperSIMD& RESTRICT A, const FGeomGJKHelperSIMD& RESTRICT B, FRealSingle EpsilonScale, const VectorRegister4Float& Translation, const VectorRegister4Float& Rotation, FRealSingle InThicknessA, const VectorRegister4Float& InitialDir)
+	template <typename T>
+	inline bool GJKIntersectionSimd(const FGeomGJKHelperSIMD& RESTRICT A, const FGeomGJKHelperSIMD& RESTRICT B, FRealSingle EpsilonScale, const T& Translation, const T& Rotation, FRealSingle InThicknessA, const T& InitialDir)
 	{
-		VectorRegister4Float V = VectorNegate(InitialDir);
+		T V = VectorNegate(InitialDir);
 		V = VectorNormalizeSafe(V, MakeVectorRegisterFloatConstant(-1.f, 0.f, 0.f, 0.f));
 
-		const VectorRegister4Float AToBRotation = VectorQuaternionInverse(Rotation);
+		const T AToBRotation = VectorQuaternionInverse(Rotation);
 		bool bTerminate;
 		bool bNearZero = false;
 		int NumIterations = 0;
-		VectorRegister4Float PrevDist2Simd = MakeVectorRegisterFloatConstant(FLT_MAX, FLT_MAX, FLT_MAX, FLT_MAX);
+		T PrevDist2Simd = MakeVectorRegisterFloatConstant(FLT_MAX, FLT_MAX, FLT_MAX, FLT_MAX);
 
-		VectorRegister4Float SimplexSimd[4] = { VectorZeroFloat(), VectorZeroFloat(), VectorZeroFloat(), VectorZeroFloat() };
-		VectorRegister4Float BarycentricSimd;
+		T SimplexSimd[4] = { VectorZero(), VectorZero(), VectorZero(), VectorZero() };
+		T BarycentricSimd;
 		int32 NumVerts = 0;
 
 		FRealSingle ThicknessA;
@@ -225,8 +229,8 @@ namespace Chaos
 		EpsilonScale = FMath::Min<FRealSingle>(EpsilonScale, 1e5f);
 		
 		const FRealSingle InflationReal = ThicknessA + ThicknessB + 1e-6f * EpsilonScale;
-		const VectorRegister4Float Inflation = VectorSet1(InflationReal);
-		const VectorRegister4Float Inflation2 = VectorMultiply(Inflation, Inflation);
+		const T Inflation = VectorSet1(InflationReal);
+		const T Inflation2 = VectorMultiply(Inflation, Inflation);
 
 		int32 VertexIndexA = INDEX_NONE, VertexIndexB = INDEX_NONE;
 		do
@@ -235,12 +239,12 @@ namespace Chaos
 			{
 				break;	//if taking too long just stop. This should never happen
 			}
-			const VectorRegister4Float NegVSimd = VectorNegate(V);
-			const VectorRegister4Float SupportASimd = A.SupportFunction(NegVSimd, ThicknessA);
-			const VectorRegister4Float VInBSimd = VectorQuaternionRotateVector(AToBRotation, V);
-			const VectorRegister4Float SupportBLocalSimd = B.SupportFunction(VInBSimd, ThicknessB);
-			const VectorRegister4Float SupportBSimd = VectorAdd(VectorQuaternionRotateVector(Rotation, SupportBLocalSimd), Translation);
-			const VectorRegister4Float WSimd = VectorSubtract(SupportASimd, SupportBSimd);
+			const T NegVSimd = VectorNegate(V);
+			const T SupportASimd = A.SupportFunction(TMakeVectorRegisterFloatFromDouble<T>(NegVSimd), ThicknessA);
+			const T VInBSimd = VectorQuaternionRotateVector(AToBRotation, V);
+			const T SupportBLocalSimd = B.SupportFunction(TMakeVectorRegisterFloatFromDouble<T>(VInBSimd), ThicknessB);
+			const T SupportBSimd = VectorAdd(VectorQuaternionRotateVector(Rotation, SupportBLocalSimd), Translation);
+			const T WSimd = VectorSubtract(SupportASimd, SupportBSimd);
 
 			if (VectorMaskBits(VectorCompareGT(VectorDot3(V, WSimd), Inflation)))
 			{
@@ -249,9 +253,9 @@ namespace Chaos
 
 			SimplexSimd[NumVerts++] = WSimd;
 
-			V = VectorSimplexFindClosestToOrigin<VectorRegister4Float, false>(SimplexSimd, NumVerts, BarycentricSimd, nullptr, nullptr);
+			V = VectorSimplexFindClosestToOrigin<T, false>(SimplexSimd, NumVerts, BarycentricSimd, nullptr, nullptr);
 
-			const VectorRegister4Float NewDist2Simd = VectorDot3(V, V);///
+			const T NewDist2Simd = VectorDot3(V, V);
 			bNearZero = VectorMaskBits(VectorCompareLT(NewDist2Simd, Inflation2)) != 0;
 
 			//as simplices become degenerate we will stop making progress. This is a side-effect of precision, in that case take V as the current best approximation
@@ -270,20 +274,19 @@ namespace Chaos
 		return bNearZero;
 	}
 
-	template <typename T>
-	bool GJKIntersection(const FGeomGJKHelperSIMD& RESTRICT A, const FGeomGJKHelperSIMD& RESTRICT B, FReal EpsilonScale, const TRigidTransform<T, 3>& BToATM, const T InThicknessA = 0, const TVector<T, 3>& InitialDir = TVector<T, 3>(1, 0, 0))
+	template <typename PrevT, typename NewT>
+	bool GJKIntersectionImpl(const FGeomGJKHelperSIMD& RESTRICT A, const FGeomGJKHelperSIMD& RESTRICT B, PrevT EpsilonScale, const TRigidTransform<PrevT, 3>& BToATM, const PrevT InThicknessA = 0, const TVector<PrevT, 3>& InitialDir = TVector<PrevT, 3>(1, 0, 0))
 	{
-		const UE::Math::TQuat<T>& RotationDouble = BToATM.GetRotation();
-		VectorRegister4Float RotationSimd = MakeVectorRegisterFloatFromDouble(MakeVectorRegister(RotationDouble.X, RotationDouble.Y, RotationDouble.Z, RotationDouble.W));
+		const UE::Math::TQuat<PrevT>& Rotation = BToATM.GetRotation();
+		TVectorRegisterType<NewT> RotationSimd = MakeVectorRegister(NewT(Rotation.X), NewT(Rotation.Y), NewT(Rotation.Z), NewT(Rotation.W));
 
-		const UE::Math::TVector<T>& TranslationDouble = BToATM.GetTranslation();
-		const VectorRegister4Float TranslationSimd = MakeVectorRegisterFloatFromDouble(MakeVectorRegister(TranslationDouble.X, TranslationDouble.Y, TranslationDouble.Z, 0.0));
-		// Normalize rotation
+		const UE::Math::TVector<PrevT>& TranslationDouble = BToATM.GetTranslation();
+		const TVectorRegisterType<NewT> TranslationSimd = MakeVectorRegister(NewT(TranslationDouble.X), NewT(TranslationDouble.Y), NewT(TranslationDouble.Z), NewT(0.0));
 		RotationSimd = VectorNormalizeSafe(RotationSimd, GlobalVectorConstants::Float0001);
 
-		const VectorRegister4Float InitialDirSimd = MakeVectorRegisterFloatFromDouble(MakeVectorRegister(InitialDir[0], InitialDir[1], InitialDir[2], 0.0));
+		const TVectorRegisterType<NewT> InitialDirSimd = MakeVectorRegister(NewT(InitialDir[0]), NewT(InitialDir[1]), NewT(InitialDir[2]), NewT(0.0));
 
-		return GJKIntersectionSimd(A, B, FRealSingle(EpsilonScale), TranslationSimd, RotationSimd, FRealSingle(InThicknessA), InitialDirSimd);
+		return GJKIntersectionSimd<TVectorRegisterType<NewT>>(A, B, FRealSingle(EpsilonScale), TranslationSimd, RotationSimd, FRealSingle(InThicknessA), InitialDirSimd);
 	}
 
 	/** 
