Rename confusing cVar name
This is to set the pull data in single thread which is the communcation from the physics thread to the Game thread

diff --git a/Engine/Source/Runtime/Experimental/Chaos/Private/PBDRigidsSolver.cpp b/Engine/Source/Runtime/Experimental/Chaos/Private/PBDRigidsSolver.cpp
--- a/Engine/Source/Runtime/Experimental/Chaos/Private/PBDRigidsSolver.cpp
+++ b/Engine/Source/Runtime/Experimental/Chaos/Private/PBDRigidsSolver.cpp
@@ -422,8 +422,8 @@ namespace Chaos
 		bool bRemoveParticleFromMovingKinematicsOnDisable = true;
 		FAutoConsoleVariableRef CVarChaosRemoveParticleFromMovingKinematicsOnDisable(TEXT("p.Chaos.RemoveParticleFromMovingKinematicsOnDisable"), bRemoveParticleFromMovingKinematicsOnDisable, TEXT(""));
 
-		bool bSingleThreadPushData = false;
-		FAutoConsoleVariableRef CVarChaosSingleThreadPushData(TEXT("p.Chaos.SingleThreadPushData"), bSingleThreadPushData, TEXT("Run push Data in Single Thread."));
+		bool bSingleThreadPullData = false;
+		FAutoConsoleVariableRef CVarChaosSingleThreadPushData(TEXT("p.Chaos.SingleThreadPullData"), bSingleThreadPullData, TEXT("Run pull Data in Single Thread."));
 	}
 }
 
@@ -2664,7 +2664,7 @@ namespace Chaos
 	{
 		FPullPhysicsData* PullData = MarshallingManager.GetCurrentPullData_Internal();
 
-		if (!(GSingleThreadedPhysics || CVars::bSingleThreadPushData))
+		if (!(GSingleThreadedPhysics || CVars::bSingleThreadPullData))
 		{
 			using TSOA = typename TParticleView<FPBDRigidParticles>::TSOA;
 			using THandle = typename TSOA::THandleType;
@@ -2781,11 +2781,12 @@ namespace Chaos
 			RewindData->BufferPhysicsResults(PullData->DirtyRigidErrors);
 		}
 
-		if (!(GSingleThreadedPhysics || CVars::bSingleThreadPushData))
+		if (!(GSingleThreadedPhysics || CVars::bSingleThreadPullData))
 		{
 			TaskDispatcher->WaitTaskEndBufferResults();
 		}
 
+
 		// Now that results have been buffered we have completed a solve step so we can broadcast that event
 		EventPostSolve.Broadcast(MLastDt);
 
