Dataflow : Add a way to get a Const struct view from a FContextValue

diff --git a/Engine/Source/Runtime/Experimental/Dataflow/Core/Private/Dataflow/DataflowContextCache.cpp b/Engine/Source/Runtime/Experimental/Dataflow/Core/Private/Dataflow/DataflowContextCache.cpp
--- a/Engine/Source/Runtime/Experimental/Dataflow/Core/Private/Dataflow/DataflowContextCache.cpp
+++ b/Engine/Source/Runtime/Experimental/Dataflow/Core/Private/Dataflow/DataflowContextCache.cpp
@@ -52,6 +52,15 @@ namespace UE::Dataflow
 		return {};
 	}
 
+	FConstStructView FContextCacheElementReference::GetConstStructView(const IContextCacheStore& Context) const
+	{
+		if (const TUniquePtr<FContextCacheElementBase>* Cache = Context.FindCacheElement(DataKey))
+		{
+			return (*Cache)->GetConstStructView(Context);
+		}
+		return {};
+	}
+
 	TUniquePtr<FContextCacheElementBase> FContextCacheElementReference::Clone(const IContextCacheStore& Context) const
 	{
 		if (const TUniquePtr<FContextCacheElementBase>* ReferencedCacheEntry = Context.FindCacheElement(DataKey))
@@ -133,6 +142,11 @@ namespace UE::Dataflow
 		return MakeUnique<FContextCacheElementUStructArray>(InNodeGuid, InProperty, FConstStructView(InstancedStruct), InNodeHash, InTimestamp);
 	}
 
+	FConstStructView FContextCacheElementUStruct::GetConstStructView(const IContextCacheStore& Context) const
+	{
+		return FConstStructView(InstancedStruct);
+	}
+
 	TUniquePtr<FContextCacheElementBase> FContextCacheElementUStruct::Clone(const IContextCacheStore& Context) const
 	{
 		return MakeUnique<FContextCacheElementUStruct>(*this);
@@ -289,7 +303,7 @@ namespace UE::Dataflow
 
 	int32 FContextValue::Num() const
 	{
-		if (const FContextCacheElementBase* CacheEntry = GetCacheEntry())
+		if (const FContextCacheElementBase* CacheEntry = GetCacheEntryPrivate())
 		{
 			return CacheEntry->GetNumArrayElements(Context);
 		}
@@ -298,7 +312,7 @@ namespace UE::Dataflow
 
 	FContextValue FContextValue::GetAt(int32 Index) const
 	{
-		if (const FContextCacheElementBase* CacheEntry = GetCacheEntry())
+		if (const FContextCacheElementBase* CacheEntry = GetCacheEntryPrivate())
 		{
 			TUniquePtr<FContextCacheElementBase> NewCacheElement = CacheEntry->CreateFromArrayElement(Context, Index, CacheEntry->GetProperty(), CacheEntry->GetNodeGuid(), CacheEntry->GetNodeHash(), CacheEntry->GetTimestamp());
 			return FContextValue(Context, MoveTemp(NewCacheElement));
@@ -308,7 +322,7 @@ namespace UE::Dataflow
 
 	FContextValue FContextValue::ToArray() const
 	{
-		if (const FContextCacheElementBase* CacheEntry = GetCacheEntry())
+		if (const FContextCacheElementBase* CacheEntry = GetCacheEntryPrivate())
 		{
 			// arrays return themselves
 			if (CacheEntry->IsArray(Context))
@@ -328,7 +342,21 @@ namespace UE::Dataflow
 		return FContextValue(Context, 0);
 	}
 
+	FConstStructView FContextValue::GetConstStructView() const
+	{
+		if (const FContextCacheElementBase* CacheEntry = GetCacheEntryPrivate())
+		{
+			return CacheEntry->GetConstStructView(Context);
+		}
+		return {};
+	}
+
 	const FContextCacheElementBase* FContextValue::GetCacheEntry() const
+	{
+		return GetCacheEntryPrivate();
+	}
+
+	const FContextCacheElementBase* FContextValue::GetCacheEntryPrivate() const
 	{
 		if (CacheElement)
 		{
diff --git a/Engine/Source/Runtime/Experimental/Dataflow/Core/Public/Dataflow/DataflowContextCache.h b/Engine/Source/Runtime/Experimental/Dataflow/Core/Public/Dataflow/DataflowContextCache.h
--- a/Engine/Source/Runtime/Experimental/Dataflow/Core/Public/Dataflow/DataflowContextCache.h
+++ b/Engine/Source/Runtime/Experimental/Dataflow/Core/Public/Dataflow/DataflowContextCache.h
@@ -99,6 +99,9 @@ namespace UE::Dataflow
 		virtual TUniquePtr<FContextCacheElementBase> CreateFromArrayElement(const IContextCacheStore& Context, int32 Index, const FProperty* InProperty, const FGuid& InNodeGuid, uint32 InNodeHash, const FTimestamp& InTimestamp) const = 0;
 		virtual TUniquePtr<FContextCacheElementBase> CreateArrayFromElement(const IContextCacheStore& Context, const FProperty* InProperty, const FGuid& InNodeGuid, uint32 InNodeHash, const FTimestamp& InTimestamp) const = 0;
 
+		// Get the matching const struct view if available, otherwise return an invalid one 
+		virtual FConstStructView GetConstStructView(const IContextCacheStore& Context) const { return {}; }
+
 		EType GetType() const {	return Type; }
 
 		const FProperty* GetProperty() const { return Property; }
@@ -232,6 +235,8 @@ namespace UE::Dataflow
 
 		virtual TUniquePtr<FContextCacheElementBase> CreateArrayFromElement(const IContextCacheStore& Context, const FProperty* InProperty, const FGuid& InNodeGuid, uint32 InNodeHash, const FTimestamp& InTimestamp) const override;
 
+		virtual FConstStructView GetConstStructView(const IContextCacheStore& Context) const;
+
 		virtual TUniquePtr<FContextCacheElementBase> Clone(const IContextCacheStore& Context) const override;
 
 	private:
@@ -419,6 +424,9 @@ namespace UE::Dataflow
 		DATAFLOWCORE_API virtual TUniquePtr<FContextCacheElementBase> CreateFromArrayElement(const IContextCacheStore& Context, int32 Index, const FProperty* InProperty, const FGuid& InNodeGuid, uint32 InNodeHash, const FTimestamp& InTimestamp) const override;
 		DATAFLOWCORE_API virtual TUniquePtr<FContextCacheElementBase> CreateArrayFromElement(const IContextCacheStore& Context, const FProperty* InProperty, const FGuid& InNodeGuid, uint32 InNodeHash, const FTimestamp& InTimestamp) const override;
 
+		// Get the matching const struct view if available, otherwise return an invalid one 
+		DATAFLOWCORE_API virtual FConstStructView GetConstStructView(const IContextCacheStore& Context) const;
+
 		DATAFLOWCORE_API virtual TUniquePtr<FContextCacheElementBase> Clone(const IContextCacheStore& Context) const override;
 
 		//~ Begin FGCObject interface
@@ -646,9 +654,16 @@ namespace UE::Dataflow
 		DATAFLOWCORE_API FContextValue GetAt(int32 Index) const;
 		DATAFLOWCORE_API FContextValue ToArray() const;
 
+		DATAFLOWCORE_API FConstStructView GetConstStructView() const;
+
+		UE_DEPRECATED(5.8, "From release 5.10, this method is going to be private, direct access to the cache element is not authorized and the rest of the API of this class should be used instead to get data from the cache")
 		DATAFLOWCORE_API const FContextCacheElementBase* GetCacheEntry() const;
 
 	private:
+		friend class FContext; // to access GetCacheEntry when it's private
+
+		DATAFLOWCORE_API const FContextCacheElementBase* GetCacheEntryPrivate() const;
+
 		IContextCacheStore& Context;
 		FContextCacheKey CacheKey; // only set when CacheElement is null
 		TUniquePtr<FContextCacheElementBase> CacheElement; // when CacheKey is define this should be null
diff --git a/Engine/Source/Runtime/Experimental/Dataflow/Core/Public/Dataflow/DataflowNodeParameters.h b/Engine/Source/Runtime/Experimental/Dataflow/Core/Public/Dataflow/DataflowNodeParameters.h
--- a/Engine/Source/Runtime/Experimental/Dataflow/Core/Public/Dataflow/DataflowNodeParameters.h
+++ b/Engine/Source/Runtime/Experimental/Dataflow/Core/Public/Dataflow/DataflowNodeParameters.h
@@ -201,7 +201,7 @@ namespace UE::Dataflow
 
 			if constexpr (std::is_same_v<T, FContextValue>)
 			{
-				if (const FContextCacheElementBase* CacheEntryToClone = InValue.GetCacheEntry())
+				if (const FContextCacheElementBase* CacheEntryToClone = InValue.GetCacheEntryPrivate())
 				{
 					DataStoreEntry = CacheEntryToClone->Clone(*this);
 					if (DataStoreEntry)
