Constraints: check for indirect dependencies
- disabled by default for now but can be enabled using `Constraints.SupportIndirectDependencies true`

#rnx

diff --git a/Engine/Plugins/Animation/ControlRig/Source/ControlRig/Private/Constraints/ControlRigTransformableHandle.cpp b/Engine/Plugins/Animation/ControlRig/Source/ControlRig/Private/Constraints/ControlRigTransformableHandle.cpp
--- a/Engine/Plugins/Animation/ControlRig/Source/ControlRig/Private/Constraints/ControlRigTransformableHandle.cpp
+++ b/Engine/Plugins/Animation/ControlRig/Source/ControlRig/Private/Constraints/ControlRigTransformableHandle.cpp
@@ -351,6 +351,65 @@ UControlRigComponent* UTransformableControlHandle::GetControlRigComponent() cons
 	return ObjectBinding ? Cast<UControlRigComponent>(ObjectBinding->GetBoundObject()) : nullptr;
 }
 
+bool UTransformableControlHandle::HasIndirectDependencyWith(const UTransformableHandle& InOther) const
+{
+	const uint32 OtherHash = InOther.GetHash();
+	if (OtherHash == 0 || OtherHash == GetHash())
+	{
+		return false;
+	}
+	
+	FRigControlElement* ControlElement = GetControlElement();
+	if (!ControlElement)
+	{
+		return false;
+	}
+	
+	const UTransformableControlHandle* OtherCtrlHandle = InOther.IsValid() ? Cast<UTransformableControlHandle>(&InOther) : nullptr;
+	const FRigControlElement* OtherControlElement = OtherCtrlHandle ? OtherCtrlHandle->GetControlElement() : nullptr;
+	if (!OtherCtrlHandle || !OtherControlElement || OtherCtrlHandle->ControlRig != ControlRig)
+	{
+		return false;
+	}
+
+	// look for common ancestor
+	const URigHierarchy* Hierarchy = ControlRig->GetHierarchy();
+	static constexpr bool bRecursive = true;	
+	const FRigBaseElementParentArray& Parents = Hierarchy->GetParents(ControlElement, bRecursive);
+	const FRigBaseElementParentArray& OtherParents = Hierarchy->GetParents(OtherControlElement, bRecursive);
+
+	int32 Index = Parents.Num(), OtherIndex = OtherParents.Num();
+	auto HasCommonAncestor = [&Index, &Parents, &OtherIndex, &OtherParents, OtherControlElement, ControlElement]()
+	{
+		if (!Parents.IsEmpty() && !OtherParents.IsEmpty())
+		{
+			FRigBaseElement* CommonAncestor = nullptr;
+			while (Index > 0 && OtherIndex > 0 && Parents[Index-1] == OtherParents[OtherIndex-1])
+			{
+				CommonAncestor = Parents[Index-1];
+				--Index;
+				--OtherIndex;
+			}
+			return CommonAncestor != nullptr;
+		}
+		return false;
+	};
+
+	// if common ancestor, compare depths
+	if (HasCommonAncestor() && Index != OtherIndex)
+	{
+		return OtherIndex < Index; 
+	}
+
+	// otherwise, compare indices
+	if (OtherControlElement->GetIndex() < ControlElement->GetIndex())
+	{
+		return true;
+	}
+
+	return false;
+}
+
 bool UTransformableControlHandle::HasDirectDependencyWith(const UTransformableHandle& InOther) const
 {
 	const uint32 OtherHash = InOther.GetHash();
diff --git a/Engine/Plugins/Animation/ControlRig/Source/ControlRig/Public/Constraints/ControlRigTransformableHandle.h b/Engine/Plugins/Animation/ControlRig/Source/ControlRig/Public/Constraints/ControlRigTransformableHandle.h
--- a/Engine/Plugins/Animation/ControlRig/Source/ControlRig/Public/Constraints/ControlRigTransformableHandle.h
+++ b/Engine/Plugins/Animation/ControlRig/Source/ControlRig/Public/Constraints/ControlRigTransformableHandle.h
@@ -103,6 +103,9 @@ public:
 	/** Check for direct dependencies (ie hierarchy + skeletal mesh) with InOther. */
 	UE_API virtual bool HasDirectDependencyWith(const UTransformableHandle& InOther) const override;
 
+	/** Check for indirect dependencies with InOther. */
+	UE_API virtual bool HasIndirectDependencyWith(const UTransformableHandle& InOther) const override;
+
 	/** Look for a possible tick function that can be used as a prerequisite. */
 	UE_API virtual FTickPrerequisite GetPrimaryPrerequisite(const bool bAllowThis = true) const override;
 
diff --git a/Engine/Source/Runtime/Experimental/Animation/Constraints/Private/Transform/DependencyBuilder.cpp b/Engine/Source/Runtime/Experimental/Animation/Constraints/Private/Transform/DependencyBuilder.cpp
--- a/Engine/Source/Runtime/Experimental/Animation/Constraints/Private/Transform/DependencyBuilder.cpp
+++ b/Engine/Source/Runtime/Experimental/Animation/Constraints/Private/Transform/DependencyBuilder.cpp
@@ -35,6 +35,12 @@ static FAutoConsoleVariableRef CVarDebugDependencies(
 	bDebugDependencies,
 	TEXT("Print debug info about dependencies when creating a new constraint.") );
 
+static bool	bSupportIndirectDependencies = false;
+static FAutoConsoleVariableRef CVarSupportIndirectDependencies(
+	TEXT("Constraints.SupportIndirectDependencies"),
+	bSupportIndirectDependencies,
+	TEXT("Support self indirect dependencies between handles.") );
+
 FString GetConstraintLabel(const UTickableConstraint* InConstraint)
 {
 #if WITH_EDITOR
@@ -71,6 +77,24 @@ void LogDependency( const FString& InDescription, const UTransformableHandle* In
 		*GetHandleLabel(InParentHandle), *GetHandleLabel(InChildHandle),
 		*GetConstraintLabel(InParentConstraint), *GetConstraintLabel(InChildConstraint));
 }
+	
+void LogIndirectDependency(const UTransformableHandle* InParentHandle, const UTransformableHandle* InChildHandle,
+	const UTickableConstraint* InParentConstraint, const UTickableConstraint* InChildConstraint)
+{
+	if (!bDebugDependencies)
+	{
+		return;
+	}
+	
+	if (!InParentHandle || !InChildHandle || !InParentConstraint || !InChildConstraint)
+	{
+		return;
+	}
+
+	UE_LOG(LogTemp, Warning, TEXT("Self Indirect Dependency: '%s' has a lower hierarchy depth compared to '%s' so '%s' must tick before '%s'"),
+		*GetHandleLabel(InParentHandle), *GetHandleLabel(InChildHandle),
+		*GetConstraintLabel(InParentConstraint), *GetConstraintLabel(InChildConstraint));
+}
 
 void LogAttachmentDependency( const FString& InDescription, const UTransformableHandle* InAttachHandle, const UTransformableHandle* InChildHandle,
 								const UTickableConstraint* InAttachConstraint, const UTickableConstraint* InChildConstraint)
@@ -171,6 +195,12 @@ void FDependencyBuilder::BuildSelfDependencies(UWorld* InWorld, UTickableTransfo
 
 	const UObject* SelfTarget = ChildTarget;
 
+	auto SetDependency = [InWorld](const FGuid& InGuidToTickBefore, const FGuid& InGuidToTickAfter)
+	{
+		FConstraintsManagerController& Controller = FConstraintsManagerController::Get(InWorld);
+		Controller.SetConstraintsDependencies(InGuidToTickBefore, InGuidToTickAfter);
+	};
+	
 	auto SelfTargetPredicate = [InConstraint, SelfTarget](const ConstraintWeakPtr& Constraint)
 	{
 		const UTickableTransformConstraint* TransformConstraint = Cast<UTickableTransformConstraint>(Constraint.Get());
@@ -183,7 +213,8 @@ void FDependencyBuilder::BuildSelfDependencies(UWorld* InWorld, UTickableTransfo
 		return ParentTarget == SelfTarget && ChildTarget == SelfTarget;
 	};
 	const TArray< ConstraintWeakPtr > SelfConstraints = FConstraintsManagerController::Get(InWorld).GetConstraintsByPredicate(SelfTargetPredicate);
-	
+
+	bool bDirectDependencyAdded = false;
 	for (const ConstraintWeakPtr& SelfConstraint: SelfConstraints)
 	{
 		const UTickableTransformConstraint* TransformConstraint = Cast<UTickableTransformConstraint>(SelfConstraint);
@@ -191,25 +222,62 @@ void FDependencyBuilder::BuildSelfDependencies(UWorld* InWorld, UTickableTransfo
 		// if the new handles depend on that constraint child then, TransformConstraint should tick before
 		if (ParentHandle->HasDirectDependencyWith(*TransformConstraint->ChildTRSHandle))
 		{
-			FConstraintsManagerController::Get(InWorld).SetConstraintsDependencies(TransformConstraint->ConstraintID, InConstraint->ConstraintID);
+			SetDependency(TransformConstraint->ConstraintID, InConstraint->ConstraintID);
 			LogDependency(SelfDependencyDesc, TransformConstraint->ChildTRSHandle, ParentHandle, TransformConstraint, InConstraint);
+			bDirectDependencyAdded = true;
 		}
 		else if (ChildHandle->HasDirectDependencyWith(*TransformConstraint->ChildTRSHandle))
 		{
-			FConstraintsManagerController::Get(InWorld).SetConstraintsDependencies(TransformConstraint->ConstraintID, InConstraint->ConstraintID);
+			SetDependency(TransformConstraint->ConstraintID, InConstraint->ConstraintID);
 			LogDependency(SelfDependencyDesc, TransformConstraint->ChildTRSHandle, ChildHandle, TransformConstraint, InConstraint);
+			bDirectDependencyAdded = true;
 		}
 
 		// if the TransformConstraint handles depend on the new constraint child then, TransformConstraint should tick after
 		if (TransformConstraint->ParentTRSHandle->HasDirectDependencyWith(*ChildHandle))
 		{
-			FConstraintsManagerController::Get(InWorld).SetConstraintsDependencies(InConstraint->ConstraintID, TransformConstraint->ConstraintID);
+			SetDependency(InConstraint->ConstraintID, TransformConstraint->ConstraintID);
 			LogDependency(SelfDependencyDesc, ChildHandle, TransformConstraint->ParentTRSHandle, InConstraint, TransformConstraint);
+			bDirectDependencyAdded = true;
 		}
 		else if (TransformConstraint->ChildTRSHandle->HasDirectDependencyWith(*ChildHandle))
 		{
-			FConstraintsManagerController::Get(InWorld).SetConstraintsDependencies(InConstraint->ConstraintID, TransformConstraint->ConstraintID);
+			SetDependency(InConstraint->ConstraintID, TransformConstraint->ConstraintID);
 			LogDependency(SelfDependencyDesc, ChildHandle, TransformConstraint->ChildTRSHandle, InConstraint, TransformConstraint);
+			bDirectDependencyAdded = true;
+		}
+	}
+
+	// look for indirect dependencies (hierarchy depth, for example)
+	if (bSupportIndirectDependencies && !bDirectDependencyAdded)
+	{
+		for (const ConstraintWeakPtr& SelfConstraint: SelfConstraints)
+		{
+			const UTickableTransformConstraint* TransformConstraint = Cast<UTickableTransformConstraint>(SelfConstraint);
+			
+			// if the new handles indirectly depend on that constraint child then, TransformConstraint should tick before
+			if (ParentHandle->HasIndirectDependencyWith(*TransformConstraint->ChildTRSHandle))
+			{
+				SetDependency(TransformConstraint->ConstraintID, InConstraint->ConstraintID);
+				LogIndirectDependency(TransformConstraint->ChildTRSHandle, ParentHandle, TransformConstraint, InConstraint);
+			}
+			else if (ChildHandle->HasIndirectDependencyWith(*TransformConstraint->ChildTRSHandle))
+			{
+				SetDependency(TransformConstraint->ConstraintID, InConstraint->ConstraintID);
+				LogIndirectDependency(TransformConstraint->ChildTRSHandle, ChildHandle, TransformConstraint, InConstraint);
+			}
+
+			// if the TransformConstraint handles indirectly depend on the new constraint child then, TransformConstraint should tick after
+			if (TransformConstraint->ParentTRSHandle->HasIndirectDependencyWith(*ChildHandle))
+			{
+				SetDependency(InConstraint->ConstraintID, TransformConstraint->ConstraintID);
+				LogIndirectDependency(ChildHandle, TransformConstraint->ParentTRSHandle, InConstraint, TransformConstraint);
+			}
+			else if (TransformConstraint->ChildTRSHandle->HasIndirectDependencyWith(*ChildHandle))
+			{
+				SetDependency(InConstraint->ConstraintID, TransformConstraint->ConstraintID);
+				LogIndirectDependency(ChildHandle, TransformConstraint->ChildTRSHandle, InConstraint, TransformConstraint);
+			}
 		}
 	}
 }
diff --git a/Engine/Source/Runtime/Experimental/Animation/Constraints/Private/Transform/TransformConstraint.cpp b/Engine/Source/Runtime/Experimental/Animation/Constraints/Private/Transform/TransformConstraint.cpp
--- a/Engine/Source/Runtime/Experimental/Animation/Constraints/Private/Transform/TransformConstraint.cpp
+++ b/Engine/Source/Runtime/Experimental/Animation/Constraints/Private/Transform/TransformConstraint.cpp
@@ -387,12 +387,6 @@ void UTickableTransformConstraint::EnsurePrimaryDependency(const UWorld* InWorld
 	}
 }
 
-void UTickableTransformConstraint::OnActiveStateChanged() const
-{
-	// tick dependencies might need to be updated when dealing with cycles (than can be created between two controls for example) 
-	FConstraintCycleChecker::CheckAndFixCycles(this);
-}
-
 void UTickableTransformConstraint::PostLoad()
 {
 	Super::PostLoad();
@@ -478,13 +472,7 @@ void UTickableTransformConstraint::Evaluate(bool bTickHandlesAlso) const
 
 void UTickableTransformConstraint::SetActive(const bool bIsActive)
 {
-	const bool bNeedsUpdate = (Active != bIsActive);
 	Super::SetActive(bIsActive);
-	
-	if (bNeedsUpdate)
-	{
-		OnActiveStateChanged();
-	}
 }
 
 void UTickableTransformConstraint::SetChildGlobalTransform(const FTransform& InGlobal) const
diff --git a/Engine/Source/Runtime/Experimental/Animation/Constraints/Public/Transform/TransformConstraint.h b/Engine/Source/Runtime/Experimental/Animation/Constraints/Public/Transform/TransformConstraint.h
--- a/Engine/Source/Runtime/Experimental/Animation/Constraints/Public/Transform/TransformConstraint.h
+++ b/Engine/Source/Runtime/Experimental/Animation/Constraints/Public/Transform/TransformConstraint.h
@@ -150,7 +150,8 @@ protected:
 	ETransformConstraintType Type = ETransformConstraintType::Parent;
 
 	/** Handle active state modification if needed. */
-	CONSTRAINTS_API void OnActiveStateChanged() const;
+	UE_DEPRECATED(5.8, "Deprecated, no longer used.")
+	void OnActiveStateChanged() const {};
 
 	/** Returns the handle's tick function (ensuring it lives in the same world). */
 	CONSTRAINTS_API FTickFunction* GetHandleTickFunction(const TObjectPtr<UTransformableHandle>& InHandle) const;
diff --git a/Engine/Source/Runtime/Experimental/Animation/Constraints/Public/Transform/TransformableHandle.h b/Engine/Source/Runtime/Experimental/Animation/Constraints/Public/Transform/TransformableHandle.h
--- a/Engine/Source/Runtime/Experimental/Animation/Constraints/Public/Transform/TransformableHandle.h
+++ b/Engine/Source/Runtime/Experimental/Animation/Constraints/Public/Transform/TransformableHandle.h
@@ -121,7 +121,10 @@ public:
 
 	/** Check for direct dependencies with InOther. */
 	virtual bool HasDirectDependencyWith(const UTransformableHandle& InOther) const PURE_VIRTUAL(HasDirectDependencyWith, return false;);
-
+	
+	/** Check for indirect dependencies with InOther. (returns false by default, override in child classes if needed) */
+	virtual bool HasIndirectDependencyWith(const UTransformableHandle& InOther) const {return false;}
+	
 	/** Look for a possible tick function that can be used as a prerequisite. */
 	virtual FTickPrerequisite GetPrimaryPrerequisite(const bool bAllowThis = true) const PURE_VIRTUAL(GetPrimaryPrerequisite, return FTickPrerequisite(););
 	
