Add access detectors to FPBDIslandManager
- Series of crashes in UnbindEdgeFromNodes, RemoveConstraint, and RemoveParticleConstraints
- Potentially happens due to multi-threaded writing to IslandManager data
- Access detectors are supposed to trigger the crash and return the callstack if this happens.
- Access detectors are only activated in Development, Debug or PIE

diff --git a/Engine/Source/Programs/HeadlessChaos/Private/HeadlessChaosTestGraphEvolution.cpp b/Engine/Source/Programs/HeadlessChaos/Private/HeadlessChaosTestGraphEvolution.cpp
--- a/Engine/Source/Programs/HeadlessChaos/Private/HeadlessChaosTestGraphEvolution.cpp
+++ b/Engine/Source/Programs/HeadlessChaos/Private/HeadlessChaosTestGraphEvolution.cpp
@@ -10,6 +10,7 @@
 #include "Chaos/Utilities.h"
 #include "Modules/ModuleManager.h"
 #include "HAL/ConsoleManager.h"
+#include "Chaos/Framework/Parallel.h"
 
 namespace ChaosTest
 {
@@ -2807,5 +2808,43 @@ namespace ChaosTest
 		CVarMinMotionlessRatio->Set(MinMotionlessRatioDefault);
 	}
 
+	// Multithreaded access to the island manager is not supported.
+	// This test attempts to remove particles from the constraint graph in parallel
+	// which will corrupt the data stored in the island manager.
+	// This is triggers similar crashes as in UE-215521, UE-296295, UE-296297, UE-297690, UE-314497.
+	// The test needs to be disabled by default and helps testing the debug access detectors added to the island manager.
+	TEST_P(GraphEvolutionTests, DISABLED_TestConstraintGraph_MultithreadedAccess)
+	{
+		const bool bEnablePartialSleep = true;
+		const bool bInParallel = GetParam();
+
+		FGraphEvolutionTest Test(10, bEnablePartialSleep);
+		Test.MakeChain();
+		Test.Advance();
+
+		auto IslandManagerRemoveParticle = [](FPBDIslandManager* IslandManager,
+			FGeometryParticleHandle* ParticleHandle)
+			{
+				IslandManager->RemoveParticle(ParticleHandle);
+			};
+
+		if (bInParallel)
+		{
+			// Multithreaded access not supported, this will crash in debug but pass in development.
+			PhysicsParallelFor(Test.ParticleHandles.Num(),
+				[&](const int32 Index)
+				{
+					IslandManagerRemoveParticle(Test.IslandManager, Test.ParticleHandles[Index]);
+				});
+		}
+		else
+		{
+			for (int32 Index = 0; Index < Test.ParticleHandles.Num(); ++Index)
+			{
+				IslandManagerRemoveParticle(Test.IslandManager, Test.ParticleHandles[Index]);
+			}
+		}
+	}
+
 	INSTANTIATE_TEST_SUITE_P(, GraphEvolutionTests,	testing::Bool());
 }
diff --git a/Engine/Source/Runtime/Experimental/Chaos/Private/Chaos/Island/IslandManager.cpp b/Engine/Source/Runtime/Experimental/Chaos/Private/Chaos/Island/IslandManager.cpp
--- a/Engine/Source/Runtime/Experimental/Chaos/Private/Chaos/Island/IslandManager.cpp
+++ b/Engine/Source/Runtime/Experimental/Chaos/Private/Chaos/Island/IslandManager.cpp
@@ -1008,13 +1008,18 @@ namespace Chaos::Private
 			UpdateGraphNode(Node);
 		}
 
-		// The above loop may leave orphaned nodes in the Nodes array when dyanmics are made
+		// The above loop may leave orphaned nodes in the Nodes array when dynamics are made
 		// kinematic and we have removed any kinematic-kinematic constraints. We don't 
 		// handle node removal above because updating one node to kinematic can result 
 		// in removal of many nodes which is awkward while iterating
 		for (int32 NodeIndex = Nodes.Num() - 1; NodeIndex >= 0; --NodeIndex)
 		{
 			FPBDIslandParticle* Node = Nodes[NodeIndex];
+			check(Node != nullptr);
+
+			// Every dynamic node needs to be in an island, non-dynamic nodes cannot be in an island.
+			check((Node->Flags.bIsDynamic && (Node->Island != nullptr)) || (!Node->Flags.bIsDynamic && (Node->Island == nullptr)));
+			
 			if (Node->Edges.IsEmpty() && (Node->Island == nullptr))
 			{
 				// We have an orphaned node - it must be removed from the graph here because
@@ -1226,6 +1231,7 @@ namespace Chaos::Private
 		{
 			check(GetGraphNode(Particle) == nullptr);
 
+			UE_MT_SCOPED_WRITE_ACCESS(GlobalNodesAccessDetector);
 			FPBDIslandParticle* Node = Nodes.Alloc(Particle);
 
 			// Initial state
@@ -1256,7 +1262,9 @@ namespace Chaos::Private
 		{
 			// We should have already been removed from any island
 			check(Node->Island == nullptr);
+			check(Node->IslandArrayIndex == INDEX_NONE);
 
+			UE_MT_SCOPED_WRITE_ACCESS(GlobalNodesAccessDetector);
 			Nodes.Free(Node);
 		}
 	}
@@ -1284,6 +1292,8 @@ namespace Chaos::Private
 			FPBDIslandParticle* Node0 = GetOrCreateGraphNode(ConstrainedParticles[0]);
 			FPBDIslandParticle* Node1 = GetOrCreateGraphNode(ConstrainedParticles[1]);
 
+			UE_MT_SCOPED_WRITE_ACCESS(GlobalEdgesAccessDetector);
+
 			// Create the edge and bind it to the nodes
 			FPBDIslandConstraint* Edge = Edges.Alloc(ContainerId, Constraint);
 
@@ -1301,10 +1311,12 @@ namespace Chaos::Private
 		{
 			// We should have already been removed from any island
 			check(Edge->Island == nullptr);
+			check(Edge->IslandArrayIndex == INDEX_NONE);
 
 			// Disconnect from the nodes
 			UnbindEdgeFromNodes(Edge);
 
+			UE_MT_SCOPED_WRITE_ACCESS(GlobalEdgesAccessDetector);
 			Edges.Free(Edge);
 		}
 	}
@@ -1314,8 +1326,13 @@ namespace Chaos::Private
 		check(Edge != nullptr);
 		check(Edge->Nodes[0] == nullptr);
 		check(Edge->Nodes[1] == nullptr);
+		check(Edge->NodeArrayIndices[0] == INDEX_NONE);
+		check(Edge->NodeArrayIndices[1] == INDEX_NONE);
 		check(Edge->Island == nullptr);
 
+		UE_MT_SCOPED_WRITE_ACCESS(LocalEdgeAccessDetector);
+		UE_MT_SCOPED_WRITE_ACCESS(LocalNodeAccessDetector);
+
 		if (Node0 != nullptr)
 		{
 			Edge->Nodes[0] = Node0;
@@ -1334,6 +1351,9 @@ namespace Chaos::Private
 
 	void FPBDIslandManager::UnbindEdgeFromNodes(FPBDIslandConstraint* Edge)
 	{
+		UE_MT_SCOPED_WRITE_ACCESS(LocalEdgeAccessDetector);
+		UE_MT_SCOPED_WRITE_ACCESS(LocalNodeAccessDetector);
+
 		for (int32 NodeIndex = 0; NodeIndex < 2; ++NodeIndex)
 		{
 			if (FPBDIslandParticle* Node = Edge->Nodes[NodeIndex])
@@ -1352,12 +1372,21 @@ namespace Chaos::Private
 					// We don't know which of the edges nodes we are...
 					if (Node->Edges[ArrayIndex]->Nodes[0] == Node)
 					{
+						// The node we swapped in must have been at the end of the array.
+						check(Node->Edges[ArrayIndex]->NodeArrayIndices[0] == Node->Edges.Num());
 						Node->Edges[ArrayIndex]->NodeArrayIndices[0] = ArrayIndex;
 					}
 					else if (Node->Edges[ArrayIndex]->Nodes[1] == Node)
 					{
+						// The node we swapped in must have been at the end of the array.
+						check(Node->Edges[ArrayIndex]->NodeArrayIndices[1] == Node->Edges.Num());
 						Node->Edges[ArrayIndex]->NodeArrayIndices[1] = ArrayIndex;
 					}
+					else
+					{
+						// If we get here, the list Node->Edges must be corrupted.
+						checkf(false, TEXT("Incorrect node referenced."));
+					}
 				}
 			}
 		}
@@ -1560,6 +1589,8 @@ namespace Chaos::Private
 			// Kinematics may be in many islands (we need to check the node's edges to visit them).
 			if (Node->Flags.bIsDynamic)
 			{
+				UE_MT_SCOPED_WRITE_ACCESS(LocalNodeAccessDetector);
+				UE_MT_SCOPED_WRITE_ACCESS(LocalIslandAccessDetector);
 				Node->Island = Island;
 				Node->IslandArrayIndex = Island->Nodes.Add(Node);
 				Island->SetIterationSettings(FIterationSettings::Merge(Island->GetIterationSettings(), Node->GetIterationSettings()));
@@ -1582,6 +1613,9 @@ namespace Chaos::Private
 			const int32 ArrayIndex = Node->IslandArrayIndex;
 			check(Island->Nodes[ArrayIndex] == Node);
 
+			UE_MT_SCOPED_WRITE_ACCESS(LocalIslandAccessDetector);
+			UE_MT_SCOPED_WRITE_ACCESS(LocalNodeAccessDetector);
+
 			Island->Nodes.RemoveAtSwap(ArrayIndex, EAllowShrinking::No);
 			if (ArrayIndex < Island->Nodes.Num())
 			{
@@ -1621,6 +1655,9 @@ namespace Chaos::Private
 		check(Edge != nullptr);
 		check(Island != nullptr);
 
+		UE_MT_SCOPED_WRITE_ACCESS(LocalEdgeAccessDetector);
+		UE_MT_SCOPED_WRITE_ACCESS(LocalIslandAccessDetector);
+
 		const int32 ContainerIndex = Edge->ContainerIndex;
 		Edge->IslandArrayIndex = Island->ContainerEdges[ContainerIndex].Add(Edge);
 		Edge->Island = Island;
@@ -1644,6 +1681,9 @@ namespace Chaos::Private
 			const int32 EdgeIndex = Edge->IslandArrayIndex;
 			check(Island->ContainerEdges[ContainerIndex][EdgeIndex] == Edge);
 
+			UE_MT_SCOPED_WRITE_ACCESS(LocalEdgeAccessDetector);
+			UE_MT_SCOPED_WRITE_ACCESS(LocalIslandAccessDetector);
+
 			Island->ContainerEdges[ContainerIndex].RemoveAtSwap(EdgeIndex, EAllowShrinking::No);
 			if (EdgeIndex < Island->ContainerEdges[ContainerIndex].Num())
 			{
diff --git a/Engine/Source/Runtime/Experimental/Chaos/Public/Chaos/Island/IslandManager.h b/Engine/Source/Runtime/Experimental/Chaos/Public/Chaos/Island/IslandManager.h
--- a/Engine/Source/Runtime/Experimental/Chaos/Public/Chaos/Island/IslandManager.h
+++ b/Engine/Source/Runtime/Experimental/Chaos/Public/Chaos/Island/IslandManager.h
@@ -693,6 +693,23 @@ namespace Chaos::Private
 
 		// Whether we should assign levels (for shock propagation)
 		bool bAssignLevels = true;
+
+		// Access to FPBDIslandParticle::Edges
+		// Access to FPBDIslandParticle::IslandArrayIndex
+		UE_MT_DECLARE_RW_ACCESS_DETECTOR(LocalNodeAccessDetector);
+
+		// Access to FPBDIslandConstraint::Nodes[2], FPBDIslandConstraint::NodeArrayIndices[2]
+		// Access to FPBDIslandConstraint::ContainerIndex, FPBDIslandConstraint::IslandArrayIndex
+		UE_MT_DECLARE_RW_ACCESS_DETECTOR(LocalEdgeAccessDetector);
+
+		// Access to FPBDIsland::Nodes, FPBDIsland::ContainerEdges
+		UE_MT_DECLARE_RW_ACCESS_DETECTOR(LocalIslandAccessDetector);
+
+		// Access to FPBDIslandManager::Nodes
+		UE_MT_DECLARE_RW_ACCESS_DETECTOR(GlobalNodesAccessDetector);
+
+		// Access to FPBDIslandManager::Edges
+		UE_MT_DECLARE_RW_ACCESS_DETECTOR(GlobalEdgesAccessDetector);
 	};
 
 
