Geometry collection : add the possibilty to set the linear / angular damping  per transform

diff --git a/Engine/Source/Runtime/Experimental/Chaos/Private/PhysicsProxy/GeometryCollectionPhysicsProxy.cpp b/Engine/Source/Runtime/Experimental/Chaos/Private/PhysicsProxy/GeometryCollectionPhysicsProxy.cpp
--- a/Engine/Source/Runtime/Experimental/Chaos/Private/PhysicsProxy/GeometryCollectionPhysicsProxy.cpp
+++ b/Engine/Source/Runtime/Experimental/Chaos/Private/PhysicsProxy/GeometryCollectionPhysicsProxy.cpp
@@ -3792,6 +3792,30 @@ void FGeometryCollectionPhysicsProxy::SetLinearDamping_Internal(float InDamping)
 	}
 }
 
+void FGeometryCollectionPhysicsProxy::SetLinearDampingPerTransform_External(const TArray<int32>& TransformIndices, float InDamping)
+{
+	// update physics thread
+	ExecuteOnPhysicsThread(*this, [this, InDamping, TransformIndices]()
+		{
+			SetLinearDampingPerTransform_Internal(TransformIndices, InDamping);
+		});
+}
+
+void FGeometryCollectionPhysicsProxy::SetLinearDampingPerTransform_Internal(const TArray<int32>& TransformIndices, float InDamping)
+{
+	for (const int32 TransformIndex: TransformIndices)
+	{
+		if (FromTransformToParticleIndex.IsValidIndex(TransformIndex))
+		{
+			const int32 ParticleIndex = FromTransformToParticleIndex[TransformIndex];
+			if (Chaos::FPBDRigidClusteredParticleHandle* Handle = SolverParticleHandles.IsValidIndex(ParticleIndex) ? SolverParticleHandles[ParticleIndex] : nullptr)
+			{
+				Handle->SetLinearEtherDrag(InDamping);
+			}
+		}
+	}
+}
+
 void FGeometryCollectionPhysicsProxy::SetAngularDamping_External(float InDamping)
 {
 	// update physics thread
@@ -3814,6 +3838,30 @@ void FGeometryCollectionPhysicsProxy::SetAngularDamping_Internal(float InDamping
 	}
 }
 
+void FGeometryCollectionPhysicsProxy::SetAngularDampingPerTransform_External(const TArray<int32>& TransformIndices, float InDamping)
+{
+	// update physics thread
+	ExecuteOnPhysicsThread(*this, [this, InDamping, TransformIndices]()
+		{
+			SetAngularDampingPerTransform_Internal(TransformIndices, InDamping);
+		});
+}
+
+void FGeometryCollectionPhysicsProxy::SetAngularDampingPerTransform_Internal(const TArray<int32>& TransformIndices, float InDamping)
+{
+	for (const int32 TransformIndex : TransformIndices)
+	{
+		if (FromTransformToParticleIndex.IsValidIndex(TransformIndex))
+		{
+			const int32 ParticleIndex = FromTransformToParticleIndex[TransformIndex];
+			if (Chaos::FPBDRigidClusteredParticleHandle* Handle = SolverParticleHandles.IsValidIndex(ParticleIndex) ? SolverParticleHandles[ParticleIndex] : nullptr)
+			{
+				Handle->SetAngularEtherDrag(InDamping);
+			}
+		}
+	}
+}
+
 void FGeometryCollectionPhysicsProxy::PushStateOnGameThread(Chaos::FPBDRigidsSolver* InSolver)
 {
 	// CONTEXT: ANYTHREAD but spawned form GAMETHREAD in a parallelFor
diff --git a/Engine/Source/Runtime/Experimental/Chaos/Public/PhysicsProxy/GeometryCollectionPhysicsProxy.h b/Engine/Source/Runtime/Experimental/Chaos/Public/PhysicsProxy/GeometryCollectionPhysicsProxy.h
--- a/Engine/Source/Runtime/Experimental/Chaos/Public/PhysicsProxy/GeometryCollectionPhysicsProxy.h
+++ b/Engine/Source/Runtime/Experimental/Chaos/Public/PhysicsProxy/GeometryCollectionPhysicsProxy.h
@@ -488,7 +488,9 @@ public:
 	CHAOS_API void ApplyLinearVelocity_External(FGeometryCollectionItemIndex ItemIndex, const FVector& LinearVelocity);
 	CHAOS_API void ApplyAngularVelocity_External(FGeometryCollectionItemIndex ItemIndex, const FVector& AngularVelocity);
 	CHAOS_API void SetLinearDamping_External(float InDamping);
+	CHAOS_API void SetLinearDampingPerTransform_External(const TArray<int32>& TransformIndices, float InDamping);
 	CHAOS_API void SetAngularDamping_External(float InDamping);
+	CHAOS_API void SetAngularDampingPerTransform_External(const TArray<int32>& TransformIndices, float InDamping);
 
 	CHAOS_API void SetProxyDirty_External();
 
@@ -654,7 +656,9 @@ protected:
 	CHAOS_API void SetOneWayInteractionLevel_Internal(int32 InOneWayInteractionLevel);
 	CHAOS_API void SetPhysicsMaterial_Internal(const Chaos::FMaterialHandle& MaterialHandle);
 	CHAOS_API void SetLinearDamping_Internal(float InDamping);
+	CHAOS_API void SetLinearDampingPerTransform_Internal(const TArray<int32>& TransformIndices, float InDamping);
 	CHAOS_API void SetAngularDamping_Internal(float InDamping);
+	CHAOS_API void SetAngularDampingPerTransform_Internal(const TArray<int32>& TransformIndices ,float InDamping);
 
 private:
 
diff --git a/Engine/Source/Runtime/Experimental/GeometryCollectionEngine/Private/GeometryCollection/GeometryCollectionComponent.cpp b/Engine/Source/Runtime/Experimental/GeometryCollectionEngine/Private/GeometryCollection/GeometryCollectionComponent.cpp
--- a/Engine/Source/Runtime/Experimental/GeometryCollectionEngine/Private/GeometryCollection/GeometryCollectionComponent.cpp
+++ b/Engine/Source/Runtime/Experimental/GeometryCollectionEngine/Private/GeometryCollection/GeometryCollectionComponent.cpp
@@ -7337,6 +7337,14 @@ void UGeometryCollectionComponent::SetLinearDamping(float InDamping)
 	}
 }
 
+void UGeometryCollectionComponent::SetLinearDampingPerTransform(const TArray<int32>& TransformIndices, float InDamping)
+{
+	if (PhysicsProxy)
+	{
+		PhysicsProxy->SetLinearDampingPerTransform_External(TransformIndices, InDamping);
+	}
+}
+
 void UGeometryCollectionComponent::SetAngularDamping(float InDamping)
 {
 	Super::SetAngularDamping(InDamping);
@@ -7346,6 +7354,14 @@ void UGeometryCollectionComponent::SetAngularDamping(float InDamping)
 	}
 }
 
+void UGeometryCollectionComponent::SetAngularDampingPerTransform(const TArray<int32>& TransformIndices, float InDamping)
+{
+	if (PhysicsProxy)
+	{
+		PhysicsProxy->SetAngularDampingPerTransform_External(TransformIndices, InDamping);
+	}
+}
+
 void UGeometryCollectionComponent::FlushNetDormancyIfNeeded() const
 {
 	if (GetDesiredNetAwakeningMode() != ENetAwakeningMode::FlushNetDormancy)
diff --git a/Engine/Source/Runtime/Experimental/GeometryCollectionEngine/Public/GeometryCollection/GeometryCollectionComponent.h b/Engine/Source/Runtime/Experimental/GeometryCollectionEngine/Public/GeometryCollection/GeometryCollectionComponent.h
--- a/Engine/Source/Runtime/Experimental/GeometryCollectionEngine/Public/GeometryCollection/GeometryCollectionComponent.h
+++ b/Engine/Source/Runtime/Experimental/GeometryCollectionEngine/Public/GeometryCollection/GeometryCollectionComponent.h
@@ -673,6 +673,11 @@ protected:
 	GEOMETRYCOLLECTIONENGINE_API virtual bool CanBeUsedInPhysicsReplication(const FName BoneName = NAME_None) const override;
 	//~ End UPrimitiveComponent Interface.
 
+	UFUNCTION(BlueprintCallable, Category = "ChaosPhysics")
+	GEOMETRYCOLLECTIONENGINE_API void SetLinearDampingPerTransform(const TArray<int32>& TransformIndices, float InDamping);
+
+	UFUNCTION(BlueprintCallable, Category = "ChaosPhysics")
+	GEOMETRYCOLLECTIONENGINE_API void SetAngularDampingPerTransform(const TArray<int32>& TransformIndices, float InDamping);
 
 	//~ Begin UMeshComponent Interface.	
 public:
