[mutable-Dataflow] Added UStructTypePolicy

#rnx

diff --git a/Engine/Source/Runtime/Experimental/Dataflow/Core/Private/Dataflow/DataflowAnyType.cpp b/Engine/Source/Runtime/Experimental/Dataflow/Core/Private/Dataflow/DataflowAnyType.cpp
--- a/Engine/Source/Runtime/Experimental/Dataflow/Core/Private/Dataflow/DataflowAnyType.cpp
+++ b/Engine/Source/Runtime/Experimental/Dataflow/Core/Private/Dataflow/DataflowAnyType.cpp
@@ -13,6 +13,7 @@ namespace UE::Dataflow
 		UE_DATAFLOW_REGISTER_ANYTYPE(FDataflowAnyType);
 		UE_DATAFLOW_REGISTER_ANYTYPE(FDataflowAllTypes);
 		UE_DATAFLOW_REGISTER_ANYTYPE(FDataflowArrayTypes);
+		UE_DATAFLOW_REGISTER_ANYTYPE(FDataflowUStructTypes);
 		UE_DATAFLOW_REGISTER_ANYTYPE(FDataflowNumericTypes);
 		UE_DATAFLOW_REGISTER_ANYTYPE(FDataflowVectorTypes);
 		UE_DATAFLOW_REGISTER_ANYTYPE(FDataflowStringTypes);
diff --git a/Engine/Source/Runtime/Experimental/Dataflow/Core/Public/Dataflow/DataflowAnyType.h b/Engine/Source/Runtime/Experimental/Dataflow/Core/Public/Dataflow/DataflowAnyType.h
--- a/Engine/Source/Runtime/Experimental/Dataflow/Core/Public/Dataflow/DataflowAnyType.h
+++ b/Engine/Source/Runtime/Experimental/Dataflow/Core/Public/Dataflow/DataflowAnyType.h
@@ -37,6 +37,16 @@ struct FDataflowArrayTypes : public FDataflowAnyType
 	GENERATED_BODY()
 };
 
+/** UStruct based types */
+USTRUCT()
+struct FDataflowUStructTypes : public FDataflowAnyType
+{
+	using FPolicyType = FDataflowUStructTypePolicy;
+	using FStorageType = void;
+	GENERATED_BODY()
+};
+
+
 /**
 * Numeric types
 * (double, float, int64, uint64, int32, uint32, int16, uint16, int8, uint8)"
diff --git a/Engine/Source/Runtime/Experimental/Dataflow/Core/Public/Dataflow/DataflowTypePolicy.h b/Engine/Source/Runtime/Experimental/Dataflow/Core/Public/Dataflow/DataflowTypePolicy.h
--- a/Engine/Source/Runtime/Experimental/Dataflow/Core/Public/Dataflow/DataflowTypePolicy.h
+++ b/Engine/Source/Runtime/Experimental/Dataflow/Core/Public/Dataflow/DataflowTypePolicy.h
@@ -258,6 +258,57 @@ struct FDataflowArrayTypePolicy: public IDataflowTypePolicy
 	}
 };
 
+struct FDataflowUStructTypePolicy : IDataflowTypePolicy
+{
+	virtual bool SupportsType(FName InType) const override
+	{
+		return SupportsTypeStatic(InType);
+	}
+
+	static bool SupportsTypeStatic(FName InType)
+	{
+		FString InnerTypeStr;
+		if (GetStructType(InType.ToString(), InnerTypeStr))
+		{
+			if (StaticFindFirstObject(UStruct::StaticClass(), *InnerTypeStr, EFindFirstObjectOptions::NativeFirst))
+			{
+				return true;
+			}
+		}
+		
+		// not a proper Ustruct
+		return false;
+	}
+	
+	static IDataflowTypePolicy* GetInterface()
+	{
+		static FDataflowUStructTypePolicy Instance;
+		return &Instance;
+	}
+
+	/**
+	 * Provided an InTypeString outputs the same type without the type decorations (like prefixes).
+	 * @param InTypeStr The object type
+	 * @param OutInnerType The object type without the struct prefix ("F")
+	 * @return true if the prefix was found, false otherwise
+	 */
+	static bool GetStructType(const FString& InTypeStr, FString& OutInnerType)
+	{
+		static constexpr const TCHAR* ObjectPtrPrefix = TEXT("F");
+		static const int32 ObjectPtrPrefixLen = FCString::Strlen(ObjectPtrPrefix);;
+		if (InTypeStr.StartsWith(ObjectPtrPrefix))
+		{
+			OutInnerType = InTypeStr
+				.RightChop(ObjectPtrPrefixLen) // remove the type prefix
+				.TrimStartAndEnd();
+			return true;
+		}
+		return false;
+	}
+};
+
+
+
 struct FDataflowNumericTypePolicy : 
 	public TDataflowMultiTypePolicy<double, float, int64, uint64, int32, uint32, int16, uint16, int8, uint8>
 {
