[Physics] Move ProcessInputs_Internal to the solver update flow while keeping the callorder and ensure time is updated for the ISimCallbackObject before ProcessInputs_Internal and PostInitialize_Internal.

diff --git a/Engine/Source/Runtime/Experimental/Chaos/Private/Chaos/Framework/PhysicsSolverBase.cpp b/Engine/Source/Runtime/Experimental/Chaos/Private/Chaos/Framework/PhysicsSolverBase.cpp
--- a/Engine/Source/Runtime/Experimental/Chaos/Private/Chaos/Framework/PhysicsSolverBase.cpp
+++ b/Engine/Source/Runtime/Experimental/Chaos/Private/Chaos/Framework/PhysicsSolverBase.cpp
@@ -351,6 +351,44 @@ namespace Chaos
 		delete &InSolver;
 	}
 
+	void FPhysicsSolverBase::ApplyCallbacks_Internal()
+	{
+		QUICK_SCOPE_CYCLE_COUNTER(ApplySimCallbacks);
+
+		//if delta time is 0 we are flushing data, user callbacks should not be triggered because there is no sim
+		if (MLastDt > 0)
+		{
+			const FReal SimTime = GetSolverTime();
+
+			for (ISimCallbackObject* Callback : SimCallbackObjects)
+			{
+				if (Callback->HasOption(ESimCallbackOptions::RunOnFrozenGameThread) == bGameThreadFrozen)
+				{
+					Callback->SetSimAndDeltaTime_Internal(SimTime, MLastDt);
+				}
+			}
+
+			if (FPushPhysicsData* PushData = MarshallingManager.GetConsumerData_Internal())
+			{
+				if (MRewindCallback && !IsShuttingDown() && !bGameThreadFrozen)
+				{
+					QUICK_SCOPE_CYCLE_COUNTER(Chaos_RewindCallback_ProcessInputs_Internal);
+					MRewindCallback->ProcessInputs_Internal(PushData->InternalStep, PushData->SimCallbackInputs);
+				}
+			}
+
+			for (ISimCallbackObject* Callback : SimCallbackObjects)
+			{
+				if (Callback->HasOption(ESimCallbackOptions::RunOnFrozenGameThread) == bGameThreadFrozen)
+				{
+					FScopedTraceSolverCallback TraceCallback(Callback);
+
+					Callback->PreSimulate_Internal();
+				}
+			}
+		}
+	}
+
 	void FPhysicsSolverBase::UpdateParticleInAccelerationStructure_External(FGeometryParticle* Particle, EPendingSpatialDataOperation InOperation)
 	{
 		//mark it as pending for async structure being built
diff --git a/Engine/Source/Runtime/Experimental/Chaos/Private/PBDRigidsSolver.cpp b/Engine/Source/Runtime/Experimental/Chaos/Private/PBDRigidsSolver.cpp
--- a/Engine/Source/Runtime/Experimental/Chaos/Private/PBDRigidsSolver.cpp
+++ b/Engine/Source/Runtime/Experimental/Chaos/Private/PBDRigidsSolver.cpp
@@ -2051,25 +2051,24 @@ namespace Chaos
 		//so for example if the user deletes a floor surface, we want all sub-steps to see that in the same way
 		//also note, the commands run after data is marshalled over. This is important because data marshalling ensures any GT property changes are seen by command
 		//for example a particle may not be created until marshalling occurs, and then a command could explicitly modify something like a collision setting
-
-		// If we are not resimulating, insert scheduled commands for this step into the current PushData to be executed
-		if (!IsResimming())
 		{
-			TArray<FSimCallbackCommandObject*>& ScheduledSimCommands = MarshallingManager.GetSimCommandsScheduled_Internal();
-			for (int32 Idx = ScheduledSimCommands.Num() - 1; Idx >= 0; Idx--)
+			QUICK_SCOPE_CYCLE_COUNTER(Chaos_ProcessPushedData_SimCommands);
+			// If we are not resimulating, insert scheduled commands for this step into the current PushData to be executed
+			if (!IsResimming())
 			{
-				FSimCallbackCommandObject* SimCallbackObject = ScheduledSimCommands[Idx];
-				if (SimCallbackObject->ShouldExecute(CurrentFrame, /*bIsResim*/false))
+				TArray<FSimCallbackCommandObject*>& ScheduledSimCommands = MarshallingManager.GetSimCommandsScheduled_Internal();
+				for (int32 Idx = ScheduledSimCommands.Num() - 1; Idx >= 0; Idx--)
 				{
-					// If this scheduled command should get processed this step, inject it into the start of PushData:SimCommands to preserve call order and to cache the SimCommand in rewind history via PushData
-					PushData.SimCommands.Insert(SimCallbackObject, 0);
-					ScheduledSimCommands.RemoveAt(Idx, EAllowShrinking::No);
+					FSimCallbackCommandObject* SimCallbackObject = ScheduledSimCommands[Idx];
+					if (SimCallbackObject->ShouldExecute(CurrentFrame, /*bIsResim*/false))
+					{
+						// If this scheduled command should get processed this step, inject it into the start of PushData:SimCommands to preserve call order and to cache the SimCommand in rewind history via PushData
+						PushData.SimCommands.Insert(SimCallbackObject, 0);
+						ScheduledSimCommands.RemoveAt(Idx, EAllowShrinking::No);
+					}
 				}
 			}
-		}
 
-		{
-			QUICK_SCOPE_CYCLE_COUNTER(Chaos_ProcessPushedData_SimCommands);
 			for (FSimCallbackCommandObject*& SimCallbackObject : PushData.SimCommands)
 			{
 				if (!SimCallbackObject)
@@ -2090,16 +2089,14 @@ namespace Chaos
 					SimCallbackObject = nullptr; // Stop tracking the command in this PushData
 				}
 			}
-		
-			for (ISimCallbackObject* SimCallbackObject : PushData.SimCallbackObjectsToAdd)
-			{
-				SimCallbackObject->PostInitialize_Internal();
-			}
+		}
 
-			if (MRewindCallback && !IsShuttingDown())
+		for (ISimCallbackObject* SimCallbackObject : PushData.SimCallbackObjectsToAdd)
+		{
+			if (SimCallbackObject->HasOption(ESimCallbackOptions::RunOnFrozenGameThread) == bGameThreadFrozen)
 			{
-				QUICK_SCOPE_CYCLE_COUNTER(Chaos_RewindCallback_ProcessInputs_Internal);
-				MRewindCallback->ProcessInputs_Internal(GetCurrentFrame(), PushData.SimCallbackInputs);
+				SimCallbackObject->SetSimAndDeltaTime_Internal(GetSolverTime(), MLastDt);
+				SimCallbackObject->PostInitialize_Internal();
 			}
 		}
 	}
@@ -2194,12 +2191,6 @@ namespace Chaos
 									SCOPED_NAMED_EVENT(ChaosRewindAndResim_ResetCollisions, FColor::Red);
 									MTime = PushData->StartTime;	//not sure if sub-steps have proper StartTime so just do this once and let solver evolve remaining time
 
-									// @todo(chaos) : replace this temporary fix with a better method of syncing sim times prior to a resimulation
-									for (ISimCallbackObject* Callback : SimCallbackObjects)
-									{
-										Callback->SetSimAndDeltaTime_Internal(MTime, MLastDt);
-									}
-
 									GetEvolution()->ResetCollisions();
 								}
 
diff --git a/Engine/Source/Runtime/Experimental/Chaos/Public/Chaos/Framework/PhysicsSolverBase.h b/Engine/Source/Runtime/Experimental/Chaos/Public/Chaos/Framework/PhysicsSolverBase.h
--- a/Engine/Source/Runtime/Experimental/Chaos/Public/Chaos/Framework/PhysicsSolverBase.h
+++ b/Engine/Source/Runtime/Experimental/Chaos/Public/Chaos/Framework/PhysicsSolverBase.h
@@ -634,26 +634,7 @@ namespace Chaos
 			bGameThreadFrozen = InGameThreadFrozen;
 		}
 
-		virtual void ApplyCallbacks_Internal()
-		{
-			QUICK_SCOPE_CYCLE_COUNTER(ApplySimCallbacks);
-
-			//if delta time is 0 we are flushing data, user callbacks should not be triggered because there is no sim
-			if(MLastDt > 0)
-			{
-				const FReal SimTime = GetSolverTime();
-				for (ISimCallbackObject* Callback : SimCallbackObjects)
-				{
-					if (Callback->HasOption(ESimCallbackOptions::RunOnFrozenGameThread) == bGameThreadFrozen)
-					{
-						FScopedTraceSolverCallback TraceCallback(Callback);
-
-						Callback->SetSimAndDeltaTime_Internal(SimTime, MLastDt);
-						Callback->PreSimulate_Internal();
-					}
-				}
-			}
-		}
+		virtual void ApplyCallbacks_Internal();
 
 		void FinalizeCallbackData_Internal()
 		{
