iax: Allow for testing higher request queue lengths.

#rnx

diff --git a/Engine/Source/Runtime/Experimental/IoStore/OnDemand/Private/OnDemandHttpIoDispacher.cpp b/Engine/Source/Runtime/Experimental/IoStore/OnDemand/Private/OnDemandHttpIoDispacher.cpp
--- a/Engine/Source/Runtime/Experimental/IoStore/OnDemand/Private/OnDemandHttpIoDispacher.cpp
+++ b/Engine/Source/Runtime/Experimental/IoStore/OnDemand/Private/OnDemandHttpIoDispacher.cpp
@@ -1363,8 +1363,7 @@ bool FHttpDispatcher::ProcessHttpRequests()
 
 	RecreateHttpClientIfNeeded();
 
-	int32		QueueLength = CalculateRequestQueueLengthMax();
-	const int32 MaxConcurrentRequests = FMath::Min(QueueLength, 32);
+	const int32 MaxConcurrentRequests = CalculateRequestQueueLengthMax();
 	int32		NumConcurrentRequests = 0;
 
 	FHttpRequest* NextHttpRequest = DequeueHttpRequests();
diff --git a/Engine/Source/Runtime/Experimental/IoStore/OnDemand/Private/OnDemandHttpThread.cpp b/Engine/Source/Runtime/Experimental/IoStore/OnDemand/Private/OnDemandHttpThread.cpp
--- a/Engine/Source/Runtime/Experimental/IoStore/OnDemand/Private/OnDemandHttpThread.cpp
+++ b/Engine/Source/Runtime/Experimental/IoStore/OnDemand/Private/OnDemandHttpThread.cpp
@@ -207,7 +207,8 @@ int32 CalculateRequestQueueLengthMax()
 	// DOES NOT have any effect on throughput throttling.
 	int32 ConnNum = (GIaxHttpVersion == 2) ? 1 : GIasHttpConnectionCount;
 	int32 InflightNum = (GIaxHttpVersion == 2) ? GIaxHttpMaxInflight : 1;
-	return FMath::Max((ConnNum * InflightNum) + 2, 1);
+	int32 Value = FMath::Max((ConnNum * InflightNum) + 2, 1);
+	return FMath::Min(Value, int32(HTTP::FEventLoop::MaxActiveTickets));
 }
 
 const int32 GOnDemandBackendThreadPriorities[5] =
@@ -483,9 +484,7 @@ void FOnDemandHttpThread::TickRequests()
 		bRecreateHttpClient = false;
 	}
 
-	// We can't limit cvars so add a hard coded cap to prevent concurrent requests from getting out of hand.
-	int32 QueueLength = CalculateRequestQueueLengthMax();
-	const int32 MaxConcurrentRequests = FMath::Min(QueueLength, 32);
+	const int32 MaxConcurrentRequests = CalculateRequestQueueLengthMax();
 
 	int32 NumConcurrentRequests = 0;
 	FHttpRequest* NextHttpRequest = HttpRequests.Dequeue(MaxConcurrentRequests);
