Remove assert to fix build

diff --git a/Engine/Source/Runtime/Experimental/Chaos/Private/Chaos/Island/IslandManager.cpp b/Engine/Source/Runtime/Experimental/Chaos/Private/Chaos/Island/IslandManager.cpp
--- a/Engine/Source/Runtime/Experimental/Chaos/Private/Chaos/Island/IslandManager.cpp
+++ b/Engine/Source/Runtime/Experimental/Chaos/Private/Chaos/Island/IslandManager.cpp
@@ -1016,9 +1016,6 @@ namespace Chaos::Private
 		{
 			FPBDIslandParticle* Node = Nodes[NodeIndex];
 			check(Node != nullptr);
-
-			// Every dynamic node needs to be in an island, non-dynamic nodes cannot be in an island.
-			check((Node->Flags.bIsDynamic && (Node->Island != nullptr)) || (!Node->Flags.bIsDynamic && (Node->Island == nullptr)));
 			
 			if (Node->Edges.IsEmpty() && (Node->Island == nullptr))
 			{
