[ChaosVD] Modified FChaosArchive to have its Context be shareable, and added a lock so the chaos visual debugger can access it from multiple threads (this lock is limited to the CVD usage currently)

This allows us to reduce the amount of duplicated geometry we serialize, which can result on huge file size reductions in CVD recordings of unoptimized maps with lots of tri mesh collision geometry. The downside of this change is that, in essence, makes geometry serialization a serial operation, but the reduction on amount of data we serialize makes up for it.

The geometry hash cache we use to know if we already serialized a specific geometry instance is still multi threaded, therefore we only pay the serial cost when we first see and fully serialize a specific geometry instance. From that point onward, checking if we already serialize a geometry instance remains a parallel and quick operation.

The way it works is the following:

FChaosArchive uses its context object to tag pointers the first time it sees them and serializes them. If during subsequent serialization with the same archive a ptr is seen again, just its tag is serialized not the contents (hence the savings). With this change CVD can keep a shared context object that all ChaosArchives used during the recordings will use, so we only serialize the content of a ptr once. The limitation of this system is that it assumes all serialization happens serially, therefore we need to make serialization serial because if two threads try to serialize the same object at the same time, the de-serialization order might not be the same, which will result in a crash because the de-serializer might de serialize a tag when it should have been de-serialized the contents.

There are ways to solve this and keep the parallelism, but that would have required changes to the chaos archive itself, potentially making it slower in non CVD scenarios. Also, all collision geometry serialization CVD uses is shared with the editor/engine asset serialization, therefore any change on it is risky.

As mentioned earlier, initial test show that although we are forcing part of the process to be serial, there is still an overall improvement (specially in complex scenes) because we are doing way less work overall. Serializing a tag is just serializing an integer

diff --git a/Engine/Plugins/ChaosVD/Source/ChaosVD/Private/Trace/ChaosVDTraceAnalyzer.cpp b/Engine/Plugins/ChaosVD/Source/ChaosVD/Private/Trace/ChaosVDTraceAnalyzer.cpp
--- a/Engine/Plugins/ChaosVD/Source/ChaosVD/Private/Trace/ChaosVDTraceAnalyzer.cpp
+++ b/Engine/Plugins/ChaosVD/Source/ChaosVD/Private/Trace/ChaosVDTraceAnalyzer.cpp
@@ -28,6 +28,8 @@ void FChaosVDTraceAnalyzer::OnAnalysisBegin(const FOnAnalysisContext& Context)
 	Builder.RouteEvent(RouteId_ChaosVDNetworkTickOffset, "ChaosVDLogger", "ChaosVDNetworkTickOffset");
 	Builder.RouteEvent(RouteId_ChaosVDRolledBackDataID, "ChaosVDLogger", "ChaosVDRolledBackDataID");
 	Builder.RouteEvent(RouteId_ChaosVDUsesAutoRTFM, "ChaosVDLogger", "ChaosVDUsesAutoRTFM");
+	
+	Builder.RouteEvent(RouteId_ChaosVDMiscHeaderData, "ChaosVDLogger", "ChaosVDMiscHeaderData");
 
 	Builder.RouteEvent(RouteId_BeginFrame, "Misc", "BeginFrame");
 	Builder.RouteEvent(RouteId_EndFrame, "Misc", "EndFrame");
@@ -335,7 +337,14 @@ bool FChaosVDTraceAnalyzer::OnEvent(uint16 RouteId, EStyle Style, const FOnEvent
 
 			break;
 		}
-
+	case RouteId_ChaosVDMiscHeaderData:
+		{
+			if (EventData.GetValue<bool>("bUsesSharedGeometryContextArchive"))
+			{
+				ChaosVDTraceProvider->CreateGeometryArchiveContext();
+			}
+			break;
+		}
 	default:
 		break;
 	}
diff --git a/Engine/Plugins/ChaosVD/Source/ChaosVD/Private/Trace/ChaosVDTraceAnalyzer.h b/Engine/Plugins/ChaosVD/Source/ChaosVD/Private/Trace/ChaosVDTraceAnalyzer.h
--- a/Engine/Plugins/ChaosVD/Source/ChaosVD/Private/Trace/ChaosVDTraceAnalyzer.h
+++ b/Engine/Plugins/ChaosVD/Source/ChaosVD/Private/Trace/ChaosVDTraceAnalyzer.h
@@ -55,6 +55,7 @@ private:
 		RouteId_ChaosVDNetworkTickOffset,
 		RouteId_ChaosVDRolledBackDataID,
 		RouteId_ChaosVDUsesAutoRTFM,
+		RouteId_ChaosVDMiscHeaderData,
 	};
 
 	TraceServices::IAnalysisSession& Session;
diff --git a/Engine/Plugins/ChaosVD/Source/ChaosVD/Private/Trace/ChaosVDTraceProvider.cpp b/Engine/Plugins/ChaosVD/Source/ChaosVD/Private/Trace/ChaosVDTraceProvider.cpp
--- a/Engine/Plugins/ChaosVD/Source/ChaosVD/Private/Trace/ChaosVDTraceProvider.cpp
+++ b/Engine/Plugins/ChaosVD/Source/ChaosVD/Private/Trace/ChaosVDTraceProvider.cpp
@@ -181,6 +181,11 @@ TSharedPtr<FChaosVDParticleMetadata> FChaosVDTraceProvider::GetParticleMetadata(
 	return MetadataInstance ? *MetadataInstance : nullptr;
 }
 
+void FChaosVDTraceProvider::CreateGeometryArchiveContext()
+{
+	GeometryArchiveContext = MakeShared<Chaos::FChaosArchiveContext>();
+}
+
 void FChaosVDTraceProvider::CommitProcessedGameFramesToRecording()
 {
 	TArray<int32, TInlineAllocator<16>> SolverIDs;
diff --git a/Engine/Plugins/ChaosVD/Source/ChaosVD/Public/Trace/ChaosVDTraceProvider.h b/Engine/Plugins/ChaosVD/Source/ChaosVD/Public/Trace/ChaosVDTraceProvider.h
--- a/Engine/Plugins/ChaosVD/Source/ChaosVD/Public/Trace/ChaosVDTraceProvider.h
+++ b/Engine/Plugins/ChaosVD/Source/ChaosVD/Public/Trace/ChaosVDTraceProvider.h
@@ -10,6 +10,11 @@
 #include "Templates/SharedPointer.h"
 #include "TraceServices/Model/AnalysisSession.h"
 
+namespace Chaos
+{
+	class FChaosArchiveContext;
+}
+
 class FChaosVDDataProcessorBase;
 
 namespace Chaos::VisualDebugger
@@ -173,6 +178,13 @@ public:
 	void AddParticleMetadata(uint64 MetadaId, const TSharedPtr<FChaosVDParticleMetadata>& InMetadata);
 	TSharedPtr<FChaosVDParticleMetadata> GetParticleMetadata(uint64 MetadataId);
 
+	void CreateGeometryArchiveContext();
+
+	TSharedPtr<Chaos::FChaosArchiveContext> GetGeometryArchiveContext()
+	{
+		return GeometryArchiveContext;
+	}
+
 private:
 	
 	int32 GetCurrentGameThreadTrackID() const
@@ -228,6 +240,8 @@ private:
 
 	TMap<uint64, TSharedPtr<FChaosVDParticleMetadata>> SerializedParticleMetadata;
 
+	TSharedPtr<Chaos::FChaosArchiveContext> GeometryArchiveContext;
+
 	friend class FChaosVDTraceAnalyzer;
 
 	// Not idea, but needed to for the serialization fix up logic of the collision data channel info
diff --git a/Engine/Plugins/ChaosVD/Source/ChaosVD/Public/Trace/DataProcessors/ChaosVDDataProcessorBase.h b/Engine/Plugins/ChaosVD/Source/ChaosVD/Public/Trace/DataProcessors/ChaosVDDataProcessorBase.h
--- a/Engine/Plugins/ChaosVD/Source/ChaosVD/Public/Trace/DataProcessors/ChaosVDDataProcessorBase.h
+++ b/Engine/Plugins/ChaosVD/Source/ChaosVD/Public/Trace/DataProcessors/ChaosVDDataProcessorBase.h
@@ -48,8 +48,16 @@ namespace Chaos::VisualDebugger
 		// Note: I don't expect we will need a proxy archive for other types, but if we end up in that situation, we should use to switch to use traits 
 		if constexpr (std::is_same_v<TDataToSerialize, FChaosVDImplicitObjectDataWrapper<Chaos::FImplicitObjectPtr, Chaos::FChaosArchive>>)
 		{
-			FChaosArchive Ar(MemReader);
-			bSuccess = Data.Serialize(Ar);
+			if (TSharedPtr<FChaosArchiveContext> ProvidedGeometryContext = DataProvider->GetGeometryArchiveContext())
+			{
+				FChaosVDSharedContextArchive Ar(MemReader, ProvidedGeometryContext.ToSharedRef());
+				bSuccess = Data.Serialize(Ar);
+			}
+			else
+			{
+				FChaosArchive Ar(MemReader);
+				bSuccess = Data.Serialize(Ar);
+			}
 		}
 		else
 		{
diff --git a/Engine/Source/Runtime/Engine/Private/PhysicsEngine/PhysUtils.cpp b/Engine/Source/Runtime/Engine/Private/PhysicsEngine/PhysUtils.cpp
--- a/Engine/Source/Runtime/Engine/Private/PhysicsEngine/PhysUtils.cpp
+++ b/Engine/Source/Runtime/Engine/Private/PhysicsEngine/PhysUtils.cpp
@@ -4,6 +4,7 @@
 
 #include "Engine/World.h"
 #include "Model.h"
+#include "PBDRigidsSolver.h"
 #include "Physics/Experimental/PhysScene_Chaos.h"
 #include "PhysicsEngine/ConvexElem.h"
 #include "PhysicsEngine/BodySetup.h"
@@ -354,6 +355,20 @@ bool FPhysicsInterface::ExecPhysCommands(const TCHAR* Cmd, FOutputDevice* Output
 #if CHAOS_MEMORY_TRACKING
 	if (FParse::Command(&Cmd, TEXT("ChaosMemoryDistribution")))
 	{
+		class FChaosArchiveExternalContext : public Chaos::FChaosArchive
+		{
+		public:
+			FChaosArchiveExternalContext(FArchive& ArIn)
+				: FChaosArchive(ArIn)
+			{
+			}
+
+			explicit FChaosArchiveExternalContext(FArchive& ArIn, const TSharedRef<Chaos::FChaosArchiveContext>& InExternalSharedContext)
+				: FChaosArchive(ArIn, InExternalSharedContext)
+			{
+			}
+		};
+
 		//
 		// NOTE: This is an awful, awful way to do this.
 		//
@@ -363,19 +378,19 @@ bool FPhysicsInterface::ExecPhysCommands(const TCHAR* Cmd, FOutputDevice* Output
 		FArchive BaseAr;
 		BaseAr.SetIsLoading(false);
 		BaseAr.SetIsSaving(true);
-		Chaos::FChaosArchive Ar(BaseAr);
+		TSharedRef<Chaos::FChaosArchiveContext> ArchiveContext = MakeShared<Chaos::FChaosArchiveContext>();
+		FChaosArchiveExternalContext Ar(BaseAr, ArchiveContext);
 		FPhysScene* PhysScene = InWorld->GetPhysicsScene();
 		Chaos::FPhysicsSolver* Solver = PhysScene->GetSolver();
-		auto* Evolution = Solver->GetEvolution();
+		Chaos::FPBDRigidsEvolution* Evolution = Solver->GetEvolution();
 		Evolution->Serialize(Ar);
-		TUniquePtr<Chaos::FChaosArchiveContext> ArchiveContext = Ar.StealContext();
 
 		// Grab the memory tracking map from the archive context
 		const TMap<FName, Chaos::FChaosArchiveSectionData>& MemoryMap = ArchiveContext->SectionMap;
 
 		OutputDevice->Logf(TEXT("Chaos serialized memory distribution:"));
 		int64 TotalBytes = 0;
-		for (auto It : MemoryMap)
+		for (const TPair<FName, Chaos::FChaosArchiveSectionData>& It : MemoryMap)
 		{
 			const FName SectionName = It.Key;
 			const Chaos::FChaosArchiveSectionData SectionData = It.Value;
diff --git a/Engine/Source/Runtime/Experimental/Chaos/Private/ChaosVisualDebugger/ChaosVisualDebuggerTrace.cpp b/Engine/Source/Runtime/Experimental/Chaos/Private/ChaosVisualDebugger/ChaosVisualDebuggerTrace.cpp
--- a/Engine/Source/Runtime/Experimental/Chaos/Private/ChaosVisualDebugger/ChaosVisualDebuggerTrace.cpp
+++ b/Engine/Source/Runtime/Experimental/Chaos/Private/ChaosVisualDebugger/ChaosVisualDebuggerTrace.cpp
@@ -28,7 +28,6 @@
 #include "DataWrappers/ChaosVDQueryDataWrappers.h"
 #include "HAL/CriticalSection.h"
 #include "Misc/ScopeRWLock.h"
-#include "Misc/TransactionallySafeRWLock.h"
 #include "Templates/UniquePtr.h"
 
 static_assert(CHAOS_VISUAL_DEBUGGER_WITHOUT_TRACE || UE_TRACE_MINIMAL_ENABLED || UE_TRACE_ENABLED, "Chaos Visual Debugger support requires Trace or Trace minimal support. Please enable trace support with the Unreal Build Tool flag bEnableTrace, disable cvd support with bCompileChaosVisualDebuggerSupport or allows CVD support without trace by defining CHAOS_VISUAL_DEBUGGER_WITHOUT_TRACE to 1");
@@ -52,6 +51,7 @@ UE_TRACE_MINIMAL_EVENT_DEFINE(ChaosVDLogger, ChaosVDNonSolverTransform)
 UE_TRACE_MINIMAL_EVENT_DEFINE(ChaosVDLogger, ChaosVDNetworkTickOffset)
 UE_TRACE_MINIMAL_EVENT_DEFINE(ChaosVDLogger, ChaosVDRolledBackDataID)
 UE_TRACE_MINIMAL_EVENT_DEFINE(ChaosVDLogger, ChaosVDUsesAutoRTFM)
+UE_TRACE_MINIMAL_EVENT_DEFINE(ChaosVDLogger, ChaosVDMiscHeaderData)
 UE_TRACE_MINIMAL_EVENT_DEFINE(ChaosVDLogger, ChaosVDTraceRelevancyVolume)
 
 namespace Chaos::VisualDebugger::Cvars
@@ -66,14 +66,20 @@ namespace Chaos::VisualDebugger::Cvars
 	FAutoConsoleVariableRef CVarChaosVDCompressionMode(
 	TEXT("p.Chaos.VD.CompressionMode"),
 	CompressionMode,
-	TEXT("Oodle compression mode to use, 4 is by default which equsals to ECompressionLevel::VeryFast"));
+	TEXT("Oodle compression mode to use, 4 is by default which equals to ECompressionLevel::VeryFast"));
+
+	static bool bUsesSharedGeometryArchiveContext = true;
+	FAutoConsoleVariableRef CVarChaosVDUsesSharedGeometryArchiveContext(
+	TEXT("p.Chaos.VD.UseSharedGeometryArchiveContext"),
+	bUsesSharedGeometryArchiveContext,
+	TEXT("If set to true, CVD will serialize collision geometry using a shared archive context to de-duplicate geometry during serialization instead of during deserialization"));
 }
 
 namespace Chaos::VD
 {
 	FRecordingSessionState::FRecordingSessionState() : CVDNameTable(MakeShared<FChaosVDSerializableNameTable>()),
 	TraceRelevancyVolume(FBox(EForceInit::ForceInitToZero)),
-	bIsTracing(false)
+	bIsTracing(false), ArchiveContext(MakeShared<FChaosArchiveContext>())
 	{
 	}
 
@@ -839,9 +845,27 @@ uint32 FChaosVisualDebuggerTrace::TraceImplicitObject(Chaos::FImplicitObject* Ge
 	GeometryDataWrapper.ImplicitObject = GeometryPtr;
 
 	FChaosVDScopedTLSBufferAccessor TLSDataBuffer;
-	Chaos::VisualDebugger::WriteDataToBuffer<FChaosVDImplicitObjectWrapper, Chaos::FChaosArchive>(TLSDataBuffer.BufferRef, GeometryDataWrapper);
 
-	TraceBinaryData(TLSDataBuffer.BufferRef, FChaosVDImplicitObjectWrapper::WrapperTypeName);
+	if (SessionState.UsingSharedGeometryArchiveContext)
+	{
+		// When we use the shared archive, geometry serialization happens essentially in serial
+		// And we need to make sure we unlock the context only after the geometry data was written to the trace stream, otherwise
+		// during deserialization we might try to deserialize it before it is loaded. The reason is that if we read an object tag in the context but we don't have the object
+		// Chaos Archive context assumes the data that follows is the serialized object.
+
+		// The perf impact of doing this serially is mitigated by serializing way less data (as everything that is instanced is serialized only once)
+		// and by the fact whe only attempt to serialize a ptr once. Every subsequent attempt earlies out as this geometry hash will be already cached in TracedGeometryHashCache
+		// and that cache can be accessed in parallel
+
+		UE::TScopeLock ScopeLock(SessionState.ArchiveContextLock);
+		Chaos::VisualDebugger::WriteDataToBuffer<FChaosVDImplicitObjectWrapper>(TLSDataBuffer.BufferRef, GeometryDataWrapper, SessionState.ArchiveContext);
+		TraceBinaryData(TLSDataBuffer.BufferRef, FChaosVDImplicitObjectWrapper::WrapperTypeName);
+	}
+	else
+	{
+		Chaos::VisualDebugger::WriteDataToBuffer<FChaosVDImplicitObjectWrapper, Chaos::FChaosArchive>(TLSDataBuffer.BufferRef, GeometryDataWrapper);
+		TraceBinaryData(TLSDataBuffer.BufferRef, FChaosVDImplicitObjectWrapper::WrapperTypeName);
+	}
 
 	return GeometryHash;
 }
@@ -892,6 +916,13 @@ void FChaosVisualDebuggerTrace::InvalidateGeometryFromCache(const Chaos::FImplic
 		return;
 	}
 
+	{
+		UE::TScopeLock ContextScopeLock(SessionState.ArchiveContextLock);
+
+		// TODO: Should we walk the object hierarchy and remove all inner objects?
+		SessionState.ArchiveContext->ObjToTag.Remove(CachedGeometryToInvalidate);
+	}
+	
 	SessionState.TracedGeometryHashCache.RemoveCachedTraceObjectID(CachedGeometryToInvalidate);
 }
 
@@ -1362,6 +1393,8 @@ void FChaosVisualDebuggerTrace::Reset()
 
 	SessionState.TracedGeometryHashCache.Reset();
 	SessionState.ParticleMetadataIDsCache.Reset();
+	SessionState.ArchiveContext = MakeShared<Chaos::FChaosArchiveContext>();
+	SessionState.UsingSharedGeometryArchiveContext = Chaos::VisualDebugger::Cvars::bUsesSharedGeometryArchiveContext;
 }
 
 void FChaosVisualDebuggerTrace::HandleRecordingStop()
@@ -1436,6 +1469,9 @@ void FChaosVisualDebuggerTrace::HandleRecordingStart()
 	UE_TRACE_MINIMAL_LOG(ChaosVDLogger, ChaosVDUsesAutoRTFM, ChaosVDChannel)
 			<< ChaosVDUsesAutoRTFM.bUsingAutoRTFM(static_cast<bool>(UE_AUTORTFM));
 
+	UE_TRACE_MINIMAL_LOG(ChaosVDLogger, ChaosVDMiscHeaderData, ChaosVDChannel)
+		<< ChaosVDMiscHeaderData.bUsesSharedGeometryContextArchive(SessionState.UsingSharedGeometryArchiveContext);
+
 	SessionState.bIsTracing = true;
 }
 
diff --git a/Engine/Source/Runtime/Experimental/Chaos/Public/Chaos/ChaosArchive.h b/Engine/Source/Runtime/Experimental/Chaos/Public/Chaos/ChaosArchive.h
--- a/Engine/Source/Runtime/Experimental/Chaos/Public/Chaos/ChaosArchive.h
+++ b/Engine/Source/Runtime/Experimental/Chaos/Public/Chaos/ChaosArchive.h
@@ -9,6 +9,8 @@
 #include "UObject/DestructionObjectVersion.h"
 #include "Templates/Models.h"
 
+class FChaosVisualDebuggerTrace;
+
 namespace Chaos
 {
 
@@ -44,13 +46,13 @@ public:
 	TMap<FName, FChaosArchiveSectionData> SectionMap;
 	TArray<FChaosArchiveSection> SectionStack;
 
-	void PushSection(const FName& SectionName, const int64 MemoryLocation, const bool bAbsorbChildren);
+	CHAOS_API void PushSection(const FName& SectionName, const int64 MemoryLocation, const bool bAbsorbChildren);
 
-	void PopSection(const int64 MemoryLocation);
+	CHAOS_API void PopSection(const int64 MemoryLocation);
 
-	void BeginSerialize(const int64 MemoryLocation);
+	CHAOS_API void BeginSerialize(const int64 MemoryLocation);
 
-	void EndSerialize(const int64 MemoryLocation);
+	CHAOS_API void EndSerialize(const int64 MemoryLocation);
 
 private:
 
@@ -168,9 +170,18 @@ class FChaosArchive : public FArchiveProxy
 public:
 	FChaosArchive(FArchive& ArIn)
 		: FArchiveProxy(ArIn)
-		, Context(MakeUnique<FChaosArchiveContext>())
+		, Context(MakeShared<FChaosArchiveContext>())
+	{
+	}
+
+protected:
+	FChaosArchive(FArchive& ArIn, const TSharedRef<FChaosArchiveContext>& InExternalSharedContext)
+	: FArchiveProxy(ArIn)
+	, Context(InExternalSharedContext)
 	{
 	}
+	
+public:
 
 	template <typename T>
 	void SerializePtr(TSerializablePtr<T>& Obj)
@@ -308,21 +319,21 @@ public:
 #if CHAOS_MEMORY_TRACKING
 	virtual void Serialize(void* V, int64 Length) override
 	{
-		if (Context) { Context->BeginSerialize(Tell()); }
+		Context->BeginSerialize(Tell());
 		FArchiveProxy::Serialize(V, Length);
-		if (Context) { Context->EndSerialize(Tell()); }
+		Context->EndSerialize(Tell());
 	}
 #endif
 
+	UE_DEPRECATED(5.8, "This method is no longer supported and there is no replacement")
 	void SetContext(TUniquePtr<FChaosArchiveContext>&& InContext)
 	{
-		Context = MoveTemp(InContext);
 	}
 
+	UE_DEPRECATED(5.8, "This method is no longer supported and there is no replacement")
 	TUniquePtr<FChaosArchiveContext> StealContext()
 	{
-		TUniquePtr<FChaosArchiveContext> Ret = MoveTemp(Context);
-		return Ret;
+		return MakeUnique<FChaosArchiveContext>();
 	}
 
 private:
@@ -353,7 +364,7 @@ private:
 		RawPtr->Serialize(Ar);
 	}
 
-	TUniquePtr<FChaosArchiveContext> Context;
+	TSharedRef<FChaosArchiveContext> Context;
 
 #if CHAOS_MEMORY_TRACKING
 	friend class FChaosArchiveScopedMemory;
@@ -366,9 +377,14 @@ public:
 	FChaosArchiveScopedMemory(FChaosArchive& ArIn, const FName& SectionName, const bool bAbsorbChildren = true)
 #if CHAOS_MEMORY_TRACKING
 	: Ar(ArIn)
-	{ if (Ar.Context) { Ar.Context->PushSection(SectionName, Ar.Tell(), bAbsorbChildren); } }
+	{
+		Ar.Context->PushSection(SectionName, Ar.Tell(), bAbsorbChildren);
+	}
+
 	~FChaosArchiveScopedMemory()
-	{ if (Ar.Context) { Ar.Context->PopSection(Ar.Tell()); } }
+	{
+		Ar.Context->PopSection(Ar.Tell());
+	}
 private:
 	FChaosArchive& Ar;
 #else
diff --git a/Engine/Source/Runtime/Experimental/Chaos/Public/ChaosVisualDebugger/ChaosVDMemWriterReader.h b/Engine/Source/Runtime/Experimental/Chaos/Public/ChaosVisualDebugger/ChaosVDMemWriterReader.h
--- a/Engine/Source/Runtime/Experimental/Chaos/Public/ChaosVisualDebugger/ChaosVDMemWriterReader.h
+++ b/Engine/Source/Runtime/Experimental/Chaos/Public/ChaosVisualDebugger/ChaosVDMemWriterReader.h
@@ -2,6 +2,7 @@
 
 #pragma once
 #include "ChaosVDSerializedNameTable.h"
+#include "Chaos/ChaosArchive.h"
 #include "HAL/Platform.h"
 #include "Misc/EngineVersion.h"
 #include "Serialization/CustomVersion.h"
@@ -87,4 +88,14 @@ namespace Chaos::VisualDebugger
 
 		TSharedRef<FChaosVDSerializableNameTable> NameTableInstance;
 	};
+
+	/** Version of Chaos Archive that allows CVD to provide its own shared context instance */
+	class FChaosVDSharedContextArchive : public FChaosArchive
+	{
+	public:
+		explicit FChaosVDSharedContextArchive(FArchive& ArIn, const TSharedRef<FChaosArchiveContext>& InExternalSharedContext)
+			: FChaosArchive(ArIn, InExternalSharedContext)
+		{
+		}
+	};
 }
diff --git a/Engine/Source/Runtime/Experimental/Chaos/Public/ChaosVisualDebugger/ChaosVisualDebuggerTrace.h b/Engine/Source/Runtime/Experimental/Chaos/Public/ChaosVisualDebugger/ChaosVisualDebuggerTrace.h
--- a/Engine/Source/Runtime/Experimental/Chaos/Public/ChaosVisualDebugger/ChaosVisualDebuggerTrace.h
+++ b/Engine/Source/Runtime/Experimental/Chaos/Public/ChaosVisualDebugger/ChaosVisualDebuggerTrace.h
@@ -136,6 +136,10 @@ UE_TRACE_MINIMAL_EVENT_BEGIN_EXTERN(ChaosVDLogger, ChaosVDUsesAutoRTFM)
 	UE_TRACE_MINIMAL_EVENT_FIELD(bool, bUsingAutoRTFM)
 UE_TRACE_MINIMAL_EVENT_END()
 
+UE_TRACE_MINIMAL_EVENT_BEGIN_EXTERN(ChaosVDLogger, ChaosVDMiscHeaderData)
+	UE_TRACE_MINIMAL_EVENT_FIELD(bool, bUsesSharedGeometryContextArchive)
+UE_TRACE_MINIMAL_EVENT_END()
+
 UE_TRACE_MINIMAL_EVENT_BEGIN_EXTERN(ChaosVDLogger, ChaosVDTraceRelevancyVolume)
 	CVD_DEFINE_TRACE_VECTOR(Chaos::FReal, BoxMin)
 	CVD_DEFINE_TRACE_VECTOR(Chaos::FReal, BoxMax)
@@ -281,7 +285,12 @@ namespace Chaos::VD
 			return NewID;
 		}
 
-		static constexpr uint64 InvalidUniqueID = 0; 
+		static constexpr uint64 InvalidUniqueID = 0;
+
+		TSharedRef<FChaosArchiveContext> ArchiveContext;
+		FTransactionallySafeMutex ArchiveContextLock;
+
+		bool UsingSharedGeometryArchiveContext = true;
 
 		private:
 		std::atomic<uint64> InternalGenericIDCounter;
@@ -800,5 +809,15 @@ namespace Chaos::VisualDebugger
 
 		Data.Serialize(Ar);
 	}
+
+	template<typename TDataToSerialize>
+	void WriteDataToBuffer(TArray<uint8>& InOutDataBuffer, TDataToSerialize& Data, const TSharedRef<FChaosArchiveContext>& InSharedArchiveContext)
+	{
+		FChaosVDMemoryWriter MemWriterAr(InOutDataBuffer, FChaosVisualDebuggerTrace::GetNameTableInstance());
+		FChaosVDSharedContextArchive Ar(MemWriterAr, InSharedArchiveContext);
+		Ar.SetShouldSkipUpdateCustomVersion(true);
+
+		Data.Serialize(Ar);
+	}
 }
 #endif // WITH_CHAOS_VISUAL_DEBUGGER
diff --git a/Engine/Source/Runtime/PhysicsCore/Private/PhysTestSerializer.cpp b/Engine/Source/Runtime/PhysicsCore/Private/PhysTestSerializer.cpp
--- a/Engine/Source/Runtime/PhysicsCore/Private/PhysTestSerializer.cpp
+++ b/Engine/Source/Runtime/PhysicsCore/Private/PhysTestSerializer.cpp
@@ -20,7 +20,8 @@
 FPhysTestSerializer::FPhysTestSerializer()
 	: bDiskDataIsChaos(false)
 	, bChaosDataReady(false)
-	, Particles(UniqueIndices)
+	, Particles(UniqueIndices),
+	ChaosContext(MakeShared<Chaos::FChaosArchiveContext>())
 {
 }
 
@@ -39,9 +40,8 @@ void FPhysTestSerializer::Serialize(const TCHAR* FilePrefix)
 	TUniquePtr<FArchive> File(IFileManager::Get().CreateFileWriter(*UseFileName));
 	if (File)
 	{
-		Chaos::FChaosArchive Ar(*File);
 		UE_LOG(LogPhysicsCore, Log, TEXT("PhysTestSerialize File: %s"), *UseFileName);
-		Serialize(Ar);
+		Serialize(*File);
 	}
 	else
 	{
@@ -49,18 +49,20 @@ void FPhysTestSerializer::Serialize(const TCHAR* FilePrefix)
 	}
 }
 
-void FPhysTestSerializer::Serialize(Chaos::FChaosArchive& Ar)
+void FPhysTestSerializer::Serialize(FArchive& Ar)
 {
+	TUniquePtr<Chaos::FChaosArchive> ChaosAr = MakeUnique<Chaos::FChaosArchive>(Ar);
+	
 	if (!Ar.IsLoading())
 	{
 		//make sure any context we had set is restored before writing out sqcapture
-		Ar.SetContext(MoveTemp(ChaosContext));
+		ChaosAr = MakeUnique<FArchiveWithExternalContext>(Ar, ChaosContext);
 	}
 
 	static const FName TestSerializerName = "PhysTestSerializer";
 
 	{
-		Chaos::FChaosArchiveScopedMemory ScopedMemory(Ar, TestSerializerName, false);
+		Chaos::FChaosArchiveScopedMemory ScopedMemory(*ChaosAr, TestSerializerName, false);
 		int Version = 1;
 		Ar << Version;
 		Ar << bDiskDataIsChaos;
@@ -88,13 +90,13 @@ void FPhysTestSerializer::Serialize(Chaos::FChaosArchive& Ar)
 		CreateChaosData();
 #endif
 
-		Ar.SetContext(MoveTemp(ChaosContext));	//make sure any context we created during load is used for sqcapture
+		ChaosAr = MakeUnique<FArchiveWithExternalContext>(Ar, ChaosContext);	//make sure any context we created during load is used for sqcapture
 	}
 
 	PRAGMA_DISABLE_DEPRECATION_WARNINGS
 	bool bHasSQCapture = !!SQCapture;
 	{
-		Chaos::FChaosArchiveScopedMemory ScopedMemory(Ar, TestSerializerName, false);
+		Chaos::FChaosArchiveScopedMemory ScopedMemory(*ChaosAr, TestSerializerName, false);
 		Ar << bHasSQCapture;
 	}
 	if(bHasSQCapture)
@@ -103,10 +105,15 @@ void FPhysTestSerializer::Serialize(Chaos::FChaosArchive& Ar)
 		{
 			SQCapture = TUniquePtr<FSQCapture>(new FSQCapture(*this));
 		}
-		SQCapture->Serialize(Ar);
+		SQCapture->Serialize(*ChaosAr);
 	}
 	PRAGMA_ENABLE_DEPRECATION_WARNINGS
-	ChaosContext = Ar.StealContext();
+}
+
+void FPhysTestSerializer::Serialize(Chaos::FChaosArchive& Ar)
+{
+	// No longer supported
+	ensure(false);
 }
 
 void FPhysTestSerializer::SetPhysicsData(Chaos::FPBDRigidsEvolution& Evolution)
@@ -114,9 +121,11 @@ void FPhysTestSerializer::SetPhysicsData(Chaos::FPBDRigidsEvolution& Evolution)
 	bDiskDataIsChaos = true;
 	Data.Empty();
 	FMemoryWriter Ar(Data);
-	Chaos::FChaosArchive ChaosAr(Ar);
+
+	TSharedRef<Chaos::FChaosArchiveContext> NewExternalContext = MakeShared<Chaos::FChaosArchiveContext>();
+	FArchiveWithExternalContext ChaosAr(Ar, NewExternalContext);
 	Evolution.Serialize(ChaosAr);
-	ChaosContext = ChaosAr.StealContext();
+	ChaosContext = NewExternalContext;
 	ArchiveVersion = Ar.GetCustomVersions();
 }
 
@@ -241,12 +250,13 @@ void FPhysTestSerializer::CreateChaosData()
 		ChaosEvolution = MakeUnique<Chaos::FPBDRigidsEvolution>(Particles, PhysicalMaterials);
 
 		FMemoryReader Ar(Data);
-		Chaos::FChaosArchive ChaosAr(Ar);
+		TSharedRef<Chaos::FChaosArchiveContext> NewExternalContext = MakeShared<Chaos::FChaosArchiveContext>();
+		FArchiveWithExternalContext ChaosAr(Ar, NewExternalContext);
 
 		Ar.SetCustomVersions(ArchiveVersion);
 
 		ChaosEvolution->Serialize(ChaosAr);
-		ChaosContext = ChaosAr.StealContext();
+		ChaosContext = NewExternalContext;
 	}
 	bChaosDataReady = true;
 }
diff --git a/Engine/Source/Runtime/PhysicsCore/Public/PhysTestSerializer.h b/Engine/Source/Runtime/PhysicsCore/Public/PhysTestSerializer.h
--- a/Engine/Source/Runtime/PhysicsCore/Public/PhysTestSerializer.h
+++ b/Engine/Source/Runtime/PhysicsCore/Public/PhysTestSerializer.h
@@ -42,7 +42,9 @@ public:
 	FPhysTestSerializer& operator=(const FPhysTestSerializer&) = delete;
 	FPhysTestSerializer& operator=(FPhysTestSerializer&&) = delete;
 
+	UE_DEPRECATED(5.8, "Use the Serializer that takes an FArchive")
 	PHYSICSCORE_API void Serialize(Chaos::FChaosArchive& Ar);
+	PHYSICSCORE_API void Serialize(FArchive& Ar);
 	PHYSICSCORE_API void Serialize(const TCHAR* FilePrefix);
 
 	//Set the data from an external source. This will obliterate any existing data. Make sure you are not holding on to old internal data as it will go away
@@ -50,10 +52,11 @@ public:
 
 	const Chaos::FChaosArchiveContext* GetChaosContext() const
 	{
-		return ChaosContext.Get();
+		return &ChaosContext.Get();
 	}
 
-	PRAGMA_DISABLE_DEPRECATION_WARNINGS
+
+PRAGMA_DISABLE_DEPRECATION_WARNINGS
 	UE_DEPRECATED(5.8, "FSQCapture is unused and will be removed")
 	FSQCapture& CaptureSQ()
 	{
@@ -75,7 +78,8 @@ public:
 		}
 		return SQCapture.Get();
 	}
-	PRAGMA_ENABLE_DEPRECATION_WARNINGS
+
+PRAGMA_ENABLE_DEPRECATION_WARNINGS
 
 	Chaos::FPBDRigidsEvolution* GetChaosData()
 	{
@@ -98,14 +102,31 @@ private:
 
 private:
 
+	class FArchiveWithExternalContext : public Chaos::FChaosArchive
+	{
+	public:
+		FArchiveWithExternalContext(FArchive& ArIn)
+			: FChaosArchive(ArIn)
+		{
+		}
+
+		explicit FArchiveWithExternalContext(FArchive& ArIn, const TSharedRef<Chaos::FChaosArchiveContext>& InExternalSharedContext)
+			: FChaosArchive(ArIn, InExternalSharedContext)
+		{
+		}
+	};
+
 	TArray<uint8> Data;	//how the data is stored before going to disk
 	bool bDiskDataIsChaos;
 	bool bChaosDataReady;
 
-	PRAGMA_DISABLE_DEPRECATION_WARNINGS
+	
+PRAGMA_DISABLE_DEPRECATION_WARNINGS
+
 	UE_DEPRECATED(5.8, "FSQCapture is unused and will be removed")
 	TUniquePtr<FSQCapture> SQCapture;
-	PRAGMA_ENABLE_DEPRECATION_WARNINGS
+
+PRAGMA_ENABLE_DEPRECATION_WARNINGS
 
 	TUniquePtr<Chaos::FPBDRigidsEvolution> ChaosEvolution;
 	Chaos::FParticleUniqueIndicesMultithreaded UniqueIndices;
@@ -113,7 +134,7 @@ private:
 	Chaos::THandleArray<Chaos::FChaosPhysicsMaterial> PhysicalMaterials;
 	TArray <TUniquePtr<Chaos::FGeometryParticle>> GTParticles;
 
-	TUniquePtr<Chaos::FChaosArchiveContext> ChaosContext;
+	TSharedRef<Chaos::FChaosArchiveContext> ChaosContext;
 
 	FCustomVersionContainer ArchiveVersion;
 };
