Fix incorrect default value before parsing values.

#rnx

- LexFromString will return 0 if it fails to parse the int64 from a string not -1, we should match that as the default value so as not to mislead people.

diff --git a/Engine/Source/Runtime/Experimental/IoStore/OnDemand/Private/OnDemandConfig.cpp b/Engine/Source/Runtime/Experimental/IoStore/OnDemand/Private/OnDemandConfig.cpp
--- a/Engine/Source/Runtime/Experimental/IoStore/OnDemand/Private/OnDemandConfig.cpp
+++ b/Engine/Source/Runtime/Experimental/IoStore/OnDemand/Private/OnDemandConfig.cpp
@@ -145,7 +145,7 @@ int64 ParseSizeParam(FStringView Value)
 {
 	Value = Value.TrimStartAndEnd();
 
-	int64 Size = -1;
+	int64 Size = 0;
 	LexFromString(Size, Value);
 
 	if (Size > 0)
