Fix incorrect logging if WSACleanup failed.

#rnx

- Copy paste error, the log would report that WSAGetLastError had failed instead of WSACleanup.

diff --git a/Engine/Source/Runtime/Experimental/IoStore/OnDemand/Private/OnDemandModule.cpp b/Engine/Source/Runtime/Experimental/IoStore/OnDemand/Private/OnDemandModule.cpp
--- a/Engine/Source/Runtime/Experimental/IoStore/OnDemand/Private/OnDemandModule.cpp
+++ b/Engine/Source/Runtime/Experimental/IoStore/OnDemand/Private/OnDemandModule.cpp
@@ -50,7 +50,7 @@ struct FPlatformSocketSystem
 		{
 			// If WSACleanup fails it returns SOCKET_ERROR and WSAGetLastError must be called instead
 			const int SystemError = WSAGetLastError();
-			LogError(TEXT("WSAGetLastError"), SystemError);
+			LogError(TEXT("WSACleanup"), SystemError);
 		}
 	}
 
