Collision Filter Refactor Part 4.1: Updating ISpatialVisitor and how preprefiltering works so that it's more type safe and will be easier to refactor for new filters.

diff --git a/Engine/Plugins/ChaosVD/Source/ChaosVD/Public/ChaosVDSceneStreaming.h b/Engine/Plugins/ChaosVD/Source/ChaosVD/Public/ChaosVDSceneStreaming.h
--- a/Engine/Plugins/ChaosVD/Source/ChaosVD/Public/ChaosVDSceneStreaming.h
+++ b/Engine/Plugins/ChaosVD/Source/ChaosVD/Public/ChaosVDSceneStreaming.h
@@ -214,11 +214,13 @@ protected:
 			return true;
 		}
 
+		UE_DEPRECATED(5.8, "This was unused and is being removed")
 		const void* GetQueryData() const
 		{
 			return nullptr;
 		}
 
+		UE_DEPRECATED(5.8, "This was unused and is being removed")
 		const void* GetSimData() const
 		{
 			return nullptr;
@@ -234,10 +236,16 @@ protected:
 			return false;
 		}
 
+		UE_DEPRECATED(5.8, "This was unused and is being removed")
 		bool ShouldIgnore(const Chaos::TSpatialVisitorData<int32>& Instance) const
 		{
 			return false;
 		}
+		
+		const bool PrePreFilter(int) const
+		{
+			return false;
+		}
 
 		TSet<int32>& CollectedResults;
 	};
diff --git a/Engine/Source/Programs/HeadlessChaos/Private/HeadlessChaosTestPhysicsScene.cpp b/Engine/Source/Programs/HeadlessChaos/Private/HeadlessChaosTestPhysicsScene.cpp
--- a/Engine/Source/Programs/HeadlessChaos/Private/HeadlessChaosTestPhysicsScene.cpp
+++ b/Engine/Source/Programs/HeadlessChaos/Private/HeadlessChaosTestPhysicsScene.cpp
@@ -30,11 +30,11 @@ namespace ChaosTest {
 	struct FSimpleRaycastVisitor: ISpatialVisitor<FAccelerationStructureHandle>
 	{
 		using FPayload = FAccelerationStructureHandle;
-		FVec3 Start;
-		bool bHit;
-		bool bQueryGameThread; // Query game thread or physics thread data?
+		FVec3 Start = FVec3::Zero();
+		bool bHit = false;
+		bool bQueryGameThread = true; // Query game thread or physics thread data?
 
-		bool bUseQueryFilter;
+		bool bUseQueryFilter = false;
 		FCollisionFilterData FilterData;
 
 		FSimpleRaycastVisitor(const FVec3& InStart, bool bInQueryGameThread)
@@ -54,14 +54,13 @@ namespace ChaosTest {
 		{
 		}
 
-		virtual const void* GetQueryData() const override
+		virtual bool PrePreFilter(const FAccelerationStructureHandle& Payload) const override
 		{
 			if (bUseQueryFilter)
 			{
-				return &FilterData;
+				return Payload.PrePreQueryFilter(&FilterData);
 			}
-
-			return nullptr;
+			return false;
 		}
 
 		enum class SQType
diff --git a/Engine/Source/Runtime/Experimental/Chaos/Public/Chaos/AABBTree.h b/Engine/Source/Runtime/Experimental/Chaos/Public/Chaos/AABBTree.h
--- a/Engine/Source/Runtime/Experimental/Chaos/Public/Chaos/AABBTree.h
+++ b/Engine/Source/Runtime/Experimental/Chaos/Public/Chaos/AABBTree.h
@@ -3736,8 +3736,7 @@ private:
 				return true;
 			}
 
-			const void* GetQueryData() const { return nullptr; }
-			const void* GetSimData() const { return nullptr; }
+			bool PrePreFilter(const TPayloadType&) const { return false; }
 
 			/** Return a pointer to the payload on which we are querying the acceleration structure */
 			const void* GetQueryPayload() const 
@@ -3750,7 +3749,6 @@ private:
 				return false;
 			}
 
-			bool ShouldIgnore(const TSpatialVisitorData<TPayloadType>& Instance) const { return false; }
 			TArray<TPayloadType>& CollectedResults;
 		};
 
diff --git a/Engine/Source/Runtime/Experimental/Chaos/Public/Chaos/BoundingVolume.h b/Engine/Source/Runtime/Experimental/Chaos/Public/Chaos/BoundingVolume.h
--- a/Engine/Source/Runtime/Experimental/Chaos/Public/Chaos/BoundingVolume.h
+++ b/Engine/Source/Runtime/Experimental/Chaos/Public/Chaos/BoundingVolume.h
@@ -219,10 +219,8 @@ public:
 				CollectedResults.Add(Instance.Payload);
 				return true;
 			}
-			const void* GetQueryData() const { return nullptr; }
-			const void* GetSimData() const { return nullptr; }
+			bool PrePreFilter(const TPayloadType&) const { return false; }
 			const void* GetQueryPayload() const { return nullptr; }
-			bool ShouldIgnore(const TSpatialVisitorData<TPayloadType>& Instance) const { return false; }
 			TArray<TPayloadType>& CollectedResults;
 		};
 
diff --git a/Engine/Source/Runtime/Experimental/Chaos/Public/Chaos/Collision/SpatialAccelerationBroadPhase.h b/Engine/Source/Runtime/Experimental/Chaos/Public/Chaos/Collision/SpatialAccelerationBroadPhase.h
--- a/Engine/Source/Runtime/Experimental/Chaos/Public/Chaos/Collision/SpatialAccelerationBroadPhase.h
+++ b/Engine/Source/Runtime/Experimental/Chaos/Public/Chaos/Collision/SpatialAccelerationBroadPhase.h
@@ -278,8 +278,10 @@ namespace Chaos
 				return false;
 			}
 
+			UE_DEPRECATED(5.8, "This has been deprecated in favor of PrePreFilter")
 			const void* GetQueryData() const { return nullptr; }
 
+			UE_DEPRECATED(5.8, "This has been deprecated in favor of PrePreFilter")
 			const void* GetSimData() const { return &SimFilterData; }
 
 			bool ShouldIgnore(const TSpatialVisitorData<FAccelerationStructureHandle>& Instance) const
@@ -297,7 +299,7 @@ namespace Chaos
 				return false;
 			}
 
-			UE_INTERNAL bool PrePreFilter(const FAccelerationStructureHandle& Payload) const
+			bool PrePreFilter(const FAccelerationStructureHandle& Payload) const
 			{
 				if (ShouldIgnore(Payload))
 				{
@@ -319,14 +321,6 @@ namespace Chaos
 		};
 	}
 
-	template <>
-	inline bool PrePreFilterHelper(const FAccelerationStructureHandle& Payload, const Private::FSimOverlapVisitor& Visitor)
-	{
-		PRAGMA_DISABLE_INTERNAL_WARNINGS
-		return Visitor.PrePreFilter(Payload);
-		PRAGMA_ENABLE_INTERNAL_WARNINGS
-	}
-
 	/**
 	 * A broad phase that iterates over particle and uses a spatial acceleration structure to output
 	 * potentially overlapping SpatialAccelerationHandles.
diff --git a/Engine/Source/Runtime/Experimental/Chaos/Public/Chaos/ISpatialAcceleration.h b/Engine/Source/Runtime/Experimental/Chaos/Public/Chaos/ISpatialAcceleration.h
--- a/Engine/Source/Runtime/Experimental/Chaos/Public/Chaos/ISpatialAcceleration.h
+++ b/Engine/Source/Runtime/Experimental/Chaos/Public/Chaos/ISpatialAcceleration.h
@@ -6,6 +6,7 @@
 #include "ChaosCheck.h"
 #include "ChaosDebugDrawDeclares.h"
 #include "Containers/HashTable.h"
+#include "Concepts/DecaysTo.h"
 
 namespace Chaos
 {
@@ -142,11 +143,16 @@ public:
 	*/
 	virtual bool Sweep(const TSpatialVisitorData<TPayloadType>& Instance, FQueryFastData& CurData) = 0;
 
-	virtual const void* GetQueryData() const { return nullptr; }
+	UE_DEPRECATED(5.8, "Implement PrePreFilter instead")
+	virtual const void* GetQueryData() const final { return nullptr; }
 
-	virtual const void* GetSimData() const { return nullptr; }
+	UE_DEPRECATED(5.8, "Implement PrePreFilter instead")
+	virtual const void* GetSimData() const final { return nullptr; }
 
-	virtual bool ShouldIgnore(const TSpatialVisitorData<TPayloadType>& Instance) const { return false; }
+	UE_DEPRECATED(5.8, "Implement PrePreFilter instead")
+	virtual bool ShouldIgnore(const TSpatialVisitorData<TPayloadType>& Instance) const final { return false; }
+
+	virtual bool PrePreFilter(const TPayloadType& Payload) const { return false; }
 
 	/** Return a pointer to the payload on which we are querying the acceleration structure */
 	virtual const void* GetQueryPayload() const { return nullptr; }
@@ -469,19 +475,33 @@ public:
 		return Visitor.Sweep(Instance, CurData);
 	}
 
+	UE_DEPRECATED(5.8, "If this is errors, your visitor is missing an implementation of PrePreFilter")
 	FORCEINLINE const void* GetQueryData() const
 	{
+		PRAGMA_DISABLE_DEPRECATION_WARNINGS
 		return Visitor.GetQueryData();
+		PRAGMA_ENABLE_DEPRECATION_WARNINGS
 	}
 
+	UE_DEPRECATED(5.8, "If this is errors, your visitor is missing an implementation of PrePreFilter")
 	FORCEINLINE const void* GetSimData() const
 	{
+		PRAGMA_DISABLE_DEPRECATION_WARNINGS
 		return Visitor.GetSimData();
+		PRAGMA_ENABLE_DEPRECATION_WARNINGS
 	}
 
+	UE_DEPRECATED(5.8, "If this is errors, your visitor is missing an implementation of PrePreFilter")
 	FORCEINLINE bool ShouldIgnore(const TSpatialVisitorData<TPayloadType>& Instance) const
 	{
+		PRAGMA_DISABLE_DEPRECATION_WARNINGS
 		return Visitor.ShouldIgnore(Instance);
+		PRAGMA_ENABLE_DEPRECATION_WARNINGS
+	}
+
+	FORCEINLINE bool PrePreFilter(const TPayloadType& Payload) const
+	{
+		return Visitor.PrePreFilter(Payload);
 	}
 
 	/** Return a pointer to the payload on which we are querying the acceleration structure */
@@ -1000,8 +1020,15 @@ FChaosArchive& operator<< (FChaosArchive& Ar, TSQMap<TKey, TValue>& Map)
 	return Ar;
 }
 
-template <typename TPayload, typename TVisitor>
-typename TEnableIf<!TIsPointer<TPayload>::Value, bool>::Type PrePreFilterHelper(const TPayload& Payload, const TVisitor& Visitor)
+template <typename PayloadType, typename VisitorType>
+concept CVisitorHasPrePreFilter = requires(PayloadType Payload, VisitorType Visitor)
+{
+    { Visitor.PrePreFilter(Payload) } -> UE::CDecaysTo<bool>;
+};
+
+template <typename PayloadType, typename VisitorType>
+UE_DEPRECATED(5.8, "Implement the new PrePreFilter method on your visitor")
+UE_INTERNAL bool PrePreFilterHelperInternal(const PayloadType& Payload, const VisitorType& Visitor)
 {
 	if (Visitor.ShouldIgnore(Payload))
 	{
@@ -1018,6 +1045,21 @@ typename TEnableIf<!TIsPointer<TPayload>::Value, bool>::Type PrePreFilterHelper(
 	return false;
 }
 
+template <typename PayloadType, typename VisitorType>
+requires CVisitorHasPrePreFilter<PayloadType, VisitorType>
+UE_INTERNAL bool PrePreFilterHelperInternal(const PayloadType& Payload, const VisitorType& Visitor)
+{
+	return Visitor.PrePreFilter(Payload);
+}
+
+template <typename TPayload, typename TVisitor>
+typename TEnableIf<!TIsPointer<TPayload>::Value, bool>::Type PrePreFilterHelper(const TPayload& Payload, const TVisitor& Visitor)
+{
+	PRAGMA_DISABLE_INTERNAL_WARNINGS
+	return PrePreFilterHelperInternal(Payload, Visitor);
+	PRAGMA_ENABLE_INTERNAL_WARNINGS
+}
+
 template <typename TPayload, typename TVisitor>
 typename TEnableIf<TIsPointer<TPayload>::Value, bool>::Type PrePreFilterHelper(const TPayload& Payload, const TVisitor& Visitor)
 {
diff --git a/Engine/Source/Runtime/PhysicsCore/Public/SQVisitor.h b/Engine/Source/Runtime/PhysicsCore/Public/SQVisitor.h
--- a/Engine/Source/Runtime/PhysicsCore/Public/SQVisitor.h
+++ b/Engine/Source/Runtime/PhysicsCore/Public/SQVisitor.h
@@ -162,14 +162,9 @@ struct TSQVisitor : public Chaos::ISpatialVisitor<TPayload, Chaos::FReal>
 		return Visit<ESQType::Overlap>(Instance, nullptr);
 	}
 
-	virtual const void* GetQueryData() const override
+	virtual bool PrePreFilter(const TPayload& Payload) const override
 	{
-		return &LegacyQueryFilterDataConcrete;
-	}
-
-	virtual const void* GetSimData() const override
-	{
-		return nullptr;
+		return Payload.PrePreQueryFilter(&LegacyQueryFilterDataConcrete);
 	}
 
 	/** Return a pointer to the payload on which we are querying the acceleration structure */
@@ -588,14 +583,9 @@ struct TBPVisitor : public Chaos::ISpatialVisitor<TPayload, Chaos::FReal>
 		return false;
 	}
 
-	virtual const void* GetQueryData() const override
+	virtual bool PrePreFilter(const TPayload& Payload) const override
 	{
-		return &LegacyQueryFilterDataConcrete;
-	}
-
-	virtual const void* GetSimData() const override
-	{
-		return nullptr;
+		return Payload.PrePreQueryFilter(&LegacyQueryFilterDataConcrete);
 	}
 
 	/** Return a pointer to the payload on which we are querying the acceleration structure */
