Chaos
- collision modifier option: per-manifold-point friction
- collision modifier option: per-manifold-point position/velocity/friction phase enable/disable
- split impulse mode (Experimental and work-in-progress)

#rnx

diff --git a/Engine/Source/Runtime/Experimental/Chaos/Private/Chaos/Collision/PBDCollisionContainerSolver.cpp b/Engine/Source/Runtime/Experimental/Chaos/Private/Chaos/Collision/PBDCollisionContainerSolver.cpp
--- a/Engine/Source/Runtime/Experimental/Chaos/Private/Chaos/Collision/PBDCollisionContainerSolver.cpp
+++ b/Engine/Source/Runtime/Experimental/Chaos/Private/Chaos/Collision/PBDCollisionContainerSolver.cpp
@@ -79,7 +79,8 @@ namespace Chaos
 		const FRealSingle MaxDepenetrationVelocity,
 		const FRealSingle MaxPushOut,
 		const FConstraintSolverBody& Body0,
-		const FConstraintSolverBody& Body1)
+		const FConstraintSolverBody& Body1,
+		const FPBDCollisionSolverSettings& SolverSettings)
 	{
 		FManifoldPoint& ManifoldPoint = Constraint->GetManifoldPoint(ConstraintPointIndex);
 
@@ -156,12 +157,12 @@ namespace Chaos
 			}
 		}
 
+		// Overlap remaining from the previous frame, estimated from current contact phi and velocity
+		FRealSingle WorldContactResidualPhi = FMath::Min(WorldContactDeltaNormal, FRealSingle(0)) - FMath::Min(ContactVelocityNormal * Dt, FRealSingle(0));
+
 		// Handle initial overlaps. Initially overlapping objects will de-penetrate at MaxDepenetrationVelocity (which may be zero).
 		if ((MaxDepenetrationVelocity >= 0) && CVars::bChaos_Collision_EnableInitialDepenetration)
 		{
-			// Overlap remaining from the previous frame, estimated from current contact phi and velocity
-			FRealSingle WorldContactResidualPhi = FMath::Min(WorldContactDeltaNormal, FRealSingle(0)) - FMath::Min(ContactVelocityNormal * Dt, FRealSingle(0));
-
 			// We can handle initial overlap depenentration individually per manifold point or
 			// globally over all maifold points. See FPBDCollisionConstraint::Setup
 			const bool bUsePerContactInitialPhi = Constraint->UsePerContactInitialPhi();
@@ -221,6 +222,14 @@ namespace Chaos
 		// Adjust depth to account for target penetration from user
 		WorldContactDeltaNormal -= ManifoldPoint.TargetPhi;
 
+		FSolverReal SplitImpulseFraction = 0.0f;
+		if (SolverSettings.bUseSplitImpulse && !FMath::IsNearlyZero(WorldContactDeltaNormal))
+		{
+			// SplitImpulseFraction: how much of WorldContactDeltaNormal comes from residual overlap from the previous frame.
+			// Use a momentumless correction for all new overlap, but still add momentum to correct any redisual overlap
+			SplitImpulseFraction = FMath::Clamp(1.0f - (WorldContactResidualPhi / WorldContactDeltaNormal), 0.0f, 1.0f);
+		}
+
 		Solver.InitManifoldPoint(
 			SolverPointIndex,
 			Dt,
@@ -232,7 +241,12 @@ namespace Chaos
 			WorldContactDeltaNormal,
 			WorldContactDeltaTangentU,
 			WorldContactDeltaTangentV,
-			WorldContactTargetVelocityNormal);
+			WorldContactTargetVelocityNormal,
+			SplitImpulseFraction,
+			ManifoldPoint.Flags.bDisablePositionSolve,
+			ManifoldPoint.Flags.bDisableVelocitySolve,
+			ManifoldPoint.Flags.bDisableFriction
+			);
 	}
 
 	void UpdateCollisionSolverManifoldFromConstraint(
@@ -264,7 +278,12 @@ namespace Chaos
 				{
 					// Transform the constraint contact data into world space for use by the solver
 					// We build this data directly into the solver's world-space contact data which looks a bit odd with "Init" called after but there you go
-					UpdateCollisionSolverContactPointFromConstraint(Solver, SolverManifoldPointIndex, Constraint, ConstraintManifoldPointIndex, Dt, MaxDepenetrationVelocity, MaxPushOut, Body0, Body1);
+					UpdateCollisionSolverContactPointFromConstraint(
+						Solver, SolverManifoldPointIndex, 
+						Constraint, ConstraintManifoldPointIndex, 
+						Dt, MaxDepenetrationVelocity, MaxPushOut, 
+						Body0, Body1, 
+						SolverSettings);
 				}
 			}
 		}
diff --git a/Engine/Source/Runtime/Experimental/Chaos/Private/Chaos/ContactModification.cpp b/Engine/Source/Runtime/Experimental/Chaos/Private/Chaos/ContactModification.cpp
--- a/Engine/Source/Runtime/Experimental/Chaos/Private/Chaos/ContactModification.cpp
+++ b/Engine/Source/Runtime/Experimental/Chaos/Private/Chaos/ContactModification.cpp
@@ -183,6 +183,24 @@ namespace Chaos
 		Modifier->MarkConstraintForManifoldUpdate(*Constraint);
 	}
 
+	void FContactPairModifier::ModifyEnablePositionSolve(bool bEnable, int32 ContactPointIdx)
+	{
+		Constraint->SetManifoldPointPositionSolveEnabled(ContactPointIdx, bEnable);
+		Constraint->SetModifierApplied();
+	}
+
+	void FContactPairModifier::ModifyEnableVelocitySolve(bool bEnable, int32 ContactPointIdx)
+	{
+		Constraint->SetManifoldPointVelocitySolveEnabled(ContactPointIdx, bEnable);
+		Constraint->SetModifierApplied();
+	}
+
+	void FContactPairModifier::ModifyEnableFriction(bool bEnable, int32 ContactPointIdx)
+	{
+		Constraint->SetManifoldPointFrictionEnabled(ContactPointIdx, bEnable);
+		Constraint->SetModifierApplied();
+	}
+
 	FReal FContactPairModifier::GetRestitution() const
 	{
 		return Constraint->GetRestitution();
@@ -538,6 +556,20 @@ namespace Chaos
 		Constraint->DisableManifoldPoint(ContactPointIdx);
 	}
 
+	int32 FContactPairModifier::AddContact(const FVec3& WorldPos0, const FVec3& WorldPos1, const FVec3& WorldNormal) const
+	{
+		FContactPointf ContactPoint;
+		ContactPoint.ShapeContactPoints[0] = Constraint->GetShapeWorldTransform0().InverseTransformPositionNoScale(WorldPos0);
+		ContactPoint.ShapeContactPoints[1] = Constraint->GetShapeWorldTransform1().InverseTransformPositionNoScale(WorldPos1);
+		ContactPoint.ShapeContactNormal = Constraint->GetShapeWorldTransform1().InverseTransformVectorNoScale(WorldNormal);
+		ContactPoint.Phi = FRealSingle(FVec3::DotProduct(WorldPos0 - WorldPos1, WorldNormal));
+	
+		int32 ContactIndex = Constraint->AddManifoldPoint(ContactPoint);
+		Constraint->SetModifierApplied();
+		return ContactIndex;
+	}
+
+
 	const FPBDCollisionConstraintContainerCookie& FContactPairModifier::GetConstraintContainerCookie() const
 	{
 		return Constraint->GetContainerCookie();
diff --git a/Engine/Source/Runtime/Experimental/Chaos/Private/PBDRigidsSolver.cpp b/Engine/Source/Runtime/Experimental/Chaos/Private/PBDRigidsSolver.cpp
--- a/Engine/Source/Runtime/Experimental/Chaos/Private/PBDRigidsSolver.cpp
+++ b/Engine/Source/Runtime/Experimental/Chaos/Private/PBDRigidsSolver.cpp
@@ -333,6 +333,11 @@ namespace Chaos
 		int32 ChaosSolverDeterministic = -1;
 		FAutoConsoleVariableRef CVarChaosSolverDeterministic(TEXT("p.Chaos.Solver.Deterministic"), ChaosSolverDeterministic, TEXT("Override determinism. 0: disabled; 1: enabled; -1: use config"));
 
+		// Whether to force split-impulse (momentumless position correction) mode in the constraint solver. -1 = no overide; 0 = force off; 1 = force on
+		// NOTE: only supported by collision solver (not joints etc)
+		int32 ChaosSolverSplitImpulseMode = -1;
+		FAutoConsoleVariableRef CVarChaosCollisionSplitImpulseMode(TEXT("p.Chaos.Solver.SplitImpulseMode"), ChaosSolverSplitImpulseMode, TEXT("EXPERIMENTAL. Whether to force split-impulse mode. Currently implemented for collison only."));
+
 		// Copied from RBAN
 		FRealSingle ChaosSolverJointPositionTolerance = 0.025f;
 		FAutoConsoleVariableRef CVarChaosSolverJointPositionTolerance(TEXT("p.Chaos.Solver.Joint.PositionTolerance"), ChaosSolverJointPositionTolerance, TEXT("PositionTolerance."));
@@ -3384,6 +3389,10 @@ TRACE_COUNTER_SET(ChaosTraceCounter_##Name, Value)
 			{
 				MEvolution->GetCollisionConstraints().SetVelocityFrictionIterations(ChaosSolverCollisionVelocityFrictionIterations);
 			}
+			if (ChaosSolverSplitImpulseMode >= 0)
+			{
+				MEvolution->GetCollisionConstraints().SetUseSplitImpulse(ChaosSolverSplitImpulseMode > 0);
+			}
 			{
 				MEvolution->SetShockPropagationIterations(ChaosSolverCollisionPositionShockPropagationIterations, ChaosSolverCollisionVelocityShockPropagationIterations);
 			}
diff --git a/Engine/Source/Runtime/Experimental/Chaos/Public/Chaos/Collision/ContactPoint.h b/Engine/Source/Runtime/Experimental/Chaos/Public/Chaos/Collision/ContactPoint.h
--- a/Engine/Source/Runtime/Experimental/Chaos/Public/Chaos/Collision/ContactPoint.h
+++ b/Engine/Source/Runtime/Experimental/Chaos/Public/Chaos/Collision/ContactPoint.h
@@ -109,10 +109,16 @@ namespace Chaos
 				uint8 bWasReplaced : 1;						// @todo(chaos): remove this
 				uint8 bHasStaticFrictionAnchor : 1;			// Whether our static friction anchor was recovered from a prior tick
 				uint8 bInitialContact : 1;					// Whether this is a new contact this tick for handling initial penetration
+				uint8 bDisablePositionSolve : 1;			// Whether we should disable position solve on this manifold point
+				uint8 bDisableVelocitySolve : 1;			// Whether we should disable velocity solve on this manifold point
+				uint8 bDisableFriction : 1;					// Whether we should disable friction solve on this manifold point
 			};
 			uint8 Bits;
 		};
 
+		// Verify that all FFlags data can be initialized to zero (even unused bits)
+		static_assert(sizeof(FFlags) == sizeof(FFlags::Bits), "FManifoldPoint::FFlags has exceeded capacity. Change to uint16 or uint32.");
+
 		FManifoldPoint()
 			: ContactPoint()
 			, Flags()
diff --git a/Engine/Source/Runtime/Experimental/Chaos/Public/Chaos/Collision/PBDCollisionConstraint.h b/Engine/Source/Runtime/Experimental/Chaos/Public/Chaos/Collision/PBDCollisionConstraint.h
--- a/Engine/Source/Runtime/Experimental/Chaos/Public/Chaos/Collision/PBDCollisionConstraint.h
+++ b/Engine/Source/Runtime/Experimental/Chaos/Public/Chaos/Collision/PBDCollisionConstraint.h
@@ -593,6 +593,21 @@ namespace Chaos
 			}
 		}
 
+		inline void SetManifoldPointPositionSolveEnabled(const int32 ManifoldPointIndex, bool bInEnabled)
+		{
+			ManifoldPoints[ManifoldPointIndex].Flags.bDisablePositionSolve = !bInEnabled;
+		}
+
+		inline void SetManifoldPointVelocitySolveEnabled(const int32 ManifoldPointIndex, bool bInEnabled)
+		{
+			ManifoldPoints[ManifoldPointIndex].Flags.bDisableVelocitySolve = !bInEnabled;
+		}
+
+		inline void SetManifoldPointFrictionEnabled(const int32 ManifoldPointIndex, bool bInEnabled)
+		{
+			ManifoldPoints[ManifoldPointIndex].Flags.bDisableFriction = !bInEnabled;
+		}
+
 		CHAOS_API void AddIncrementalManifoldContact(const FContactPoint& ContactPoint);
 
 		// @todo(chaos): remove this and use SetOneShotManifoldContacts
@@ -945,7 +960,7 @@ namespace Chaos
 				uint16 bDisabled : 1;					// Is this contact disabled (by the user or because cull distance is exceeded)
 				uint16 bUseManifold : 1;				// Should we use contact manifolds or single points (faster but poor behaviour)
 				uint16 bUseIncrementalManifold : 1;		// Do we need to run incremental collision detection (only LavelSets now)
-				uint16 bCanRestoreManifold : 1;			// Can we try to restore the manifold this frame (set folowing narrowphase, cleared when reset for some reason)
+				uint16 bCanRestoreManifold : 1;			// Can we try to restore the manifold this frame (set following narrowphase, cleared when reset for some reason)
 				uint16 bWasManifoldRestored : 1;		// Did we restore the manifold this frame
 				uint16 bIsQuadratic0 : 1;				// Is the first shape a sphere or capsule
 				uint16 bIsQuadratic1 : 1;				// Is the second shape a sphere or capsule
diff --git a/Engine/Source/Runtime/Experimental/Chaos/Public/Chaos/Collision/PBDCollisionSolver.h b/Engine/Source/Runtime/Experimental/Chaos/Public/Chaos/Collision/PBDCollisionSolver.h
--- a/Engine/Source/Runtime/Experimental/Chaos/Public/Chaos/Collision/PBDCollisionSolver.h
+++ b/Engine/Source/Runtime/Experimental/Chaos/Public/Chaos/Collision/PBDCollisionSolver.h
@@ -72,13 +72,28 @@ namespace Chaos
 			FSolverReal NetImpulseTangentV;
 			FSolverReal NetSoftPushOutNormal;
 
+			FSolverReal SplitImpulseAlpha;
+
 			// A measure of how much we exceeded the static friction threshold.
 			// Equal to (NormalPushOut / TangentialPushOut) before clamping to the friction cone.
 			// Used to move the static friction anchors to the edge of the cone in Scatter.
 			FSolverReal StaticFrictionRatio;
 
-			// Transient - whether to apply friction on the current iteration
-			uint32 bApplyFriction : 1;
+			union FFlags
+			{
+				struct 
+				{
+					uint8 bDisablePositionSolve : 1;	// Disable the position solver step
+					uint8 bDisableVelocitySolve : 1;	// Disable the velocity solver step
+					uint8 bDisableFriction : 1;		// Disable friction
+				};
+				uint8 Bits;
+			};
+
+			FFlags Flags;
+
+			// Verify that all FFlags data can be initialized to zero (even unused bits)
+			static_assert(sizeof(FFlags) == sizeof(FFlags::Bits), "FPBDCollisionSolverManifoldPoint::FFlags has exceeded capacity. Change to uint16 or uint32.");
 		};
 
 		/**
@@ -198,8 +213,12 @@ namespace Chaos
 				const FSolverReal InWorldContactDeltaNormal,
 				const FSolverReal InWorldContactDeltaTangentU,
 				const FSolverReal InWorldContactDeltaTangentV,
-				const FSolverReal InWorldContactVelocityTargetNormal);
-
+				const FSolverReal InWorldContactVelocityTargetNormal,
+				const FSolverReal InSplitImpulseAlpha = 0.0f,
+				const bool bDisablePositionSolve = false,
+				const bool bDisableVelocitySolve = false,
+				const bool bDisableFriction = false);
+		
 			/**
 			 * Call once all manifold points have been initialized. Calculate mass properties etc.
 			 */
@@ -553,6 +572,8 @@ namespace Chaos
 				const FSolverVec3 DX0 = State.InvMs[0] * PushOut;
 				const FSolverVec3 DR0 = ManifoldPoint.ContactTangentUAngular0 * PushOutTangentU + ManifoldPoint.ContactTangentVAngular0 * PushOutTangentV;
 				FConstraintSolverBody& Body0 = SolverBody0();
+
+				// NOTE: No split impulse mode for friction
 				Body0.ApplyPositionDelta(DX0);
 				Body0.ApplyRotationDelta(DR0);
 			}
@@ -561,6 +582,8 @@ namespace Chaos
 				const FSolverVec3 DX1 = -State.InvMs[1] * PushOut;
 				const FSolverVec3 DR1 = ManifoldPoint.ContactTangentUAngular1 * -PushOutTangentU + ManifoldPoint.ContactTangentVAngular1 * -PushOutTangentV;
 				FConstraintSolverBody& Body1 = SolverBody1();
+
+				// NOTE: No split impulse mode for friction
 				Body1.ApplyPositionDelta(DX1);
 				Body1.ApplyRotationDelta(DR1);
 			}
@@ -588,16 +611,34 @@ namespace Chaos
 				const FSolverVec3 DX0 = (State.InvMs[0] * PushOutNormal) * ManifoldPoint.ContactNormal;
 				const FSolverVec3 DR0 = ManifoldPoint.ContactNormalAngular0 * PushOutNormal;
 				FConstraintSolverBody& Body0 = SolverBody0();
-				Body0.ApplyPositionDelta(DX0);
-				Body0.ApplyRotationDelta(DR0);
+				if (ManifoldPoint.SplitImpulseAlpha > 0)
+				{
+					Body0.ApplyPositionCorrectionDelta(ManifoldPoint.SplitImpulseAlpha * DX0);
+					Body0.ApplyRotationCorrectionDelta(ManifoldPoint.SplitImpulseAlpha * DR0);
+				}
+
+				if (ManifoldPoint.SplitImpulseAlpha < 1)
+				{
+					Body0.ApplyPositionDelta((1.0f - ManifoldPoint.SplitImpulseAlpha) * DX0);
+					Body0.ApplyRotationDelta((1.0f - ManifoldPoint.SplitImpulseAlpha) * DR0);
+				}
 			}
 			if (IsDynamic(1))
 			{
 				const FSolverVec3 DX1 = (State.InvMs[1] * -PushOutNormal) * ManifoldPoint.ContactNormal;
 				const FSolverVec3 DR1 = ManifoldPoint.ContactNormalAngular1 * -PushOutNormal;
 				FConstraintSolverBody& Body1 = SolverBody1();
-				Body1.ApplyPositionDelta(DX1);
-				Body1.ApplyRotationDelta(DR1);
+				if (ManifoldPoint.SplitImpulseAlpha > 0)
+				{
+					Body1.ApplyPositionCorrectionDelta(ManifoldPoint.SplitImpulseAlpha * DX1);
+					Body1.ApplyRotationCorrectionDelta(ManifoldPoint.SplitImpulseAlpha * DR1);
+				}
+
+				if (ManifoldPoint.SplitImpulseAlpha < 1)
+				{
+					Body1.ApplyPositionDelta((1.0f - ManifoldPoint.SplitImpulseAlpha) * DX1);
+					Body1.ApplyRotationDelta((1.0f - ManifoldPoint.SplitImpulseAlpha) * DR1);
+				}
 			}
 		}
 
@@ -963,10 +1004,18 @@ namespace Chaos
 			const FSolverReal InWorldContactDeltaNormal,
 			const FSolverReal InWorldContactDeltaTangentU,
 			const FSolverReal InWorldContactDeltaTangentV,
-			const FSolverReal InWorldContactVelocityTargetNormal)
+			const FSolverReal InWorldContactVelocityTargetNormal,
+			const FSolverReal InSplitImpulseAlpha,
+			const bool bDisablePositionSolve,
+			const bool bDisableVelocitySolve,
+			const bool bDisableFriction
+			)
 		{
 			check(State.ManifoldPoints != nullptr);
 			check(PointIndex < State.NumManifoldPoints);
+			check(InSplitImpulseAlpha >= 0.0f);
+			check(InSplitImpulseAlpha <= 1.0f);
+
 			FPBDCollisionSolverManifoldPoint& ManifoldPoint = State.ManifoldPoints[PointIndex];
 
 			ManifoldPoint.RelativeContactPoints[0] = InRelativeContactPosition0;
@@ -979,6 +1028,7 @@ namespace Chaos
 			ManifoldPoint.ContactDeltaTangentU = InWorldContactDeltaTangentU;
 			ManifoldPoint.ContactDeltaTangentV = InWorldContactDeltaTangentV;
 			ManifoldPoint.ContactTargetVelocityNormal = InWorldContactVelocityTargetNormal;
+			ManifoldPoint.SplitImpulseAlpha = InSplitImpulseAlpha;
 
 			ManifoldPoint.ContactRxNormal0 = FSolverVec3::CrossProduct(InRelativeContactPosition0, InWorldContactNormal);
 			ManifoldPoint.ContactRxTangentU0 = FSolverVec3::CrossProduct(InRelativeContactPosition0, InWorldContactTangentU);
@@ -993,7 +1043,13 @@ namespace Chaos
 			ManifoldPoint.NetImpulseNormal = FSolverReal(0);
 			ManifoldPoint.NetImpulseTangentU = FSolverReal(0);
 			ManifoldPoint.NetImpulseTangentV = FSolverReal(0);
+			ManifoldPoint.NetSoftPushOutNormal = FSolverReal(0);
 			ManifoldPoint.StaticFrictionRatio = FSolverReal(0);
+
+			ManifoldPoint.Flags.Bits = 0;
+			ManifoldPoint.Flags.bDisablePositionSolve = bDisablePositionSolve;
+			ManifoldPoint.Flags.bDisableVelocitySolve = bDisableVelocitySolve;
+			ManifoldPoint.Flags.bDisableFriction = bDisableFriction;
 		}
 
 		FORCEINLINE_DEBUGGABLE void FPBDCollisionSolver::FinalizeManifold()
@@ -1005,7 +1061,6 @@ namespace Chaos
 			const FSolverReal Dt,
 			const FSolverReal MaxPushOut)
 		{
-
 			// Apply the position correction along the normal and determine if we want to run friction on each point
 			SolvePositionNoFriction(Dt, MaxPushOut);
 
@@ -1016,6 +1071,11 @@ namespace Chaos
 				for (int32 PointIndex = 0; PointIndex < NumManifoldPoints(); ++PointIndex)
 				{
 					FPBDCollisionSolverManifoldPoint& SolverManifoldPoint = State.ManifoldPoints[PointIndex];
+					const bool bFrictionSolveDisabled = !!SolverManifoldPoint.Flags.bDisableFriction || !!SolverManifoldPoint.Flags.bDisablePositionSolve;
+					if (bFrictionSolveDisabled)
+					{
+						continue;
+					}
 
 					const FSolverReal TotalPushOutNormal
 						= SolverManifoldPoint.NetPushOutNormal
@@ -1051,6 +1111,10 @@ namespace Chaos
 			for (int32 PointIndex = 0; PointIndex < NumManifoldPoints(); ++PointIndex)
 			{
 				FPBDCollisionSolverManifoldPoint& SolverManifoldPoint = State.ManifoldPoints[PointIndex];
+				if (!!SolverManifoldPoint.Flags.bDisablePositionSolve)
+				{
+					continue;
+				}
 
 				FSolverReal ContactCorrectionNormal;
 				CalculateContactPositionCorrectionNormal(SolverManifoldPoint, MaxPushOut, ContactCorrectionNormal);
@@ -1095,9 +1159,27 @@ namespace Chaos
 			for (int32 PointIndex = 0; PointIndex < NumManifoldPoints(); ++PointIndex)
 			{
 				FPBDCollisionSolverManifoldPoint& SolverManifoldPoint = State.ManifoldPoints[PointIndex];
+				if (!!SolverManifoldPoint.Flags.bDisableVelocitySolve)
+				{
+					continue;
+				}
 
-				const bool bShouldSolveNormalVelocity = (SolverManifoldPoint.NetPushOutNormal > FSolverReal(0));
-				const bool bShouldSolveTangentVelocity = (DynamicFriction > 0) && (bShouldSolveNormalVelocity || (SolverManifoldPoint.ContactDeltaNormal < State.SoftPhi) || (SolverManifoldPoint.NetSoftPushOutNormal > FSolverReal(0)));
+				// Solve for normal velocity only for points where we have a position correction
+				// (unless the position solve step is disabled, in which case we assume we want to run the velocity solve)
+				const bool bShouldSolveNormalVelocity = 
+					(SolverManifoldPoint.NetPushOutNormal > FSolverReal(0)) 
+					|| !!SolverManifoldPoint.Flags.bDisablePositionSolve;
+
+				// Solver for tangential velocity if we have non-zero friction and a position correction
+				// (unless the position solve step is disabled, in which case we assume we want to run the velocity solve)
+				const bool bShouldSolveTangentVelocity = 
+					(DynamicFriction > 0)
+					&& (
+						bShouldSolveNormalVelocity 
+						|| (SolverManifoldPoint.ContactDeltaNormal < State.SoftPhi) 
+						|| (SolverManifoldPoint.NetSoftPushOutNormal > FSolverReal(0)) 
+					) 
+					&& !SolverManifoldPoint.Flags.bDisableFriction;
 
 				if (bShouldSolveNormalVelocity || bShouldSolveTangentVelocity)
 				{
diff --git a/Engine/Source/Runtime/Experimental/Chaos/Public/Chaos/Collision/PBDCollisionSolverSettings.h b/Engine/Source/Runtime/Experimental/Chaos/Public/Chaos/Collision/PBDCollisionSolverSettings.h
--- a/Engine/Source/Runtime/Experimental/Chaos/Public/Chaos/Collision/PBDCollisionSolverSettings.h
+++ b/Engine/Source/Runtime/Experimental/Chaos/Public/Chaos/Collision/PBDCollisionSolverSettings.h
@@ -18,6 +18,7 @@ namespace Chaos
 			, NumVelocityFrictionIterations(1)
 			, NumPositionShockPropagationIterations(3)
 			, NumVelocityShockPropagationIterations(2)
+			, bUseSplitImpulse(false)
 		{
 		}
 
@@ -45,6 +46,10 @@ namespace Chaos
 
 		// How many velocity iterations should have shock propagation enabled
 		int32 NumVelocityShockPropagationIterations;
+
+		// Whether to use a split-impulse for the position correction. In split-impulse mode
+		// the position correction is momentumless
+		uint32 bUseSplitImpulse : 1;
 	};
 
 }	// namespace Chaos
diff --git a/Engine/Source/Runtime/Experimental/Chaos/Public/Chaos/ContactModification.h b/Engine/Source/Runtime/Experimental/Chaos/Public/Chaos/ContactModification.h
--- a/Engine/Source/Runtime/Experimental/Chaos/Public/Chaos/ContactModification.h
+++ b/Engine/Source/Runtime/Experimental/Chaos/Public/Chaos/ContactModification.h
@@ -142,6 +142,10 @@ namespace Chaos
 		*/
 		CHAOS_API void ModifyWorldContactLocations(const FVec3& Location0, const FVec3& Location1, int32 ContactPointIdx);
 
+		CHAOS_API void ModifyEnablePositionSolve(bool bEnable, int32 ContactPointIdx);
+		CHAOS_API void ModifyEnableVelocitySolve(bool bEnable, int32 ContactPointIdx);
+		CHAOS_API void ModifyEnableFriction(bool bEnable, int32 ContactPointIdx);
+
 		/**
 		* @return Restitution of contact.
 		*/
@@ -292,6 +296,14 @@ namespace Chaos
 		*/
 		CHAOS_API void SetContactPointDisabled(int32 ContactPointIdx) const;
 
+		/*
+		* Add a new contact point. WorldPos0 and WorldPos1 are the world-space contact positions on each particle.
+		* WorldNormal is the contact normal. The contact separation is implicitly (WorldPos0 - WorldPos1).WorldNormal.
+		* If WorldPos0 and WorldPos1 are not aligned along WorldNormal, then friction will be applied to attempt to 
+		* correct the tangential separation.
+		*/
+		CHAOS_API int32 AddContact(const FVec3& WorldPos0, const FVec3& WorldPos1, const FVec3& WorldNormal) const;
+
 	private:
 
 		/*
diff --git a/Engine/Source/Runtime/Experimental/Chaos/Public/Chaos/PBDCollisionConstraints.h b/Engine/Source/Runtime/Experimental/Chaos/Public/Chaos/PBDCollisionConstraints.h
--- a/Engine/Source/Runtime/Experimental/Chaos/Public/Chaos/PBDCollisionConstraints.h
+++ b/Engine/Source/Runtime/Experimental/Chaos/Public/Chaos/PBDCollisionConstraints.h
@@ -283,6 +283,11 @@ public:
 		SolverSettings.NumVelocityShockPropagationIterations = InNumIterations;
 	}
 
+	UE_INTERNAL void SetUseSplitImpulse(bool bInUseSplitImpulse)
+	{
+		SolverSettings.bUseSplitImpulse = bInUseSplitImpulse;
+	}
+
 	bool IsShockPropagationEnabled() const
 	{
 		return (SolverSettings.NumPositionShockPropagationIterations > 0) || (SolverSettings.NumVelocityShockPropagationIterations > 0);
