Geometry collection : add blueprint function to query current transforms
* GetCurrentTransforms
* GetCurrentTransformsFromIndices

diff --git a/Engine/Source/Runtime/Experimental/GeometryCollectionEngine/Private/GeometryCollection/GeometryCollectionComponent.cpp b/Engine/Source/Runtime/Experimental/GeometryCollectionEngine/Private/GeometryCollection/GeometryCollectionComponent.cpp
--- a/Engine/Source/Runtime/Experimental/GeometryCollectionEngine/Private/GeometryCollection/GeometryCollectionComponent.cpp
+++ b/Engine/Source/Runtime/Experimental/GeometryCollectionEngine/Private/GeometryCollection/GeometryCollectionComponent.cpp
@@ -7011,6 +7011,52 @@ void UGeometryCollectionComponent::SetLocalRestTransforms(const TArray<FTransfor
 	}
 }
 
+TArray<FTransform> UGeometryCollectionComponent::GetCurrentTransforms()
+{
+	TArray<FTransform> ReturnedTransforms;
+
+	const TArray<FTransform3f>& AllTransforms = ComponentSpaceTransforms.RequestAllTransforms();
+	ReturnedTransforms.Reserve(AllTransforms.Num());
+
+	const TManagedArray<FTransform>* MassToLocal = nullptr;
+	if (RestCollection && RestCollection->GetGeometryCollection())
+	{
+		MassToLocal = RestCollection->GetGeometryCollection()->FindAttribute<FTransform>("MassToLocal", FGeometryCollection::TransformGroup);
+	}
+
+	for (int32 TransformIndex = 0; TransformIndex < AllTransforms.Num(); ++TransformIndex)
+	{
+		const FTransform MassToLocalTransform = (MassToLocal && MassToLocal->IsValidIndex(TransformIndex)) ? (*MassToLocal)[TransformIndex] : FTransform::Identity;
+		ReturnedTransforms.Emplace(MassToLocalTransform * FTransform(AllTransforms[TransformIndex]));
+	}
+
+	return ReturnedTransforms;
+}
+
+TArray<FTransform> UGeometryCollectionComponent::GetCurrentTransformsFromIndices(const TArray<int32>& TransformIndices)
+{
+	TArray<FTransform> ReturnedTransforms;
+	ReturnedTransforms.Reserve(TransformIndices.Num());
+
+	const TArray<FTransform3f>& AllTransforms = ComponentSpaceTransforms.RequestAllTransforms();
+
+	const TManagedArray<FTransform>* MassToLocal = nullptr;
+	if (RestCollection && RestCollection->GetGeometryCollection())
+	{
+		MassToLocal = RestCollection->GetGeometryCollection()->FindAttribute<FTransform>("MassToLocal", FGeometryCollection::TransformGroup);
+	}
+
+	for (const int32 TransformIndex : TransformIndices)
+	{
+		const FTransform MassToLocalTransform = (MassToLocal && MassToLocal->IsValidIndex(TransformIndex)) ? (*MassToLocal)[TransformIndex] : FTransform::Identity;
+		const FTransform Transform = (AllTransforms.IsValidIndex(TransformIndex)) ? FTransform(AllTransforms[TransformIndex]) : FTransform::Identity;
+		ReturnedTransforms.Emplace(MassToLocalTransform * Transform);
+
+	}
+
+	return ReturnedTransforms;
+}
+
 TArray<FMatrix> UGeometryCollectionComponent::ComputeGlobalMatricesFromComponentSpaceTransforms() const
 {
 	TArray<FMatrix> ComponentSpaceMatrices;
diff --git a/Engine/Source/Runtime/Experimental/GeometryCollectionEngine/Public/GeometryCollection/GeometryCollectionComponent.h b/Engine/Source/Runtime/Experimental/GeometryCollectionEngine/Public/GeometryCollection/GeometryCollectionComponent.h
--- a/Engine/Source/Runtime/Experimental/GeometryCollectionEngine/Public/GeometryCollection/GeometryCollectionComponent.h
+++ b/Engine/Source/Runtime/Experimental/GeometryCollectionEngine/Public/GeometryCollection/GeometryCollectionComponent.h
@@ -1336,6 +1336,20 @@ public:
 
 	GEOMETRYCOLLECTIONENGINE_API const TArray<FTransform3f>& GetComponentSpaceTransforms3f();
 
+	/**
+	* Get current transforms
+	* The transforms are representing the center of mass of each piece and are in component space 
+	*/
+	UFUNCTION(BlueprintCallable, Category = "ChaosPhysics")
+	GEOMETRYCOLLECTIONENGINE_API TArray<FTransform> GetCurrentTransforms();
+
+	/**
+	* Get current transforms from a specific list of indices
+	* The transforms are representing the center of mass of each piece and are in component space
+	*/
+	UFUNCTION(BlueprintCallable, Category = "ChaosPhysics")
+	GEOMETRYCOLLECTIONENGINE_API TArray<FTransform> GetCurrentTransformsFromIndices(const TArray<int32>& TransformIndices);
+
 	GEOMETRYCOLLECTIONENGINE_API const FGeometryDynamicCollection* GetDynamicCollection() const;
 	GEOMETRYCOLLECTIONENGINE_API FGeometryDynamicCollection* GetDynamicCollection();  // TEMP HACK?
 
