[ModularVehicles] Add Average input type to handle both inter and intra frame averaging.

diff --git a/Engine/Plugins/Experimental/ChaosModularVehicle/Source/ChaosModularVehicleEngine/Private/ChaosModularVehicle/InputProducer.cpp b/Engine/Plugins/Experimental/ChaosModularVehicle/Source/ChaosModularVehicleEngine/Private/ChaosModularVehicle/InputProducer.cpp
--- a/Engine/Plugins/Experimental/ChaosModularVehicle/Source/ChaosModularVehicleEngine/Private/ChaosModularVehicle/InputProducer.cpp
+++ b/Engine/Plugins/Experimental/ChaosModularVehicle/Source/ChaosModularVehicleEngine/Private/ChaosModularVehicle/InputProducer.cpp
@@ -32,6 +32,12 @@ void UVehicleDefaultInputProducer::BufferInput(const FInputNameMap& InNameMap, c
 			Inputs.CombineValue(InName, InValue);
 			break;
 		}
+	case EModuleInputBufferActionType::Average:
+		{
+			FInputInterface Inputs(InNameMap, MergedInput, InputQuantizationType);
+			Inputs.AverageValue(InName, InValue);
+			break;
+		}
 	default:
 		break;
 	}
diff --git a/Engine/Source/Runtime/Experimental/ChaosVehicles/ChaosVehiclesCore/Private/SimModule/ModuleInput.cpp b/Engine/Source/Runtime/Experimental/ChaosVehicles/ChaosVehiclesCore/Private/SimModule/ModuleInput.cpp
--- a/Engine/Source/Runtime/Experimental/ChaosVehicles/ChaosVehiclesCore/Private/SimModule/ModuleInput.cpp
+++ b/Engine/Source/Runtime/Experimental/ChaosVehicles/ChaosVehiclesCore/Private/SimModule/ModuleInput.cpp
@@ -155,15 +155,16 @@ void FModuleInputValue::SetMagnitude(float NewSize)
 		Value.X = NewSize;
 		break;
 	case EModuleInputValueType::MAxis2D:
-		Value = Value.GetSafeNormal2D()* NewSize;
+		Value = Value.GetSafeNormal2D(UE_SMALL_NUMBER,FVector::UnitX()) * NewSize;
 		break;
 	case EModuleInputValueType::MAxis3D:
-		Value = Value.GetSafeNormal()* NewSize;
+		Value = Value.GetSafeNormal(UE_SMALL_NUMBER,FVector::UnitX())* NewSize;
 		break;
 	default:
 		checkf(false, TEXT("Unsupported value type for module input value!"));
 		break;
 	}
+	bHasValue = true;
 }
 
 void FModuleInputValue::Quantize(double& InOutValue, EModuleInputQuantizationType InInputQuantizationType)
@@ -212,15 +213,13 @@ bool FModuleInputValue::SerializeQuantized(double& InOutValue, FArchive& Ar, EMo
 
 FModuleInputValue FModuleInputValue::ReturnQuantized(EModuleInputQuantizationType InInputQuantizationType) const
 {
-	FModuleInputValue OutValue = FModuleInputValue(ValueType, Value);
-
+	FModuleInputValue OutValue = FModuleInputValue(ValueType);
+	OutValue.Set(*this);
 	switch (ValueType)
 	{
-		case EModuleInputValueType::MBoolean:
 		case EModuleInputValueType::MInteger:
-			OutValue = FModuleInputValue(ValueType, ValueInt);
-			break;
-
+		case EModuleInputValueType::MBoolean:
+		break;
 		case EModuleInputValueType::MAxis1D:
 			Quantize(OutValue.Value.X, InInputQuantizationType);
 			break;
@@ -255,6 +254,7 @@ void FModuleInputValue::Serialize(FArchive& Ar, UPackageMap* Map, bool& bOutSucc
 	Ar.SerializeBits(&ValueType, 3);
 	Ar.SerializeBits(&bApplyInputDecay, 1);
 	Ar.SerializeBits(&bClearAfterConsumed, 1);
+	Ar.SerializeBits(&bHasValue, 1);
 
 
 	if (bIsNonZero)
@@ -303,6 +303,7 @@ void FModuleInputValue::NetSerialize(FArchive& Ar, UPackageMap* Map, bool& bOutS
 {
 	bool bIsNonZero = Ar.IsLoading() ? false : IsQuantizedNonZero(InInputQuantizationType);
 	Ar.SerializeBits(&bIsNonZero, 1);
+	Ar.SerializeBits(&bHasValue, 1);
 	if (ChaosModularVehicleCVars::bEnableInputSendingTypeInfo)
 	{
 		Ar.SerializeBits(&ValueType, 3);
@@ -371,6 +372,7 @@ void FModuleInputValue::DeltaNetSerialize(FArchive& Ar, UPackageMap* Map, bool&
 		ValueInt = PreviousInputValue.ValueInt;
 		bApplyInputDecay = PreviousInputValue.bApplyInputDecay;
 		bClearAfterConsumed = PreviousInputValue.bClearAfterConsumed;
+		bHasValue = PreviousInputValue.bHasValue;
 		return;
 	}
 
@@ -382,6 +384,7 @@ void FModuleInputValue::Set(const FModuleInputValue& In)
 {
 	ValueInt = In.ValueInt;
 	Value = In.Value;
+	bHasValue = In.bHasValue;
 }
 
 void FModuleInputValue::Lerp(const FModuleInputValue& Min, const FModuleInputValue& Max, float Alpha)
@@ -397,6 +400,7 @@ void FModuleInputValue::Lerp(const FModuleInputValue& Min, const FModuleInputVal
 		Value = FMath::Lerp(Min.Value, Max.Value, Alpha);
 		break;
 	}
+	bHasValue = true;
 }
 
 void FModuleInputValue::Merge(const FModuleInputValue& From)
@@ -425,6 +429,7 @@ void FModuleInputValue::Merge(const FModuleInputValue& From)
 		checkf(false, TEXT("Unsupported value type for module input value!"));
 		break;
 	}
+	bHasValue = true;
 }
 
 void FModuleInputValue::Combine(const FModuleInputValue& With)
@@ -452,6 +457,42 @@ void FModuleInputValue::Combine(const FModuleInputValue& With)
 		checkf(false, TEXT("Unsupported value type for module input value!"));
 		break;
 	}
+	bHasValue = true;
+}
+
+void FModuleInputValue::Average(const FModuleInputValue& With)
+{
+	if (!bHasValue)
+	{
+		Set(With);
+		return;
+	}
+	switch (GetValueType())
+	{
+	case EModuleInputValueType::MBoolean:
+		// ensure we capture edges of digital inputs by taking largest absolute value
+		if (FMath::Abs(With.ValueInt) >= FMath::Abs(ValueInt))
+		{
+			ValueInt = With.ValueInt;
+		}
+		break;
+
+	case EModuleInputValueType::MInteger:
+		ValueInt += With.ValueInt;
+		ValueInt *= 0.5f;
+		break;
+
+	case EModuleInputValueType::MAxis1D:
+	case EModuleInputValueType::MAxis2D:
+	case EModuleInputValueType::MAxis3D:
+		Value += With.Value;
+		Value *= 0.5f;
+		break;
+	default:
+		checkf(false, TEXT("Unsupported value type for module input value!"));
+		break;
+	}
+	bHasValue = true;
 }
 
 void FModuleInputValue::Decay(const float DecayAmount)
@@ -600,23 +641,12 @@ void FModuleInputContainer::Serialize(FArchive& Ar, UPackageMap* Map, bool& bOut
 	{
 		InputValues[I].Serialize(Ar, Map, bOutSuccess, InInputQuantizationType);
 	}
-
 	return;
 }
 
 int FModuleInputContainer::AddInput(EModuleInputValueType Type)
 {
-	switch(Type)
-	{
-		case EModuleInputValueType::MBoolean:
-		case EModuleInputValueType::MInteger:
-			InputValues.Add(FModuleInputValue(Type, 0));
-		break;
-
-		default:
-			InputValues.Add(FModuleInputValue(Type, FVector::ZeroVector));
-		break;
-	}
+	InputValues.Add(FModuleInputValue(Type));
 	return InputValues.Num() - 1;
 }
 
@@ -700,6 +730,21 @@ void FInputInterface::CombineValue(const FName& InName, const FModuleInputValue&
 	}
 }
 
+void FInputInterface::AverageValue(const FName& InName, const FModuleInputValue& InValue)
+{
+	if (ValueContainer.GetNumInputs() > 0)
+	{
+		if (const int* Index = NameMap.Find(InName))
+		{
+			ValueContainer.AverageValueAtIndex(*Index, InValue, InputQuantizationType);
+		}
+		else
+		{
+			UE_LOG(LogModularInput, Warning, TEXT("Trying to set the value of an undefined control input %s"), *InName.ToString());
+		}
+	}
+}
+
 void FInputInterface::MergeValue(const FName& InName, const FModuleInputValue& InValue)
 {
 	if (ValueContainer.GetNumInputs() > 0)
@@ -730,7 +775,7 @@ FModuleInputValue FInputInterface::GetValue(const FName& InName) const
 		}
 	}
 
-	return FModuleInputValue(EModuleInputValueType::MBoolean, FVector::ZeroVector);
+	return FModuleInputValue(false);
 }
 
 EModuleInputValueType FInputInterface::GetValueType(const FName& InName) const
diff --git a/Engine/Source/Runtime/Experimental/ChaosVehicles/ChaosVehiclesCore/Public/SimModule/ModuleInput.h b/Engine/Source/Runtime/Experimental/ChaosVehicles/ChaosVehiclesCore/Public/SimModule/ModuleInput.h
--- a/Engine/Source/Runtime/Experimental/ChaosVehicles/ChaosVehiclesCore/Public/SimModule/ModuleInput.h
+++ b/Engine/Source/Runtime/Experimental/ChaosVehicles/ChaosVehiclesCore/Public/SimModule/ModuleInput.h
@@ -47,7 +47,10 @@ enum class EModuleInputBufferActionType : uint8
 	Override = 0,
 
 	// Adds new value to current value in buffer
-	Combine
+	Combine,
+	
+	// A running average of the inputs.
+	Average
 };
 
 namespace ModularQuantize
@@ -197,14 +200,14 @@ public:
 	FModuleInputValue() = default;
 	FModuleInputValue(const FModuleInputValue&) = default;
 	FModuleInputValue& operator= (const FModuleInputValue&) = default;
-
+	
 	// Specialized constructors for supported types
 	// Converting a value to a different type (e.g. Val = FVector(1, 1, 1); Val = true;) zeroes out any unused components to ensure getters continue to function correctly.
-	FModuleInputValue(bool bInValue) : Value(FVector::ZeroVector), ValueInt(bInValue), ValueType(EModuleInputValueType::MBoolean) {}
-	FModuleInputValue(MInteger InValue) : Value(FVector::ZeroVector), ValueInt(InValue), ValueType(EModuleInputValueType::MInteger) {}
-	FModuleInputValue(MAxis1D InValue) : Value(InValue, 0.f, 0.f), ValueInt(0), ValueType(EModuleInputValueType::MAxis1D) {}
-	FModuleInputValue(MAxis2D InValue) : Value(InValue.X, InValue.Y, 0.f), ValueInt(0), ValueType(EModuleInputValueType::MAxis2D) {}
-	FModuleInputValue(MAxis3D InValue) : Value(InValue), ValueInt(0), ValueType(EModuleInputValueType::MAxis3D) {}
+	FModuleInputValue(bool bInValue) : bHasValue(true), Value(FVector::ZeroVector), ValueInt(bInValue), ValueType(EModuleInputValueType::MBoolean)  {}
+	FModuleInputValue(MInteger InValue) : bHasValue(true), Value(FVector::ZeroVector), ValueInt(InValue), ValueType(EModuleInputValueType::MInteger) {}
+	FModuleInputValue(MAxis1D InValue) : bHasValue(true), Value(InValue, 0.f, 0.f), ValueInt(0), ValueType(EModuleInputValueType::MAxis1D) {}
+	FModuleInputValue(MAxis2D InValue) : bHasValue(true), Value(InValue.X, InValue.Y, 0.f), ValueInt(0), ValueType(EModuleInputValueType::MAxis2D) {}
+	FModuleInputValue(MAxis3D InValue) : bHasValue(true), Value(InValue), ValueInt(0), ValueType(EModuleInputValueType::MAxis3D) {}
 
 	FModuleInputValue ReturnQuantized(EModuleInputQuantizationType InInputQuantizationType) const;
 
@@ -212,34 +215,12 @@ public:
 	
 	static bool SerializeQuantized(double& InOutValue, FArchive& Ar, EModuleInputQuantizationType InInputQuantizationType);
 	
-	// Build a specific type with an arbitrary Axis3D value
-	FModuleInputValue(EModuleInputValueType InValueType, MAxis3D InValue) : Value(InValue), ValueType(InValueType)
+	// Build a specific type with no value
+	FModuleInputValue(EModuleInputValueType InValueType)
 	{
+		ValueType = InValueType;
 		ValueInt = 0.0f; // not used in this case
-
-		// Clear out value components to match type
-		switch (ValueType)
-		{
-		case EModuleInputValueType::MBoolean:
-		case EModuleInputValueType::MAxis1D:
-			Value.Y = 0.f;
-			//[[fallthrough]];
-		case EModuleInputValueType::MAxis2D:
-			Value.Z = 0.f;
-			//[[fallthrough]];
-		case EModuleInputValueType::MAxis3D:
-		default:
-			return;
-		}
-	}
-
-	// Build a specific type with an Integer value
-	FModuleInputValue(EModuleInputValueType InValueType, MInteger InValue) : ValueInt(InValue), ValueType(InValueType)
-	{
-		// not used in thos case clear for good measure
-		Value.X = 0.f;
-		Value.Y = 0.f;
-		Value.Z = 0.f;
+		Value = FVector::ZeroVector;
 	}
 
 	// Resets Value without affecting ValueType
@@ -247,6 +228,7 @@ public:
 	{
 		Value = FVector::ZeroVector;
 		ValueInt = 0;
+		bHasValue = false;
 	}
 
 	FModuleInputValue& operator+=(const FModuleInputValue& Rhs)
@@ -339,7 +321,8 @@ public:
 	{
 		if (ValueType != Type)
 		{
-			*this = FModuleInputValue(Type, Value);
+			Reset();
+			ValueType = Type;
 		}
 		return *this;
 	}
@@ -361,8 +344,8 @@ public:
 	UE_API void Lerp(const FModuleInputValue& Min, const FModuleInputValue& Max, float Alpha);
 
 	UE_API void Merge(const FModuleInputValue& From);
-
 	UE_API void Combine(const FModuleInputValue& With);
+	UE_API void Average(const FModuleInputValue& With);
 
 	UE_API void Decay(const float DecayAmount);
 
@@ -380,6 +363,9 @@ public:
 protected:
 	void SetMagnitude(float NewSize);
 
+	UPROPERTY(Transient)
+	bool bHasValue = false;
+	
 	UPROPERTY()
 	FVector Value = FVector::ZeroVector;
 
@@ -611,6 +597,10 @@ public:
 		}
 	}
 
+	void AverageValueAtIndex(int Index, const FModuleInputValue& InValue, EModuleInputQuantizationType InInputQuantizationType)
+	{
+		InputValues[Index].Average(InValue.ReturnQuantized(InInputQuantizationType));
+	}
 	void MergeValueAtIndex(int Index, const FModuleInputValue& InValue, EModuleInputQuantizationType InInputQuantizationType)
 	{
 		InputValues[Index].Merge(InValue.ReturnQuantized(InInputQuantizationType));
@@ -681,6 +671,7 @@ public:
 
 	UE_API void SetValue(const FName& InName, const FModuleInputValue& InValue, bool Quantize = true);
 	UE_API void CombineValue(const FName& InName, const FModuleInputValue& InValue);
+	UE_API void AverageValue(const FName& InName, const FModuleInputValue& InValue);
 	UE_API void MergeValue(const FName& InName, const FModuleInputValue& InValue);
 	UE_API FModuleInputValue GetValue(const FName& InName) const;
 	UE_API EModuleInputValueType GetValueType(const FName& InName) const;
@@ -718,7 +709,7 @@ public:
 	void SetBool(const FName& InName, bool InBool)
 	{
 		ensure(GetValueType(InName) == EModuleInputValueType::MBoolean);
-		FModuleInputValue Value(EModuleInputValueType::MBoolean, InBool);
+		FModuleInputValue Value(InBool);
 		SetValue(InName, Value);
 	}
 	void SetInteger(const FName& InName, int32 InInteger)
