Add soem additional safety checks when updating the ondemand package store backend.

#rnx

- Added asserts to make sure that the data in ReferencedChunkEntryIndices matches the containers when updating the lookup tables.
-- ReferencedChunkEntryIndices will have just been created from the current set of containers so everything should match but these checks may help catch unintended changes in the future.
-- Technically this is a follow on from CL 46757514 where I added similar checks to ::UpdateReferencedPackages
- Added FOnDemandContainer::GetNumChunks to avoid accessing members directly.

diff --git a/Engine/Source/Runtime/Experimental/IoStore/OnDemand/Private/OnDemandIoStore.h b/Engine/Source/Runtime/Experimental/IoStore/OnDemand/Private/OnDemandIoStore.h
--- a/Engine/Source/Runtime/Experimental/IoStore/OnDemand/Private/OnDemandIoStore.h
+++ b/Engine/Source/Runtime/Experimental/IoStore/OnDemand/Private/OnDemandIoStore.h
@@ -97,11 +97,14 @@ struct FOnDemandContainerChunkEntryReferences
 struct FOnDemandContainer
 {
 	FString									UniqueName() const;
+
+	inline int32							GetNumChunks() const;
 	inline int32							FindChunkEntryIndex(const FIoChunkId& ChunkId) const;
 	inline const FOnDemandChunkEntry*		FindChunkEntry(const FIoChunkId& ChunkId, int32* OutIndex = nullptr) const;
 	inline FOnDemandChunkEntryReferences&	FindOrAddChunkEntryReferences(const FOnDemandInternalContentHandle& ContentHandle);
 	inline TBitArray<>						GetReferencedChunkEntries() const;
 	inline bool								IsReferenced(int32 ChunkEntryIndex) const;
+
 	inline FAnsiString						GetTestUrl();
 	bool									HasAnyFlags(EOnDemandContainerFlags Contains) const { return EnumHasAnyFlags(Flags, Contains); }
 	bool									HasAllFlags(EOnDemandContainerFlags Contains) const { return EnumHasAllFlags(Flags, Contains); }
@@ -128,6 +131,11 @@ struct FOnDemandContainer
 	EOnDemandContainerFlags 				Flags = EOnDemandContainerFlags::None;
 };
 
+int32 FOnDemandContainer::GetNumChunks() const
+{
+	return ChunkIds.Num();
+}
+
 int32 FOnDemandContainer::FindChunkEntryIndex(const FIoChunkId& ChunkId) const
 {
 	if (const int32 Index = Algo::LowerBound(ChunkIds, ChunkId); Index < ChunkIds.Num())
diff --git a/Engine/Source/Runtime/Experimental/IoStore/OnDemand/Private/OnDemandPackageStoreBackend.cpp b/Engine/Source/Runtime/Experimental/IoStore/OnDemand/Private/OnDemandPackageStoreBackend.cpp
--- a/Engine/Source/Runtime/Experimental/IoStore/OnDemand/Private/OnDemandPackageStoreBackend.cpp
+++ b/Engine/Source/Runtime/Experimental/IoStore/OnDemand/Private/OnDemandPackageStoreBackend.cpp
@@ -227,6 +227,7 @@ TConstArrayView<uint32> FOnDemandPackageStoreBackend::GetSoftReferences(FPackage
 void FOnDemandPackageStoreBackend::UpdateLookupTables(const TArray<FSharedOnDemandContainer>& AllContainers, const TArray<TBitArray<>>& ReferencedChunkEntryIndices)
 {
 	check(Mutex.IsLocked());
+	check(ReferencedChunkEntryIndices.Num() ==  AllContainers.Num());
 
 	Containers.Empty(AllContainers.Num());
 
@@ -252,6 +253,8 @@ void FOnDemandPackageStoreBackend::UpdateLookupTables(const TArray<FSharedOnDema
 			Container.SoftRefs.Reserve(Hdr.PackageIds.Num());
 		}
 
+		check(ReferencedChunkEntryIndices[ContainerIdx].Num() == SharedContainer->GetNumChunks());
+
 		for (int32 EntryIdx = 0; const FFilePackageStoreEntry& Entry : Entries)
 		{
 			const FPackageId PkgId = Hdr.PackageIds[EntryIdx];
@@ -339,7 +342,7 @@ void FOnDemandPackageStoreBackend::UpdateReferencedPackages(const TArray<TBitArr
 			}
 			else
 			{
-				ensure(Container.Container->ChunkEntries.Num() == Refs.Num());
+				ensure(Container.Container->GetNumChunks() == Refs.Num());
 				for (TPair<FPackageId, FOnDemandPackageStoreEntry>& Kv : Container.EntryMap)
 				{
 					Kv.Value.IsReferenced = Refs[Kv.Value.ChunkEntryIndex];
