We now assume that FChunkPluginStats will have both IAD and IAS entries.

#rnx

- Changed FChunkPluginStats::TotalChunksSizeKb to ::TotalIADChunksSizeKb to make it clear which system it is for.
- Fixed a bug where FChunkPluginStats::TotalIADChunksSizeKb was counting from -1, meaning we'd be off by 1.

diff --git a/Engine/Source/Runtime/Experimental/IoStore/OnDemand/Private/Tool/ChunkPlugin.cpp b/Engine/Source/Runtime/Experimental/IoStore/OnDemand/Private/Tool/ChunkPlugin.cpp
--- a/Engine/Source/Runtime/Experimental/IoStore/OnDemand/Private/Tool/ChunkPlugin.cpp
+++ b/Engine/Source/Runtime/Experimental/IoStore/OnDemand/Private/Tool/ChunkPlugin.cpp
@@ -34,10 +34,10 @@ namespace UE::IoStore::Tool
 
 struct FChunkPluginStats : FJsonSerializable
 {
-	double TotalChunksSizeKb = -1.0;
+	double TotalIADChunksSizeKb = 0.0;
 
 BEGIN_JSON_SERIALIZER
-	JSON_SERIALIZE("TotalChunksSizeKb", TotalChunksSizeKb);
+	JSON_SERIALIZE("TotalIADChunksSizeKb", TotalIADChunksSizeKb);
 END_JSON_SERIALIZER
 };
 
@@ -484,7 +484,7 @@ static int32 ChunkContainer(const FString& ContainerUTocPath, const FChunkPlugin
 		TocEntry.BlockOffset = BlockOffset;
 		TocEntry.BlockCount = BlockCount;
 
-		OutStats.TotalChunksSizeKb += (double)EncodedChunkSize / 1024.0;
+		OutStats.TotalIADChunksSizeKb += (double)EncodedChunkSize / 1024.0;
 	}
 
 	if (bIgnoreContainerHeader)
