Fix some minor risks with string view use.

#rnx

- We should not pass a dynamic string to ::WriteLine as the string format as this could break if the string contains a %. Instead we should use "%" as the format and pass the string in as an arg.
- We cannot safely pass FStringView to printf style formatting as a string pointer as it may not be null terminated. Although we currently know that the string will be null terminated this might change in the future and then the code could start breaking. We need to place the string view into a string builder so we know it will be terminated.
-- If this were too slow we could consider structured logging which can handle string views safely, but given the use case for the code there is no real reason to care about the additional string copy and we should just er on the side of caution.
- Use *FStringView, rather than .ToString as part of a long term project to deprecate ToString.

diff --git a/Engine/Source/Runtime/Experimental/IoStore/OnDemand/Private/Tool/Command.cpp b/Engine/Source/Runtime/Experimental/IoStore/OnDemand/Private/Tool/Command.cpp
--- a/Engine/Source/Runtime/Experimental/IoStore/OnDemand/Private/Tool/Command.cpp
+++ b/Engine/Source/Runtime/Experimental/IoStore/OnDemand/Private/Tool/Command.cpp
@@ -385,7 +385,7 @@ void FCommand::Usage() const
 {
 	if (!Desc.IsEmpty())
 	{
-		WriteLine(Desc.GetData());
+		WriteLine(TEXT("%s"), *WriteToString<128>(Desc));
 		WriteLine(TEXT(""));
 	}
 
@@ -413,7 +413,7 @@ void FCommand::Usage() const
 	}
 
 	WriteLine(TEXT("Usage:"));
-	WriteLine(TEXT("  %s"), UsageLine.ToString());
+	WriteLine(TEXT("  %s"), *UsageLine);
 
 	if (!bHasOpts)
 	{
