Add a missing error check when writing out the results for IASTool Graph command.

#rnx

- Make sure we check the result of FArchive::Close when writing out the csv files to make sure that they were completed correctly.

diff --git a/Engine/Source/Runtime/Experimental/IoStore/OnDemand/Private/Tool/Graph.cpp b/Engine/Source/Runtime/Experimental/IoStore/OnDemand/Private/Tool/Graph.cpp
--- a/Engine/Source/Runtime/Experimental/IoStore/OnDemand/Private/Tool/Graph.cpp
+++ b/Engine/Source/Runtime/Experimental/IoStore/OnDemand/Private/Tool/Graph.cpp
@@ -44,7 +44,13 @@ FIoStatus ExportGephiCsv(const FPackageGraph& Graph, const FString& OutFilename)
 			const FPackageGraph::FNodeData&	Data = Graph.GetNodeData(Node);
 			Ar->Logf(TEXT("%d, %s"), Idx, *LexToString(Data.PackageId));
 		}
+
+		if (!Ar->Close())
+		{
+			return FIoStatus(EIoErrorCode::WriteError);
+		}
 	}
+
 	{
 		FString Filename		= FPaths::GetBaseFilename(OutFilename);
 		FString Path			= FPaths::GetPath(OutFilename);
@@ -65,6 +71,11 @@ FIoStatus ExportGephiCsv(const FPackageGraph& Graph, const FString& OutFilename)
 			const TCHAR*					Type	= TEXT("Directed");
 			Ar->Logf(TEXT("%d, %d, %s, %.2f, %d"), Edge.From.Index, Edge.To.Index, Type, Weight, Idx);
 		}
+
+		if (!Ar->Close())
+		{
+			return FIoStatus(EIoErrorCode::WriteError);
+		}
 	}
 
 	return FIoStatus::Ok;
