ISMPoolComponent precaching fix follow up: simplify the precaching logic by removing the bSkipPSOPrecaching flag

diff --git a/Engine/Source/Runtime/Engine/Classes/Components/PrimitiveComponent.h b/Engine/Source/Runtime/Engine/Classes/Components/PrimitiveComponent.h
--- a/Engine/Source/Runtime/Engine/Classes/Components/PrimitiveComponent.h
+++ b/Engine/Source/Runtime/Engine/Classes/Components/PrimitiveComponent.h
@@ -708,9 +708,6 @@ public:
 	/** PSOs requested priority */
 	EPSOPrecachePriority PSOPrecacheRequestPriority : 2;
 	static_assert((int)EPSOPrecachePriority::Highest < 1 << 2);
-
-	/** Whether the PrecachePSOs call in SetStaticMesh should be skipped. */
-	uint8 bSkipPSOPrecaching : 1;
 #endif
 
 #if WITH_EDITOR
diff --git a/Engine/Source/Runtime/Engine/Private/Components/PrimitiveComponent.cpp b/Engine/Source/Runtime/Engine/Private/Components/PrimitiveComponent.cpp
--- a/Engine/Source/Runtime/Engine/Private/Components/PrimitiveComponent.cpp
+++ b/Engine/Source/Runtime/Engine/Private/Components/PrimitiveComponent.cpp
@@ -409,7 +409,6 @@ UPrimitiveComponent::UPrimitiveComponent(const FObjectInitializer& ObjectInitial
 #if UE_WITH_PSO_PRECACHING
 	bPSOPrecacheCalled = false;
 	PSOPrecacheRequestPriority = EPSOPrecachePriority::Medium;
-	bSkipPSOPrecaching = false;
 	bPrecachedPSOs = false;
 #endif // UE_WITH_PSO_PRECACHING
 	
@@ -5103,7 +5102,7 @@ void UPrimitiveComponent::PrecachePSOs()
 	// Also only request PSOs from game thread because TStrongObjectPtr is used on the material to make
 	// it's not deleted via garbage collection when PSO precaching is still busy. TStrongObjectPtr can only
 	// be constructed on the GameThread
-	if (!FApp::CanEverRender() || !IsComponentPSOPrecachingEnabled() || bSkipPSOPrecaching)
+	if (!FApp::CanEverRender() || !IsComponentPSOPrecachingEnabled())
 	{
 		return;
 	}
diff --git a/Engine/Source/Runtime/Experimental/ISMPool/Private/ISMPool/ISMPoolComponent.cpp b/Engine/Source/Runtime/Experimental/ISMPool/Private/ISMPool/ISMPoolComponent.cpp
--- a/Engine/Source/Runtime/Experimental/ISMPool/Private/ISMPool/ISMPoolComponent.cpp
+++ b/Engine/Source/Runtime/Experimental/ISMPool/Private/ISMPool/ISMPoolComponent.cpp
@@ -168,12 +168,6 @@ void FISMPoolISM::InitISM(const FISMPoolStaticMeshInstance& InMeshInstance, bool
 		bDisallowNanite |= Material->GetBlendMode() == BLEND_Translucent;
 	}
 
-#if UE_WITH_PSO_PRECACHING
-	// disable the implicit precachePSO in SetStaticMesh, so we can explicitly call it later to make sure we precache the correct state
-	ISMComponent->bSkipPSOPrecaching = true;
-#endif
-
-	ISMComponent->SetStaticMesh(StaticMesh);
 	ISMComponent->SetMobility((MeshInstance.Desc.Flags & FISMPoolComponentDescription::StaticMobility) != 0 ? EComponentMobility::Static : EComponentMobility::Movable);
 
 	ISMComponent->NumCustomDataFloats = MeshInstance.Desc.NumCustomDataFloats;
@@ -237,10 +231,8 @@ void FISMPoolISM::InitISM(const FISMPoolStaticMeshInstance& InMeshInstance, bool
 	// A possible improvement would be to compute a hash from the owner component and use that as the seed.
 	ISMComponent->InstancingRandomSeed = 1;	
 
-#if UE_WITH_PSO_PRECACHING
-	ISMComponent->bSkipPSOPrecaching = false;
-	ISMComponent->PrecachePSOs();
-#endif
+	// Make sure this is called at the end so the full state is set, since SetStaticMesh implicitly calls PrecachePSO()
+	ISMComponent->SetStaticMesh(StaticMesh);
 }
 
 FISMPoolInstanceGroups::FInstanceGroupId FISMPoolISM::AddInstanceGroup(int32 InstanceCount, TArrayView<const float> CustomDataFloats)
