Fix potential divide by 0 when logging stats for IASTool Download

#rnx

- Fix typo in the variable name.
- Make sure we do not divide by zero if nothing got downloaded.

diff --git a/Engine/Source/Runtime/Experimental/IoStore/OnDemand/Private/Tool/Download.cpp b/Engine/Source/Runtime/Experimental/IoStore/OnDemand/Private/Tool/Download.cpp
--- a/Engine/Source/Runtime/Experimental/IoStore/OnDemand/Private/Tool/Download.cpp
+++ b/Engine/Source/Runtime/Experimental/IoStore/OnDemand/Private/Tool/Download.cpp
@@ -284,10 +284,10 @@ FIoStatus DownloadContainerFiles(const FDownloadParams& DownloadParams, const FS
 
 			if (Stats.CompressionMethod != NAME_None)
 			{
-				double Procentage = (double(Stats.RawSize - Stats.CompressedSize) / double(Stats.RawSize)) * 100.0;
+				const double Percentage = Stats.RawSize > 0 ? (double(Stats.RawSize - Stats.CompressedSize) / double(Stats.RawSize)) * 100.0 : 0.0;
 				CompressionInfo = FString::Printf(TEXT("%.2lf (%.2lf%% %s)"),
 					(double)Stats.CompressedSize / 1024.0 / 1024.0,
-					Procentage,
+					Percentage,
 					*Stats.CompressionMethod.ToString());
 			}
 
