[Chaos] Ensure that rotating kinematic objects get their angular velocity updated when their orientation is close to 180 degrees

diff --git a/Engine/Source/Runtime/Experimental/Chaos/Private/Chaos/PBDRigidsEvolutionGBF.cpp b/Engine/Source/Runtime/Experimental/Chaos/Private/Chaos/PBDRigidsEvolutionGBF.cpp
--- a/Engine/Source/Runtime/Experimental/Chaos/Private/Chaos/PBDRigidsEvolutionGBF.cpp
+++ b/Engine/Source/Runtime/Experimental/Chaos/Private/Chaos/PBDRigidsEvolutionGBF.cpp
@@ -1217,7 +1217,7 @@ void FPBDRigidsEvolutionGBF::DispatchKinematicTargetsTasks(int32 DispatchBatchIn
 					}
 
 					const bool bPositionChanged = !FVec3::IsNearlyEqual(NewX, CurrentX, UE_SMALL_NUMBER);
-					const bool bRotationChanged = !FRotation3::IsNearlyEqual(NewR, CurrentR, UE_KINDA_SMALL_NUMBER);
+					const bool bRotationChanged = !FRotation3::IsNearlyEqual(NewR, CurrentR, UE_SMALL_NUMBER);
 					bMoved = bPositionChanged || bRotationChanged;
 					FVec3 NewV = FVec3(0);
 					FVec3f NewW = FVec3f(0);
