Improve error handling when writing out ondemand tocs during IASTool Upload

#rnx

- Even though we were checking if the archive was valid, it was still possible for it to fail when flushing to disk. By checking the return value of FArchive::Close we can be sure that we have caught all error cases.
- Added some code documentation to explain the use of a string builder without inline storage.

diff --git a/Engine/Source/Runtime/Experimental/IoStore/OnDemand/Private/Tool/Upload.cpp b/Engine/Source/Runtime/Experimental/IoStore/OnDemand/Private/Tool/Upload.cpp
--- a/Engine/Source/Runtime/Experimental/IoStore/OnDemand/Private/Tool/Upload.cpp
+++ b/Engine/Source/Runtime/Experimental/IoStore/OnDemand/Private/Tool/Upload.cpp
@@ -231,7 +231,7 @@ static FIoStatus WriteContainerFiles(FOnDemandToc& OnDemandToc, const TMap<FIoHa
 		FOnDemandTocSentinel Sentinel;
 		(*Ar) << Sentinel;
 
-		if (Ar->IsError() || Ar->IsCriticalError())
+		if (!Ar->Close())
 		{
 			return FIoStatus(EIoErrorCode::WriteError, FString::Printf(TEXT("Failed to write to '%s'"), *OutputPath));
 		}
@@ -315,6 +315,7 @@ static void WriteConfigFile(
 	const FUploadResult& UploadResult,
 	const FKeyChain& KeyChain)
 {
+	// This string can get quite long so don't bother with any inline storage
 	FStringBuilderBase Sb;
 	Sb << TEXT("[Endpoint]") << TEXT("\r\n");
 
