Add api documentation for FMultiEndpointHttpClient::GetHttpStats

#rnx

diff --git a/Engine/Source/Runtime/Experimental/IoStore/OnDemand/Private/OnDemandHttpClient.h b/Engine/Source/Runtime/Experimental/IoStore/OnDemand/Private/OnDemandHttpClient.h
--- a/Engine/Source/Runtime/Experimental/IoStore/OnDemand/Private/OnDemandHttpClient.h
+++ b/Engine/Source/Runtime/Experimental/IoStore/OnDemand/Private/OnDemandHttpClient.h
@@ -151,6 +151,11 @@ public:
 	void CancelRequest(FHttpTicketId TicketId);
 
 	void UpdateConnections();
+
+	/**
+	 * Samples the underlying http client for performance statistics and stores them in the provided FOnDemandHttpStats.
+	 * Can be called from any thread, not just the http thread.
+	 */
 	void GetHttpStats(FOnDemandHttpStats& Out) const;
 
 private:
