Dataflow : protect against infinite recursive subgraph calls
- Allow maximum subgraph nesting calls depth of 20
- Prevent subgraph node dynamic connection refresh method from being re-entrant

diff --git a/Engine/Source/Runtime/Experimental/Dataflow/Engine/Private/Dataflow/DataflowDynamicConnections.cpp b/Engine/Source/Runtime/Experimental/Dataflow/Engine/Private/Dataflow/DataflowDynamicConnections.cpp
--- a/Engine/Source/Runtime/Experimental/Dataflow/Engine/Private/Dataflow/DataflowDynamicConnections.cpp
+++ b/Engine/Source/Runtime/Experimental/Dataflow/Engine/Private/Dataflow/DataflowDynamicConnections.cpp
@@ -51,6 +51,13 @@ void FDataflowDynamicConnections::Refresh()
 		return;
 	}
 
+	// prevent re-entrance
+	if (bIsRefreshing)
+	{
+		return;
+	}
+	bIsRefreshing = true;
+
 	OwnerNode->PauseInvalidations();
 
 	// save the remote connections from the dynamic connections
@@ -119,6 +126,8 @@ void FDataflowDynamicConnections::Refresh()
 	OwnerNode->Invalidate();
 
 	DataflowAsset->RefreshEdNodeByGuid(OwnerNode->GetGuid());
+
+	bIsRefreshing = false;
 }
 
 FDataflowConnection* FDataflowDynamicConnections::AddNewConnectionFromPropertyDesc(const FPropertyBagPropertyDesc& PropertyDesc)
diff --git a/Engine/Source/Runtime/Experimental/Dataflow/Engine/Private/Dataflow/DataflowSubGraphNodes.cpp b/Engine/Source/Runtime/Experimental/Dataflow/Engine/Private/Dataflow/DataflowSubGraphNodes.cpp
--- a/Engine/Source/Runtime/Experimental/Dataflow/Engine/Private/Dataflow/DataflowSubGraphNodes.cpp
+++ b/Engine/Source/Runtime/Experimental/Dataflow/Engine/Private/Dataflow/DataflowSubGraphNodes.cpp
@@ -19,6 +19,8 @@ namespace UE::Dataflow::SubGraphNodes::Private
 {
 	class FContextSubGraph: public FEngineContext
 	{
+		static constexpr int32 MaxSubGraphNestingDepth = 20;
+
 	public:
 		DATAFLOW_CONTEXT_INTERNAL(FEngineContext, FContextSubGraph);
 
@@ -29,6 +31,12 @@ namespace UE::Dataflow::SubGraphNodes::Private
 		{
 			// make sure to match the threaded mode
 			SetThreaded(ParentContext.IsThreaded());
+
+			// Record the subgraph nesting depth
+			if (const FContextSubGraph* ParentSubGraphContext = ParentContext.AsType<FContextSubGraph>())
+			{
+				SubGraphNestingDepth = ParentSubGraphContext->SubGraphNestingDepth + 1;
+			}
 		}
 
 		virtual void SetDataImpl(FContextCacheKey Key, TUniquePtr<FContextCacheElementBase>&& DataStoreEntry) override
@@ -63,6 +71,13 @@ namespace UE::Dataflow::SubGraphNodes::Private
 
 		bool EvaluateSubGraph(const FDataflowSubGraphOutputNode& SubGraphOutputNode, const FDataflowNode& CallerNode, const FDataflowOutput& CallerOutput)
 		{
+			if (SubGraphNestingDepth >= MaxSubGraphNestingDepth)
+			{
+				ParentContext.Error("Maximum nesting depth reached - skipping evaluation (Potential recursive subgraph call)", &CallerNode, &CallerOutput);
+				CallerOutput.SetNullValue(*this);
+				return false;
+			}
+
 			auto OnInfo = [this](const FDataflowNode* Node, const FDataflowOutput* Output, const FString& Info)
 				{
 					ParentContext.Info(Info, Node, Output);
@@ -154,6 +169,7 @@ namespace UE::Dataflow::SubGraphNodes::Private
 
 	private:
 		int32 CurrentArrayIndex = 0;
+		int32 SubGraphNestingDepth = 0;
 		FContext& ParentContext;
 		const UE::Dataflow::ISubGraphContextCallback& SubGraphCallback;
 	};
diff --git a/Engine/Source/Runtime/Experimental/Dataflow/Engine/Public/Dataflow/DataflowDynamicConnections.h b/Engine/Source/Runtime/Experimental/Dataflow/Engine/Public/Dataflow/DataflowDynamicConnections.h
--- a/Engine/Source/Runtime/Experimental/Dataflow/Engine/Public/Dataflow/DataflowDynamicConnections.h
+++ b/Engine/Source/Runtime/Experimental/Dataflow/Engine/Public/Dataflow/DataflowDynamicConnections.h
@@ -53,4 +53,8 @@ private:
 	UE::Dataflow::FPin::EDirection PinDirection = UE::Dataflow::FPin::EDirection::NONE;
 	IOwnerInterface* OwnerInterface;
 	TWeakObjectPtr<UDataflow> DataflowAssetWeakPtr;
+
+	// used to prevent re-entrance of the refresh function
+	// this happens if subgraph  includes a node to itself and am output pin is connected to another node
+	bool bIsRefreshing = false;
 };
\ No newline at end of file
