Fix duplicate entry when serializing FOnDemandTocAdditionalFile to compact binary.

#rnx

diff --git a/Engine/Source/Runtime/Experimental/IoStore/OnDemandCore/Private/OnDemandToc.cpp b/Engine/Source/Runtime/Experimental/IoStore/OnDemandCore/Private/OnDemandToc.cpp
--- a/Engine/Source/Runtime/Experimental/IoStore/OnDemandCore/Private/OnDemandToc.cpp
+++ b/Engine/Source/Runtime/Experimental/IoStore/OnDemandCore/Private/OnDemandToc.cpp
@@ -387,7 +387,7 @@ FCbWriter& operator<<(FCbWriter& Writer, const FOnDemandTocAdditionalFile& Addit
 	Writer.BeginObject();
 	Writer.AddHash(UTF8TEXTVIEW("Hash"), AdditionalFile.Hash);
 	Writer.AddString(UTF8TEXTVIEW("Filename"), AdditionalFile.Filename);
-	Writer.AddInteger(UTF8TEXTVIEW("Filename"), AdditionalFile.FileSize);
+	Writer.AddInteger(UTF8TEXTVIEW("FileSize"), AdditionalFile.FileSize);
 	Writer.EndObject();
 
 	return Writer;
