Record the 'time to first byte' when sending IAX analytics.

#rnx

- Note that this value is for the entire event loop and cannot be separated by request type so we only send a single version of the data using the Iax prefix rather than one for Iad and one for Iad as we have been doing with older analytics values.

diff --git a/Engine/Source/Runtime/Experimental/IoStore/OnDemand/Private/Statistics.cpp b/Engine/Source/Runtime/Experimental/IoStore/OnDemand/Private/Statistics.cpp
--- a/Engine/Source/Runtime/Experimental/IoStore/OnDemand/Private/Statistics.cpp
+++ b/Engine/Source/Runtime/Experimental/IoStore/OnDemand/Private/Statistics.cpp
@@ -1000,9 +1000,22 @@ void FOnDemandIoBackendStats::ReportEndPointAnalytics(TArray<FAnalyticsEventAttr
 #undef UE_PREFIX
 		};
 
+		uint32 TtfbMs = 0;
+		if (IOnDemandIoStore* Store = TryGetOnDemandIoStore(); Store != nullptr)
+		{
+			FOnDemandHttpStats Stats;
+			Store->GetHttpStats(Stats);
+
+			TtfbMs = Stats.GetTimeToFirstByteMs();
+		}
+
+		AppendAnalyticsEventAttributeArray(OutAnalyticsArray
+			, TEXT("IaxTtfbMs"), TtfbMs
+			, TEXT("HttpIoDispatcherEnabled"), GHttpIoDispatcherEnabled
+		);
+
 		ReportHttpStats(TEXT("Ias"), EHttpRequestType::Streaming);
 		ReportHttpStats(TEXT("Iad"), EHttpRequestType::Installed);
-		AppendAnalyticsEventAttributeArray(OutAnalyticsArray, TEXT("HttpIoDispatcherEnabled"), GHttpIoDispatcherEnabled);
 	}
 
 	if (GIasReportCacheAnalyticsEnabled)
