Remove the IASCache analytic 'IasCacheTotalCount' and replace it with a much more useful 'IasCacheCachedEntries'

#rnx

-The old value 'IasCacheTotalCount' combined the total GET count and total PUT count for the total duration of the cache's lifespan.
- The new value 'IasCacheCachedEntries' tracks the number of entries that make up the cache.
- The new value is sent to analytics and written to the periodic logging system. We do not currently display it on screen.

diff --git a/Engine/Source/Runtime/Experimental/IoStore/OnDemand/Private/JournaledCache.cpp b/Engine/Source/Runtime/Experimental/IoStore/OnDemand/Private/JournaledCache.cpp
--- a/Engine/Source/Runtime/Experimental/IoStore/OnDemand/Private/JournaledCache.cpp
+++ b/Engine/Source/Runtime/Experimental/IoStore/OnDemand/Private/JournaledCache.cpp
@@ -793,9 +793,11 @@ bool FDiskCache::Evict(uint64 Key)
 
 	const TPair<uint64, FMapEntry>& Element = DataMap.Get(Id);
 	MappedBytes -= Element.Value.Size;
-	FOnDemandIoBackendStats::Get()->OnCachePersistedBytes(MappedBytes);
 
 	DataMap.Remove(Id);
+
+	FOnDemandIoBackendStats::Get()->OnCachePersistedBytes(MappedBytes, DataMap.Num());
+
 	return true;
 }
 
@@ -898,7 +900,10 @@ uint64 FDiskCache::Insert(uint64 DataBase, const FDataEntry* Entries, uint32 Ent
 	}
 
 	MappedBytes += TotalSize;
-	FOnDemandIoBackendStats::Get()->OnCachePersistedBytes(MappedBytes);
+
+	// DataMap is updated in the ::Insert call above
+	FOnDemandIoBackendStats::Get()->OnCachePersistedBytes(MappedBytes, DataMap.Num());
+
 	return MappedBytes;
 }
 
@@ -973,7 +978,7 @@ void FDiskCache::Drop()
 	MappedBytes = 0;
 	DataMap.Reset();
 
-	FOnDemandIoBackendStats::Get()->OnCachePersistedBytes(MappedBytes);
+	FOnDemandIoBackendStats::Get()->OnCachePersistedBytes(MappedBytes, DataMap.Num());
 
 	OpenDataFile();
 }
diff --git a/Engine/Source/Runtime/Experimental/IoStore/OnDemand/Private/Statistics.cpp b/Engine/Source/Runtime/Experimental/IoStore/OnDemand/Private/Statistics.cpp
--- a/Engine/Source/Runtime/Experimental/IoStore/OnDemand/Private/Statistics.cpp
+++ b/Engine/Source/Runtime/Experimental/IoStore/OnDemand/Private/Statistics.cpp
@@ -362,6 +362,7 @@ FCounterAtomicInt		GCachePutCount(TEXT("Ias/CachePutCount"), TraceCounterDisplay
 FCounterAtomicInt		GCachePutExistingCount(TEXT("Ias/CachePutExistingCount"), TraceCounterDisplayHint_None);
 FCounterAtomicInt		GCachePutRejectCount(TEXT("Ias/CachePutRejectCount"), TraceCounterDisplayHint_None);
 FCounterAtomicInt		GCacheCachedBytes(TEXT("Ias/CacheCachedBytes"), TraceCounterDisplayHint_Memory);
+FCounterAtomicInt		GCacheCachedEntries(TEXT("Ias/CacheCachedEntries"), TraceCounterDisplayHint_Memory);
 FCounterAtomicInt		GCacheWrittenBytes(TEXT("Ias/CacheWrittenBytes"), TraceCounterDisplayHint_Memory);
 FCounterFloat			GCacheSuspendedSeconds(TEXT("Ias/CacheSuspendedSeconds"), TraceCounterDisplayHint_None);
 int64					GCacheMaxBytes = 0;
@@ -686,8 +687,8 @@ void FOnDemandIoBackendStats::PrintPeriodicLogging()
 
 		if (BackendStatus.IsCacheEnabled())
 		{
-			UE_LOG(LogIas, Log, TEXT("CacheStats: CachedKiB=%d, WrittenKiB=%d, ReadKiB=%d, RejectedKiB=%d, Get=%d, Error=%d, Put=%d, PutReject=%d, PutExisting=%d"),
-				(int32)CCachedKiB, (int32)CWrittenKiB, (int32)CReadKiB, (int32)CRejectedKiB, CGetCount, CErrorCount, CPutCount, CPutRejectCount, CPutExistingCount);
+			UE_LOG(LogIas, Log, TEXT("CacheStats: Count=%d, CachedKiB=%d, WrittenKiB=%d, ReadKiB=%d, RejectedKiB=%d, Get=%d, Error=%d, Put=%d, PutReject=%d, PutExisting=%d"),
+				(int32)GCacheCachedEntries.Get(), (int32)CCachedKiB, (int32)CWrittenKiB, (int32)CReadKiB, (int32)CRejectedKiB, CGetCount, CErrorCount, CPutCount, CPutRejectCount, CPutExistingCount);
 		}
 		else
 		{
@@ -1021,20 +1022,19 @@ void FOnDemandIoBackendStats::ReportEndPointAnalytics(TArray<FAnalyticsEventAttr
 	if (GIasReportCacheAnalyticsEnabled)
 	{
 #define UE_TRACK_DELTA(Name, Value) TEXT(Name), GDeltaTracking.Get(TEXTVIEW(Name), Value)
-		const int64 CacheTotalCount = GCacheGetCount.Get() + GCachePutCount.Get();
 		const float CacheUsagePercent = GCacheMaxBytes > 0 ? 100.f * (float(GCacheCachedBytes.Get()) / float(GCacheMaxBytes)) : 0.f;
 
 		AppendAnalyticsEventAttributeArray(OutAnalyticsArray
 
 			, TEXT("IasCacheEnabled"), BackendStatus.IsCacheEnabled()
 
-			, UE_TRACK_DELTA("IasCacheTotalCount", CacheTotalCount)
 			, UE_TRACK_DELTA("IasCacheErrorCount", GCacheErrorCount.Get())
 			, UE_TRACK_DELTA("IasCacheDecodeErrors", GCacheDecodeErrorCount.Get())
 			, UE_TRACK_DELTA("IasCacheGetCount", GCacheGetCount.Get())
 			, UE_TRACK_DELTA("IasCachePutCount", GCachePutCount.Get())
 
 			, TEXT("IasCacheCachedBytes"), GCacheCachedBytes.Get()
+			, TEXT("IasCacheCachedEntries"), GCacheCachedEntries.Get()
 			, TEXT("IasCacheMaxBytes"), GCacheMaxBytes
 			, TEXT("IasCacheUsagePercent"), CacheUsagePercent
 
@@ -1300,9 +1300,10 @@ void FOnDemandIoBackendStats::OnCachePendingBytes(uint64 TotalSize)
 	GCachePendingBytes.Set(TotalSize);
 }
 
-void FOnDemandIoBackendStats::OnCachePersistedBytes(uint64 TotalSize)
+void FOnDemandIoBackendStats::OnCachePersistedBytes(uint64 TotalSize, int32 NumEntries)
 {
 	GCacheCachedBytes.Set(TotalSize);
+	GCacheCachedEntries.Set(NumEntries);
 }
 
 void FOnDemandIoBackendStats::OnCacheWriteBytes(uint64 WriteSize)
diff --git a/Engine/Source/Runtime/Experimental/IoStore/OnDemand/Private/Statistics.h b/Engine/Source/Runtime/Experimental/IoStore/OnDemand/Private/Statistics.h
--- a/Engine/Source/Runtime/Experimental/IoStore/OnDemand/Private/Statistics.h
+++ b/Engine/Source/Runtime/Experimental/IoStore/OnDemand/Private/Statistics.h
@@ -83,7 +83,7 @@ public:
 	void OnCachePutExisting(uint64 DataSize) IAS_STATISTICS_IMPL()
 	void OnCachePutReject(uint64 DataSize) IAS_STATISTICS_IMPL()
 	void OnCachePendingBytes(uint64 TotalSize) IAS_STATISTICS_IMPL()
-	void OnCachePersistedBytes(uint64 TotalSize) IAS_STATISTICS_IMPL()
+	void OnCachePersistedBytes(uint64 TotalSize, int32 NumEntries) IAS_STATISTICS_IMPL()
 	void OnCacheWriteBytes(uint64 WriteSize) IAS_STATISTICS_IMPL()
 	void OnCacheSetMaxBytes(uint64 TotalSize) IAS_STATISTICS_IMPL()
 	void OnCacheSuspended(double Seconds) IAS_STATISTICS_IMPL()
