Fix potential out of bounds access inFOnDemandHostGroup.

#rnx

- Return an empty string from FOnDemandHostGroup::Host if the index is out of bounds or the host group has no urls.
- If all urls are filtered out from the hostgroup we shouyld explicitly set PrimaryIndex to INDEX_NONE to indicate that the host group is invalid.

diff --git a/Engine/Source/Runtime/Experimental/IoStore/OnDemandCore/Private/OnDemandHostGroup.cpp b/Engine/Source/Runtime/Experimental/IoStore/OnDemandCore/Private/OnDemandHostGroup.cpp
--- a/Engine/Source/Runtime/Experimental/IoStore/OnDemandCore/Private/OnDemandHostGroup.cpp
+++ b/Engine/Source/Runtime/Experimental/IoStore/OnDemandCore/Private/OnDemandHostGroup.cpp
@@ -111,7 +111,14 @@ TConstArrayView<FAnsiString> FOnDemandHostGroup::Hosts() const
 
 FAnsiStringView FOnDemandHostGroup::Host(int32 Index) const
 {
-	return Impl->HostUrls.IsEmpty() ? FAnsiStringView() : Impl->HostUrls[Index];
+	if (Index < 0 || Index >= Impl->HostUrls.Num())
+	{
+		return FAnsiStringView();
+	}
+	else
+	{
+		return Impl->HostUrls[Index];
+	}
 }
 
 FAnsiStringView FOnDemandHostGroup::CycleHost(int32& InOutIndex) const
@@ -257,14 +264,16 @@ TIoStatusOr<FOnDemandHostGroup> FOnDemandHostGroup::Create(TConstArrayView<FStri
 		Impl->HostUrls.Emplace(MoveTemp(AnsiUrl));
 	}
 
-	if (Impl->HostUrls.IsEmpty())
+	if (!Impl->HostUrls.IsEmpty())
 	{
+		Impl->PrimaryIndex = 0;
+		return FOnDemandHostGroup(MoveTemp(Impl));
+	}
+	else
+	{
+		Impl->PrimaryIndex = -1;
 		return FIoStatus(EIoErrorCode::InvalidParameter, TEXT("Alls URLs were filtered out"));
 	}
-
-	Impl->PrimaryIndex = 0;
-
-	return FOnDemandHostGroup(MoveTemp(Impl));
 }
 
 FName FOnDemandHostGroup::DefaultName = FName("Default");
