Fix incorrect code documentation.

#rnx

diff --git a/Engine/Source/Runtime/Experimental/IoStore/OnDemandCore/Public/IO/OnDemandHostGroup.h b/Engine/Source/Runtime/Experimental/IoStore/OnDemandCore/Public/IO/OnDemandHostGroup.h
--- a/Engine/Source/Runtime/Experimental/IoStore/OnDemandCore/Public/IO/OnDemandHostGroup.h
+++ b/Engine/Source/Runtime/Experimental/IoStore/OnDemandCore/Public/IO/OnDemandHostGroup.h
@@ -49,7 +49,7 @@ public:
 	/** Default host group name. */
 	UE_API static FName DefaultName;
 
-	/** Returns whether two endpoint instance are the same or not. */
+	/** Returns whether two host group instances are the same or not. */
 	friend inline bool operator==(const FOnDemandHostGroup& Lhs, const FOnDemandHostGroup& Rhs)
 	{
 		return Lhs.Impl.Get() == Rhs.Impl.Get();
