Fix incorrect code documentation.

#rnx

diff --git a/Engine/Source/Runtime/Experimental/IoStore/OnDemandCore/Public/IO/OnDemandHostGroup.h b/Engine/Source/Runtime/Experimental/IoStore/OnDemandCore/Public/IO/OnDemandHostGroup.h
--- a/Engine/Source/Runtime/Experimental/IoStore/OnDemandCore/Public/IO/OnDemandHostGroup.h
+++ b/Engine/Source/Runtime/Experimental/IoStore/OnDemandCore/Public/IO/OnDemandHostGroup.h
@@ -54,10 +54,10 @@ public:
 	{
 		return Lhs.Impl.Get() == Rhs.Impl.Get();
 	}
-	/** Returns the type hash of the endpoint. */ 
-	friend inline uint32 GetTypeHash(const FOnDemandHostGroup& Endpoint)
+	/** Returns the type hash of the host group. */ 
+	friend inline uint32 GetTypeHash(const FOnDemandHostGroup& HostGroup)
 	{
-		return *reinterpret_cast<const uint32*>(Endpoint.Impl.Get());
+		return *reinterpret_cast<const uint32*>(HostGroup.Impl.Get());
 	}
 
 private:
