Added a new IAX analytic 'IaxTimeSpent that tracks the length of time (in seconds) that an analytics event covers.

#rnx

diff --git a/Engine/Source/Runtime/Experimental/IoStore/OnDemand/Private/Statistics.cpp b/Engine/Source/Runtime/Experimental/IoStore/OnDemand/Private/Statistics.cpp
--- a/Engine/Source/Runtime/Experimental/IoStore/OnDemand/Private/Statistics.cpp
+++ b/Engine/Source/Runtime/Experimental/IoStore/OnDemand/Private/Statistics.cpp
@@ -369,7 +369,8 @@ int64					GCacheMaxBytes = 0;
 FCounterAtomicInt		GCachePendingBytes(TEXT("Ias/CachePendingBytes"), TraceCounterDisplayHint_Memory);
 FCounterAtomicInt		GCacheReadBytes(TEXT("Ias/CacheReadBytes"), TraceCounterDisplayHint_Memory);
 FCounterAtomicInt		GCacheRejectBytes(TEXT("Ias/CachePutRejectBytes"), TraceCounterDisplayHint_Memory);
-
+// general stats
+double					GAnalyticsStartTiming = 0.0;
 // http stats
 int64					GAppResumeCount = 0;
 bool					GHttpDistributedEndpointResolved = false;
@@ -574,6 +575,8 @@ FOnDemandIoBackendStats::FOnDemandIoBackendStats(FBackendStatus& InStatus)
 	check(GStatistics == nullptr);
 	GStatistics = this;
 
+	GAnalyticsStartTiming = FPlatformTime::Seconds();
+
 	OnApplicationResumeHandle = FCoreDelegates::ApplicationHasEnteredForegroundDelegate.AddRaw(this, &FOnDemandIoBackendStats::OnApplicationResume);
 
 	EndFrameDelegateHandle = FCoreDelegates::OnEndFrame.AddLambda([this]()
@@ -956,11 +959,12 @@ void FOnDemandIoBackendStats::PrintOnScreenStatistics(FCoreDelegates::FSeverityM
 
 void FOnDemandIoBackendStats::ReportGeneralAnalytics(TArray<FAnalyticsEventAttribute>& OutAnalyticsArray) const
 {
-	// Note that this analytics section is not optional, if we are reporting analytics then we report this section
-	// first. This means we can use the values here to determine if an analytics payload contains ondemand data or
-	// not since with the current system we are unable to specify our own analytics payload.
+	const double CurrentTime = FPlatformTime::Seconds();
+	const int64 TimeSpent = FMath::RoundToInt64(CurrentTime - GAnalyticsStartTiming);
+	GAnalyticsStartTiming = CurrentTime;
 
 	AppendAnalyticsEventAttributeArray(OutAnalyticsArray
+		, TEXT("IaxTimeSpent"), TimeSpent
 		, TEXT("IasAppResumeCount"), GAppResumeCount
 		, TEXT("IasHttpDistributedEndpointResolved"), GHttpDistributedEndpointResolved
 		, TEXT("IasHttpHasEverConnected"), GHttpConnectCount.Get() > 0 // Report if the system has ever actually managed to make a connection
