Add Apply in XPBDWeakConstraints to follow convention in AddPerIterationPBDConstraintRule_Apply: Constraint->Apply(Particles, Dt);

diff --git a/Engine/Source/Runtime/Experimental/Chaos/Public/Chaos/XPBDWeakConstraints.h b/Engine/Source/Runtime/Experimental/Chaos/Public/Chaos/XPBDWeakConstraints.h
--- a/Engine/Source/Runtime/Experimental/Chaos/Public/Chaos/XPBDWeakConstraints.h
+++ b/Engine/Source/Runtime/Experimental/Chaos/Public/Chaos/XPBDWeakConstraints.h
@@ -113,6 +113,11 @@ namespace Chaos::Softs
 			ApplyInParallel(Range, Dt);
 		}
 
+		void Apply(FSolverParticlesRange& Particles, const FSolverReal Dt) const
+		{
+			ApplyInParallel(Particles, Dt);
+		}
+
 		const TArray<TArray<int32>>& GetIndices()
 		{
 			return Indices;
