Add IOnDemandIoStore::GetIsOnDemand
- Does less work than checking install size > 0
- Allow mutliple queries to be batched

diff --git a/Engine/Source/Runtime/CoreUObject/Private/Serialization/PackageStore.cpp b/Engine/Source/Runtime/CoreUObject/Private/Serialization/PackageStore.cpp
--- a/Engine/Source/Runtime/CoreUObject/Private/Serialization/PackageStore.cpp
+++ b/Engine/Source/Runtime/CoreUObject/Private/Serialization/PackageStore.cpp
@@ -552,7 +552,7 @@ void FPackageStore::Mount(TSharedRef<IPackageStoreBackend> Backend, int32 Priori
 }
 
 EPackageStoreEntryStatus FPackageStore::GetPackageStoreEntry(FPackageId PackageId, FName PackageName,
-	FPackageStoreEntry& OutPackageStoreEntry)
+	FPackageStoreEntry& OutPackageStoreEntry, TSharedPtr<IPackageStoreBackend>* FoundBackend /*= nullptr*/)
 {
 	check(ThreadReadCount);
 	for (const FBackendAndPriority& Backend : Backends)
@@ -564,6 +564,11 @@ EPackageStoreEntryStatus FPackageStore::GetPackageStoreEntry(FPackageId PackageI
 
 		if (!bContinueSearch)
 		{
+			if (FoundBackend)
+			{
+				*FoundBackend = Backend.Value;
+			}
+
 			checkSlow((Backend.Value->GetSupportedLoaders() & OutPackageStoreEntry.LoaderType));
 			return Status;
 		}
diff --git a/Engine/Source/Runtime/CoreUObject/Public/Serialization/PackageStore.h b/Engine/Source/Runtime/CoreUObject/Public/Serialization/PackageStore.h
--- a/Engine/Source/Runtime/CoreUObject/Public/Serialization/PackageStore.h
+++ b/Engine/Source/Runtime/CoreUObject/Public/Serialization/PackageStore.h
@@ -235,7 +235,7 @@ public:
 
 	/* Returns the package store entry data with export info and imported packages for the specified package ID. */
 	COREUOBJECT_API EPackageStoreEntryStatus GetPackageStoreEntry(FPackageId PackageId, FName PackageName, 
-		FPackageStoreEntry& OutPackageStoreEntry);
+		FPackageStoreEntry& OutPackageStoreEntry, TSharedPtr<IPackageStoreBackend>* FoundBackend = nullptr);
 
 	/* Returns the redirected package ID and source package name for the specified package ID if it's being redirected. */
 	COREUOBJECT_API bool GetPackageRedirectInfo(FPackageId PackageId, FName& OutSourcePackageName, FPackageId& OutRedirectedToPackageId);
diff --git a/Engine/Source/Runtime/Experimental/IoStore/OnDemand/Private/OnDemandContentInstaller.cpp b/Engine/Source/Runtime/Experimental/IoStore/OnDemand/Private/OnDemandContentInstaller.cpp
--- a/Engine/Source/Runtime/Experimental/IoStore/OnDemand/Private/OnDemandContentInstaller.cpp
+++ b/Engine/Source/Runtime/Experimental/IoStore/OnDemand/Private/OnDemandContentInstaller.cpp
@@ -4,7 +4,6 @@
 
 #include "Algo/Accumulate.h"
 #include "Async/UniqueLock.h"
-#include "Containers/RingBuffer.h"
 #include "Experimental/UnifiedError/CoreErrorTypes.h"
 #include "HAL/Platform.h"
 #include "IO/IoBuffer.h"
@@ -18,7 +17,6 @@
 #include "OnDemandInstallCache.h"
 #include "OnDemandIoStore.h"
 #include "OnDemandPackageStoreBackend.h"
-#include "Serialization/PackageStore.h"
 #include "Statistics.h"
 #include "Templates/Overload.h"
 #include <atomic>
@@ -99,74 +97,30 @@ void ResolvePackageDependencies(
 	const TSet<FPackageId>& PackageIds,
 	bool bIncludeSoftReferences,
 	TSet<FPackageId>& OutResolved,
-	TSet<FPackageId>& OutMissing)
+	TSet<FPackageId>& OutMissing,
+	FPackageStoreReadScope* InReadScope /*= nullptr*/)
 {
 	TRACE_CPUPROFILER_EVENT_SCOPE(FOnDemandContentInstaller::ResolvePackageDependencies);
 
-	TRingBuffer<FPackageId> Queue;
-	TSet<FPackageId>		Visitied;
-
-	Visitied.Reserve(PackageIds.Num());
-	Queue.Reserve(PackageIds.Num());
-	for (const FPackageId& PackageId : PackageIds)
-	{
-		Queue.Add(PackageId);
-	}
-
-	FPackageStore& PackageStore = FPackageStore::Get();
-	FPackageStoreReadScope _(PackageStore);
-
-	while (!Queue.IsEmpty())
+	auto Visitor = [&OutResolved, &OutMissing](
+		FPackageId PackageId,
+		EPackageStoreEntryStatus EntryStatus,
+		const FPackageStoreEntry& PackageStoreEntry,
+		const TSharedPtr<IPackageStoreBackend>& Backend) -> bool
 	{
-		FPackageId PackageId = Queue.PopFrontValue();
-		{
-			FName		SourcePackageName;
-			FPackageId	RedirectedToPackageId;
-			if (PackageStore.GetPackageRedirectInfo(PackageId, SourcePackageName, RedirectedToPackageId))
-			{
-				PackageId = RedirectedToPackageId;
-			}
-		}
-
-		bool bIsAlreadyInSet = false;
-		Visitied.Add(PackageId, &bIsAlreadyInSet);
-		if (bIsAlreadyInSet)
-		{
-			continue;
-		}
-
-		FPackageStoreEntry PackageStoreEntry;
-		const EPackageStoreEntryStatus EntryStatus = PackageStore.GetPackageStoreEntry(PackageId, NAME_None, PackageStoreEntry);
 		if (EntryStatus != EPackageStoreEntryStatus::Missing)
 		{
 			OutResolved.Add(PackageId);
-			for (const FPackageId& ImportedPackageId : PackageStoreEntry.ImportedPackageIds)
-			{
-				if (!Visitied.Contains(ImportedPackageId))
-				{
-					Queue.Add(ImportedPackageId);
-				}
-			}
-
-			if (bIncludeSoftReferences)
-			{
-				TConstArrayView<FPackageId> SoftReferences;
-				TConstArrayView<uint32> Indices = PackageStore.GetSoftReferences(PackageId, SoftReferences);
-				for (uint32 Idx : Indices)
-				{
-					const FPackageId& SoftRef = SoftReferences[Idx];
-					if (!Visitied.Contains(SoftRef))
-					{
-						Queue.Add(SoftRef);
-					}
-				}
-			}
 		}
 		else
 		{
 			OutMissing.Add(PackageId);
 		}
-	}
+
+		return true;
+	};
+
+	TVisitPackageDependencies(PackageIds, bIncludeSoftReferences, Visitor, InReadScope);
 }
 
 ////////////////////////////////////////////////////////////////////////////////
@@ -181,6 +135,7 @@ void ResolveChunksToInstall(
 	TRACE_CPUPROFILER_EVENT_SCOPE(FOnDemandContentInstaller::ResolveChunksToInstall);
 
 	// For now we always download these required chunks
+	// TODO: only add these if we ref packages from that container
 	for (const FSharedOnDemandContainer& Container : Containers)
 	{
 		FResolvedContainerChunks& ResolvedChunks = OutResolvedContainerChunks.AddDefaulted_GetRef();
diff --git a/Engine/Source/Runtime/Experimental/IoStore/OnDemand/Private/OnDemandContentInstaller.h b/Engine/Source/Runtime/Experimental/IoStore/OnDemand/Private/OnDemandContentInstaller.h
--- a/Engine/Source/Runtime/Experimental/IoStore/OnDemand/Private/OnDemandContentInstaller.h
+++ b/Engine/Source/Runtime/Experimental/IoStore/OnDemand/Private/OnDemandContentInstaller.h
@@ -4,6 +4,7 @@
 
 #include "Async/Mutex.h"
 #include "Containers/Set.h"
+#include "Containers/RingBuffer.h"
 #include "Experimental/UnifiedError/CoreErrorTypes.h"
 #include "IO/IoStoreOnDemand.h"
 #include "IO/OnDemandError.h"
@@ -14,11 +15,13 @@
 #include "IO/HttpIoDispatcher.h"
 #include "Misc/TVariant.h"
 #include "OnDemandIoStore.h"
+#include "Serialization/PackageStore.h"
 #include "Tasks/Pipe.h"
 
 #include <atomic>
 
 class FIoBuffer;
+class FPackageStoreReadScope;
 struct FAnalyticsEventAttribute;
 
 namespace UE::IoStore
@@ -40,12 +43,92 @@ struct FResolvedContainerChunks
 	uint64						TotalSize = 0;
 };
 
+////////////////////////////////////////////////////////////////////////////////
+template<typename VisitorType>
+void TVisitPackageDependencies(const TSet<FPackageId>& PackageIds, bool bIncludeSoftReferences, const VisitorType& Visitor, FPackageStoreReadScope* InReadScope = nullptr)
+{
+	TRACE_CPUPROFILER_EVENT_SCOPE(FOnDemandContentInstaller::TVisitPackageDependencies);
+
+	TRingBuffer<FPackageId> Queue;
+	TSet<FPackageId>		Visitied;
+
+	Visitied.Reserve(PackageIds.Num());
+	Queue.Reserve(PackageIds.Num());
+	for (const FPackageId& PackageId : PackageIds)
+	{
+		Queue.Add(PackageId);
+	}
+
+	FPackageStore& PackageStore = FPackageStore::Get();
+	TOptional<FPackageStoreReadScope> ReadScope;
+	if (InReadScope == nullptr)
+	{
+		ReadScope.Emplace(PackageStore);
+	}
+
+	while (!Queue.IsEmpty())
+	{
+		FPackageId PackageId = Queue.PopFrontValue();
+		{
+			FName		SourcePackageName;
+			FPackageId	RedirectedToPackageId;
+			if (PackageStore.GetPackageRedirectInfo(PackageId, SourcePackageName, RedirectedToPackageId))
+			{
+				PackageId = RedirectedToPackageId;
+			}
+		}
+
+		bool bIsAlreadyInSet = false;
+		Visitied.Add(PackageId, &bIsAlreadyInSet);
+		if (bIsAlreadyInSet)
+		{
+			continue;
+		}
+
+		FPackageStoreEntry PackageStoreEntry;
+		TSharedPtr<IPackageStoreBackend> Backend;
+		const EPackageStoreEntryStatus EntryStatus = PackageStore.GetPackageStoreEntry(PackageId, NAME_None, PackageStoreEntry, &Backend);
+
+		const bool bContinue = Visitor(PackageId, EntryStatus, PackageStoreEntry, Backend);
+		if (bContinue == false)
+		{
+			break;
+		}
+
+		if (EntryStatus != EPackageStoreEntryStatus::Missing)
+		{
+			for (const FPackageId& ImportedPackageId : PackageStoreEntry.ImportedPackageIds)
+			{
+				if (!Visitied.Contains(ImportedPackageId))
+				{
+					Queue.Add(ImportedPackageId);
+				}
+			}
+
+			if (bIncludeSoftReferences)
+			{
+				TConstArrayView<FPackageId> SoftReferences;
+				TConstArrayView<uint32> Indices = PackageStore.GetSoftReferences(PackageId, SoftReferences);
+				for (uint32 Idx : Indices)
+				{
+					const FPackageId& SoftRef = SoftReferences[Idx];
+					if (!Visitied.Contains(SoftRef))
+					{
+						Queue.Add(SoftRef);
+					}
+				}
+			}
+		}
+	}
+}
+
 ////////////////////////////////////////////////////////////////////////////////
 void ResolvePackageDependencies(
 	const TSet<FPackageId>& PackageIds,
 	bool bIncludeSoftReferences,
 	TSet<FPackageId>& OutResolved,
-	TSet<FPackageId>& OutMissing);
+	TSet<FPackageId>& OutMissing,
+	FPackageStoreReadScope* ReadScope = nullptr);
 
 ////////////////////////////////////////////////////////////////////////////////
 void ResolveChunksToInstall(
diff --git a/Engine/Source/Runtime/Experimental/IoStore/OnDemand/Private/OnDemandIoStore.cpp b/Engine/Source/Runtime/Experimental/IoStore/OnDemand/Private/OnDemandIoStore.cpp
--- a/Engine/Source/Runtime/Experimental/IoStore/OnDemand/Private/OnDemandIoStore.cpp
+++ b/Engine/Source/Runtime/Experimental/IoStore/OnDemand/Private/OnDemandIoStore.cpp
@@ -901,6 +901,75 @@ FIoStatus FOnDemandIoStore::GetInstallSizesByMountId(const FOnDemandGetInstallSi
 	return EIoErrorCode::Ok;
 }
 
+FIoStatus FOnDemandIoStore::GetIsOnDemand(TConstArrayView<FOnDemandIsOnDemandArgs> InArgs, TArrayView<FOnDemandIsOnDemandResult> OutResults) const
+{
+	TRACE_CPUPROFILER_EVENT_SCOPE(FOnDemandIoStore::GetIsOnDemand);
+
+	using namespace UE::IoStore::Private;
+
+	if (InArgs.Num() != OutResults.Num())
+	{
+		FIoStatus Status = FIoStatusBuilder(EIoErrorCode::InvalidParameter)
+			<< TEXT("Result count must match arg count");
+		return Status;
+	}
+
+	if (InArgs.Num() == 0)
+	{
+		// trivial success
+		return EIoErrorCode::Ok;
+	}
+
+	TSet<FSharedOnDemandContainer> ContainersForInstallation;
+	if (FIoStatus Status = GetContainersForInstall({}, ContainersForInstallation);
+		Status.IsOk() == false)
+	{
+		return Status;
+	}
+
+	auto PackageVisitor = [this, OutResults](
+		int32 ResultIndex,
+		FPackageId PackageId,
+		EPackageStoreEntryStatus EntryStatus,
+		const FPackageStoreEntry& PackageStoreEntry,
+		const TSharedPtr<IPackageStoreBackend>& Backend) -> bool
+	{
+		if (Backend == PackageStoreBackend)
+		{
+			OutResults[ResultIndex].bIsOnDemand = true;
+			return false; // Don't continue
+		}
+
+		return true; // continue
+	};
+
+	FPackageStoreReadScope ReadScope(FPackageStore::Get());
+
+	// TODO: Parallel for here? Is that safe?
+	for (int Index = 0; const FOnDemandIsOnDemandArgs& Args : InArgs)
+	{
+		OutResults[Index].bIsOnDemand = false;
+
+		const bool bIncludeSoftReferences = EnumHasAnyFlags(Args.Options, EOnDemandIsOnDemandOptions::IncludeSoftReferences);
+
+		TVisitPackageDependencies(TSet<FPackageId>(Args.PackageIds), bIncludeSoftReferences,
+			[PackageVisitor, Index](
+				FPackageId PackageId,
+				EPackageStoreEntryStatus EntryStatus,
+				const FPackageStoreEntry& PackageStoreEntry,
+				const TSharedPtr<IPackageStoreBackend>& Backend) 
+			{
+				return PackageVisitor(Index, PackageId, EntryStatus, PackageStoreEntry, Backend);
+			},
+			&ReadScope
+		);
+
+		++Index;
+	}
+
+	return EIoErrorCode::Ok;
+}
+
 FOnDemandChunkInfo FOnDemandIoStore::GetStreamingChunkInfo(const FIoChunkId& ChunkId)
 {
 	const EOnDemandContainerFlags ContainerFlags = EOnDemandContainerFlags::Mounted | EOnDemandContainerFlags::StreamOnDemand;
@@ -1850,6 +1919,45 @@ void FOnDemandIoStore::FlushLastAccess(FOnDemandFlushLastAccessArgs&& Args, FOnD
 	Installer->EnqueueFlushLastAccessRequest(MoveTemp(Args), MoveTemp(OnCompleted));
 }
 
+FIoStatus FOnDemandIoStore::GetContainersForInstall(
+	FStringView MountId, 
+	TSet<FSharedOnDemandContainer>& OutContainersForInstallation) const
+{
+	TRACE_CPUPROFILER_EVENT_SCOPE(FOnDemandIoStore::GetContainersForInstall);
+
+	UE::TUniqueLock Lock(ContainerMutex);
+
+	OutContainersForInstallation.Reserve(Containers.Num());
+	for (const FSharedOnDemandContainer& Container : Containers)
+	{
+		if (Container->HasAnyFlags(EOnDemandContainerFlags::InstallOnDemand) == false)
+		{
+			continue;
+		}
+
+		if (Container->HasAnyFlags(EOnDemandContainerFlags::PendingEncryptionKey))
+		{
+			check(Container->HasAnyFlags(EOnDemandContainerFlags::Mounted) == false);
+			if (MountId == Container->MountId)
+			{
+				check(EnumHasAnyFlags(Container->Flags, EOnDemandContainerFlags::Mounted) == false);
+				FIoStatus Status = FIoStatusBuilder(EIoErrorCode::PendingEncryptionKey)
+					<< TEXT("Trying to install content from encrypted container '")
+					<< Container->Name
+					<< TEXT("'");
+				return Status;
+			}
+
+			continue;
+		}
+
+		check(Container->HasAnyFlags(EOnDemandContainerFlags::Mounted));
+		OutContainersForInstallation.Add(Container);
+	}
+
+	return FIoStatus::Ok;
+}
+
 FIoStatus FOnDemandIoStore::GetContainersAndPackagesForInstall(
 	FStringView MountId,
 	const TArray<FString>& TagSets,
@@ -1857,36 +1965,13 @@ FIoStatus FOnDemandIoStore::GetContainersAndPackagesForInstall(
 	TSet<FSharedOnDemandContainer>& OutContainersForInstallation, 
 	TSet<FPackageId>& OutPackageIdsToInstall) const
 {
+	TRACE_CPUPROFILER_EVENT_SCOPE(FOnDemandIoStore::GetContainersAndPackagesForInstall);
+
 	OutContainersForInstallation.Reset();
+	if (FIoStatus Status = GetContainersForInstall(MountId, OutContainersForInstallation);
+		Status.IsOk() == false)
 	{
-		UE::TUniqueLock Lock(ContainerMutex);
-
-		OutContainersForInstallation.Reserve(Containers.Num());
-		for (const FSharedOnDemandContainer& Container : Containers)
-		{
-			if (Container->HasAnyFlags(EOnDemandContainerFlags::InstallOnDemand) == false)
-			{
-				continue;
-			}
-
-			if (Container->HasAnyFlags(EOnDemandContainerFlags::PendingEncryptionKey))
-			{
-				check(Container->HasAnyFlags(EOnDemandContainerFlags::Mounted) == false);
-				if (MountId == Container->MountId)
-				{
-					check(EnumHasAnyFlags(Container->Flags, EOnDemandContainerFlags::Mounted) == false);
-					FIoStatus Status = FIoStatusBuilder(EIoErrorCode::PendingEncryptionKey)
-						<< TEXT("Trying to install content from encrypted container '")
-						<< Container->Name
-						<< TEXT("'");
-					return Status;
-				}
-
-				continue;
-			}
-			check(Container->HasAnyFlags(EOnDemandContainerFlags::Mounted));
-			OutContainersForInstallation.Add(Container);
-		}
+		return Status;
 	}
 
 	// It's not allowed to install all content
diff --git a/Engine/Source/Runtime/Experimental/IoStore/OnDemand/Private/OnDemandIoStore.h b/Engine/Source/Runtime/Experimental/IoStore/OnDemand/Private/OnDemandIoStore.h
--- a/Engine/Source/Runtime/Experimental/IoStore/OnDemand/Private/OnDemandIoStore.h
+++ b/Engine/Source/Runtime/Experimental/IoStore/OnDemand/Private/OnDemandIoStore.h
@@ -388,6 +388,7 @@ public:
 	virtual FIoStatus					Unmount(FStringView MountId) override;
 	virtual TIoStatusOr<FOnDemandInstallSizeResult> GetInstallSize(const FOnDemandGetInstallSizeArgs& Args) const override;
 	virtual FIoStatus					GetInstallSizesByMountId(const FOnDemandGetInstallSizeArgs& Args, TMap<FString, uint64>& OutSizesByMountId) const override;
+	virtual FIoStatus					GetIsOnDemand(TConstArrayView<FOnDemandIsOnDemandArgs> Args, TArrayView<FOnDemandIsOnDemandResult> OutResults) const override;
 	virtual FOnDemandCacheUsage			GetCacheUsage(const FOnDemandGetCacheUsageArgs& Args) const override;
 	virtual void						DumpMountedContainersToLog() const override;
 	virtual bool						IsOnDemandStreamingEnabled() const override;
@@ -406,6 +407,9 @@ public:
 private:
 	friend class FOnDemandContentInstaller;
 
+	FIoStatus							GetContainersForInstall(
+											FStringView MountId, 
+											TSet<FSharedOnDemandContainer>& OutContainersForInstallation) const;
 	FIoStatus							GetContainersAndPackagesForInstall(
 											FStringView MountId,
 											const TArray<FString>& TagSets,
diff --git a/Engine/Source/Runtime/Experimental/IoStore/OnDemandCore/Public/IO/IoStoreOnDemand.h b/Engine/Source/Runtime/Experimental/IoStore/OnDemandCore/Public/IO/IoStoreOnDemand.h
--- a/Engine/Source/Runtime/Experimental/IoStore/OnDemandCore/Public/IO/IoStoreOnDemand.h
+++ b/Engine/Source/Runtime/Experimental/IoStore/OnDemandCore/Public/IO/IoStoreOnDemand.h
@@ -457,6 +457,31 @@ struct FOnDemandInstallSizeResult
 	TOptional<uint64> DownloadSize = 0;
 };
 
+/** Options for controlling the behavior of the GetIsOnDemand request. */
+enum class EOnDemandIsOnDemandOptions
+{
+	/** No additional options. */
+	None = 0,
+	/** Follow soft references when gathering packages. */
+	IncludeSoftReferences = 1 << 0,
+};
+ENUM_CLASS_FLAGS(EOnDemandIsOnDemandOptions);
+
+/** Arguments for checking if any on-demand content exists in the dependecy graph. */
+struct FOnDemandIsOnDemandArgs
+{
+	/** Return true if any of the specified package IDs has on-demand dependencies. */
+	TArray<FPackageId> PackageIds;
+	/** Options */
+	EOnDemandIsOnDemandOptions Options = EOnDemandIsOnDemandOptions::None;
+};
+
+/** Results for GetIsOnDemand */
+struct FOnDemandIsOnDemandResult
+{
+	bool bIsOnDemand = false;
+};
+
 /** Options for controlling the behavior of the install size request. */
 enum class EOnDemandGetCacheUsageOptions
 {
@@ -643,6 +668,18 @@ public:
 	 * @return		The total size by mount ID.
 	 */
 	virtual FIoStatus GetInstallSizesByMountId(const FOnDemandGetInstallSizeArgs& Args, TMap<FString, uint64>& OutSizesByMountId) const = 0;
+	/**
+	 * Query whether any of the specified packages have on-demand dependencies.
+	 * @param Args	Array of package sets to query.
+	 * @return		One result per package set, in the same order as Args.
+	 */
+	virtual FIoStatus GetIsOnDemand(TConstArrayView<FOnDemandIsOnDemandArgs> Args, TArrayView<FOnDemandIsOnDemandResult> OutResults) const = 0;
+	/**
+	 * Query whether any of the specified packages have on-demand dependencies.
+	 * @param Args	Package set to query.
+	 * @return		Results indicating whether on-demand depedencies exist.
+	 */
+	TIoStatusOr<FOnDemandIsOnDemandResult> GetIsOnDemand(const FOnDemandIsOnDemandArgs& Args) const;
 	/** Returns the total cache size in bytes. */
 	virtual FOnDemandCacheUsage GetCacheUsage(const FOnDemandGetCacheUsageArgs& Args) const = 0;
 	/** Dump information about mounted containers to the log */
