Missed file for 47673797

diff --git a/Engine/Source/Runtime/Experimental/IoStore/OnDemandCore/Private/IoStoreOnDemand.cpp b/Engine/Source/Runtime/Experimental/IoStore/OnDemandCore/Private/IoStoreOnDemand.cpp
--- a/Engine/Source/Runtime/Experimental/IoStore/OnDemandCore/Private/IoStoreOnDemand.cpp
+++ b/Engine/Source/Runtime/Experimental/IoStore/OnDemandCore/Private/IoStoreOnDemand.cpp
@@ -559,6 +559,19 @@ IOnDemandIoStore& GetOnDemandIoStore()
 
 FName IOnDemandIoStoreFactory::FeatureName = FName("OnDemandIoStoreFactory");
 
+////////////////////////////////////////////////////////////////////////////////
+TIoStatusOr<FOnDemandIsOnDemandResult> IOnDemandIoStore::GetIsOnDemand(const FOnDemandIsOnDemandArgs& Args) const
+{
+	FOnDemandIsOnDemandResult Result;
+	FIoStatus Status = GetIsOnDemand(MakeConstArrayView(&Args, 1), MakeArrayView(&Result, 1));
+	if (Status.IsOk())
+	{
+		return TIoStatusOr<FOnDemandIsOnDemandResult>(MoveTemp(Result));
+	}
+
+	return Status;
+}
+
 } // namespace UE::IoStore
 
 ////////////////////////////////////////////////////////////////////////////////
