Fix build - missing API markup

diff --git a/Engine/Source/Runtime/Experimental/IoStore/OnDemandCore/Public/IO/IoStoreOnDemand.h b/Engine/Source/Runtime/Experimental/IoStore/OnDemandCore/Public/IO/IoStoreOnDemand.h
--- a/Engine/Source/Runtime/Experimental/IoStore/OnDemandCore/Public/IO/IoStoreOnDemand.h
+++ b/Engine/Source/Runtime/Experimental/IoStore/OnDemandCore/Public/IO/IoStoreOnDemand.h
@@ -679,7 +679,7 @@ public:
 	 * @param Args	Package set to query.
 	 * @return		Results indicating whether on-demand depedencies exist.
 	 */
-	TIoStatusOr<FOnDemandIsOnDemandResult> GetIsOnDemand(const FOnDemandIsOnDemandArgs& Args) const;
+	UE_API TIoStatusOr<FOnDemandIsOnDemandResult> GetIsOnDemand(const FOnDemandIsOnDemandArgs& Args) const;
 	/** Returns the total cache size in bytes. */
 	virtual FOnDemandCacheUsage GetCacheUsage(const FOnDemandGetCacheUsageArgs& Args) const = 0;
 	/** Dump information about mounted containers to the log */
