Chaos Unified Deformable Solver:
-Expose functions with CHAOS_API for unified deformable solver
-Use FSolverParticlesRange instead of FSolverParticles in corotated constraints
- Added CHAOS_API to FClothingPatternData

diff --git a/Engine/Plugins/ChaosCloth/Source/ChaosCloth/Public/ChaosCloth/ChaosClothingPatternData.h b/Engine/Plugins/ChaosCloth/Source/ChaosCloth/Public/ChaosCloth/ChaosClothingPatternData.h
--- a/Engine/Plugins/ChaosCloth/Source/ChaosCloth/Public/ChaosCloth/ChaosClothingPatternData.h
+++ b/Engine/Plugins/ChaosCloth/Source/ChaosCloth/Public/ChaosCloth/ChaosClothingPatternData.h
@@ -15,7 +15,7 @@ struct FClothingPatternData
 	TConstArrayView<uint32> PatternToWeldedIndices;
 	TArray<TVec3<FVec2f>> WeldedFaceVertexPatternPositions; // This will be empty if there are no PatternPositions
 
-	FClothingPatternData(
+	CHAOSCLOTH_API FClothingPatternData(
 		int32 InNumParticles,
 		const TConstArrayView<uint32>& InIndices,
 		const TConstArrayView<FVector2f>& InPatternPositions,
diff --git a/Engine/Source/Runtime/Experimental/Chaos/Public/Chaos/PerParticleDampVelocity.h b/Engine/Source/Runtime/Experimental/Chaos/Public/Chaos/PerParticleDampVelocity.h
--- a/Engine/Source/Runtime/Experimental/Chaos/Public/Chaos/PerParticleDampVelocity.h
+++ b/Engine/Source/Runtime/Experimental/Chaos/Public/Chaos/PerParticleDampVelocity.h
@@ -27,7 +27,7 @@ public:
 	~FPerParticleDampVelocity() = default;
 
 	void UpdatePositionBasedState(const FSolverParticles& Particles, const int32 Offset, const int32 Range);
-	void UpdatePositionBasedState(const FSolverParticlesRange& Particles);
+	UE_INTERNAL CHAOS_API void UpdatePositionBasedState(const FSolverParticlesRange& Particles);
 
 	// Apply damping without first checking for kinematic particles
 	void ApplyFast(FSolverParticles& Particles, const FSolverReal /*Dt*/, const int32 Index) const
@@ -52,7 +52,7 @@ public:
 		return LocalDampingSpace == EChaosSoftsLocalDampingSpace::CenterOfMass;
 	}
 
-	void SetProperties(
+	UE_INTERNAL CHAOS_API void SetProperties(
 		const FCollectionPropertyConstFacade& PropertyCollection,
 		const FSolverVec3& ReferenceLocation,
 		const FSolverVec3& ReferenceVelocity,
diff --git a/Engine/Source/Runtime/Experimental/Chaos/Public/Chaos/XPBDCorotatedConstraints.h b/Engine/Source/Runtime/Experimental/Chaos/Public/Chaos/XPBDCorotatedConstraints.h
--- a/Engine/Source/Runtime/Experimental/Chaos/Public/Chaos/XPBDCorotatedConstraints.h
+++ b/Engine/Source/Runtime/Experimental/Chaos/Public/Chaos/XPBDCorotatedConstraints.h
@@ -109,14 +109,14 @@ namespace Chaos::Softs
 		}
 
 		FXPBDCorotatedTetrahedralConstraints(
-			const FSolverParticles& InParticles,
+			const FSolverParticlesRange& InParticles,
 			const TArray<TVector<int32, 4>>& InMesh,
 			const TArray<FSolverReal>& EMeshArray,
 			const TArray<FSolverReal>& NuMeshArray,
 			TArray<FSolverReal>&& AlphaJMeshArray,
 			const FDeformableXPBDCorotatedParams& InParams,
 			const FSolverReal& NuMesh = (FSolverReal).3,
-			const bool bRecordMetricIn = false, 
+			const bool bRecordMetricIn = false,
 			const bool bDoColoring = true
 		)
 			: CorotatedParams(InParams), AlphaJArray(MoveTemp(AlphaJMeshArray)), bRecordMetric(bRecordMetricIn), MeshConstraints(InMesh)
@@ -133,13 +133,13 @@ namespace Chaos::Softs
 				for (int32 j = 0; j < 4; j++)
 				{
 					ensure(MeshConstraints[e][j] > -1 && MeshConstraints[e][j] < int32(InParticles.Size()));
-					
+
 					if (InParticles.InvM(MeshConstraints[e][j]) == (FSolverReal)0.)
 					{
 						AlphaJArray[e] = (FSolverReal)1.;
 					}
 				}
-				
+
 			}
 
 			for (int e = 0; e < InMesh.Num(); e++)
@@ -168,6 +168,30 @@ namespace Chaos::Softs
 			}
 		}
 
+		FXPBDCorotatedTetrahedralConstraints(
+			const FSolverParticles& InParticles,
+			const TArray<TVector<int32, 4>>& InMesh,
+			const TArray<FSolverReal>& EMeshArray,
+			const TArray<FSolverReal>& NuMeshArray,
+			TArray<FSolverReal>&& AlphaJMeshArray,
+			const FDeformableXPBDCorotatedParams& InParams,
+			const FSolverReal& NuMesh = (FSolverReal).3,
+			const bool bRecordMetricIn = false,
+			const bool bDoColoring = true
+		)
+			: FXPBDCorotatedTetrahedralConstraints(
+				FSolverParticlesRange(const_cast<FSolverParticles*>(&InParticles), /*InOffset = */ 0, /*InRange = */ InParticles.Size()),
+				InMesh,
+				EMeshArray,
+				NuMeshArray,
+				MoveTemp(AlphaJMeshArray),
+				InParams,
+				NuMesh,
+				bRecordMetricIn,
+				bDoColoring
+			)
+		{}
+
 		FXPBDCorotatedTetrahedralConstraints(
 			const FSolverParticles& InParticles,
 			const TArray<TVector<int32, 4>>& InMesh,
@@ -204,16 +228,24 @@ namespace Chaos::Softs
 
 		virtual ~FXPBDCorotatedTetrahedralConstraints() = default;
 
-		PMatrix<FSolverReal, 3, 3> DsInit(const int e, const FSolverParticles& InParticles) const {
+		PMatrix<FSolverReal, 3, 3> DsInit(const int e, const FSolverParticlesRange& InParticles) const
+		{
 			PMatrix<FSolverReal, 3, 3> Result((FSolverReal)0.);
-			for (int i = 0; i < 3; i++) {
-				for (int c = 0; c < 3; c++) {
+			for (int i = 0; i < 3; i++)
+			{
+				for (int c = 0; c < 3; c++)
+				{
 					Result.SetAt(c, i, InParticles.GetX(MeshConstraints[e][i + 1])[c] - InParticles.GetX(MeshConstraints[e][0])[c]);
 				}
 			}
 			return Result;
 		}
 
+		PMatrix<FSolverReal, 3, 3> DsInit(const int e, const FSolverParticles& InParticles) const
+		{	
+			return DsInit(e, FSolverParticlesRange(const_cast<FSolverParticles*>(&InParticles), /*InOffset = */ 0, /*InRange = */ InParticles.Size()));
+		}
+
 		PMatrix<FSolverReal, 3, 3> Ds(const int e, const FSolverParticlesRange& InParticles) const {
 			PMatrix<FSolverReal, 3, 3> Result((FSolverReal)0.);
 			for (int i = 0; i < 3; i++) {
@@ -338,6 +370,11 @@ namespace Chaos::Softs
 			ApplyInParallel(Range, Dt);
 		}
 
+		void Apply(FSolverParticlesRange& Particles, const FSolverReal Dt) const
+		{
+			ApplyInParallel(Particles, Dt);
+		}
+
 		TVec4<TVector<FSolverReal, 3>> GetPolarGradient(const PMatrix<FSolverReal, 3, 3>& Fe, const PMatrix<FSolverReal, 3, 3>& Re, const PMatrix<FSolverReal, 3, 3>& DmInvT, const FSolverReal C1) const
 		{
 			//TVector<FSolverReal, 81> dRdF((FSolverReal)0.);
@@ -446,48 +483,50 @@ namespace Chaos::Softs
 		}
 	protected:
 
-		void InitColor(const FSolverParticles& Particles)
+		void InitColor(const FSolverParticlesRange& Particles)
 		{
-
+			const TArray<TArray<int32>> ConstraintsPerColor = FGraphColoring::ComputeGraphColoringParticlesOrRange(MeshConstraints, Particles, 0, Particles.Size());
+
+			// Reorder constraints based on color so each array in ConstraintsPerColor contains contiguous elements.
+			TArray<TVec4<int32>> ReorderedConstraints;
+			TArray<FSolverReal> ReorderedMeasure;
+			TArray<FSolverReal> ReorderedDmInverse;
+			TArray<int32> OrigToReorderedIndices; // used to reorder stiffness indices
+			ReorderedConstraints.SetNumUninitialized(MeshConstraints.Num());
+			ReorderedMeasure.SetNumUninitialized(Measure.Num());
+			ReorderedDmInverse.SetNumUninitialized(DmInverse.Num());
+			OrigToReorderedIndices.SetNumUninitialized(MeshConstraints.Num());
+
+			ConstraintsPerColorStartIndex.Reset(ConstraintsPerColor.Num() + 1);
+
+			int32 ReorderedIndex = 0;
+			for (const TArray<int32>& ConstraintsBatch : ConstraintsPerColor)
 			{
-				const TArray<TArray<int32>> ConstraintsPerColor = FGraphColoring::ComputeGraphColoring(MeshConstraints, Particles);
-
-				// Reorder constraints based on color so each array in ConstraintsPerColor contains contiguous elements.
-				TArray<TVec4<int32>> ReorderedConstraints;
-				TArray<FSolverReal> ReorderedMeasure;
-				TArray<FSolverReal> ReorderedDmInverse;
-				TArray<int32> OrigToReorderedIndices; // used to reorder stiffness indices
-				ReorderedConstraints.SetNumUninitialized(MeshConstraints.Num());
-				ReorderedMeasure.SetNumUninitialized(Measure.Num());
-				ReorderedDmInverse.SetNumUninitialized(DmInverse.Num());
-				OrigToReorderedIndices.SetNumUninitialized(MeshConstraints.Num());
-
-				ConstraintsPerColorStartIndex.Reset(ConstraintsPerColor.Num() + 1);
-
-				int32 ReorderedIndex = 0;
-				for (const TArray<int32>& ConstraintsBatch : ConstraintsPerColor)
+				ConstraintsPerColorStartIndex.Add(ReorderedIndex);
+				for (const int32& BatchConstraint : ConstraintsBatch)
 				{
-					ConstraintsPerColorStartIndex.Add(ReorderedIndex);
-					for (const int32& BatchConstraint : ConstraintsBatch)
+					const int32 OrigIndex = BatchConstraint;
+					ReorderedConstraints[ReorderedIndex] = MeshConstraints[OrigIndex];
+					ReorderedMeasure[ReorderedIndex] = Measure[OrigIndex];
+					for (int32 kk = 0; kk < 9; kk++)
 					{
-						const int32 OrigIndex = BatchConstraint;
-						ReorderedConstraints[ReorderedIndex] = MeshConstraints[OrigIndex];
-						ReorderedMeasure[ReorderedIndex] = Measure[OrigIndex];
-						for (int32 kk = 0; kk < 9; kk++)
-						{
-							ReorderedDmInverse[9 * ReorderedIndex + kk] = DmInverse[9 * OrigIndex + kk];
-						}
-						OrigToReorderedIndices[OrigIndex] = ReorderedIndex;
-
-						++ReorderedIndex;
+						ReorderedDmInverse[9 * ReorderedIndex + kk] = DmInverse[9 * OrigIndex + kk];
 					}
-				}
-				ConstraintsPerColorStartIndex.Add(ReorderedIndex);
+					OrigToReorderedIndices[OrigIndex] = ReorderedIndex;
 
-				MeshConstraints = MoveTemp(ReorderedConstraints);
-				Measure = MoveTemp(ReorderedMeasure);
-				DmInverse = MoveTemp(ReorderedDmInverse);
+					++ReorderedIndex;
+				}
 			}
+			ConstraintsPerColorStartIndex.Add(ReorderedIndex);
+
+			MeshConstraints = MoveTemp(ReorderedConstraints);
+			Measure = MoveTemp(ReorderedMeasure);
+			DmInverse = MoveTemp(ReorderedDmInverse);
+		}
+
+		void InitColor(const FSolverParticles& Particles)
+		{
+			return InitColor(FSolverParticlesRange(const_cast<FSolverParticles*>(&Particles), /*InOffset = */ 0, /*InRange = */ Particles.Size()));
 		}
 
 		virtual TVec4<TVector<FSolverReal, 3>> GetDeterminantDelta(const FSolverParticlesRange& Particles, const FSolverReal Dt, const int32 ElementIndex, const FSolverReal Tol = (FSolverReal)1e-3) const
