Fix crashes on missing unregistration for Chaos optional data channels in GFPs

Previously the CVD would own the lifetime for the data channels, even after the module defining them winds down.
The fix adds a destructor to FChaosVDOptionalDataChannel and changes storage types so that CVD_DEFINE_OPTIONAL_DATA_CHANNEL constitutes ownership.
The CVD internals have been updated to use weak pointers instead of shared ptr/refs so that ownership stays at the user level.

#rnx

diff --git a/Engine/Plugins/ChaosVD/Source/ChaosVD/Private/ChaosVDEngine.cpp b/Engine/Plugins/ChaosVD/Source/ChaosVD/Private/ChaosVDEngine.cpp
--- a/Engine/Plugins/ChaosVD/Source/ChaosVD/Private/ChaosVDEngine.cpp
+++ b/Engine/Plugins/ChaosVD/Source/ChaosVD/Private/ChaosVDEngine.cpp
@@ -153,11 +153,14 @@ void FChaosVDEngine::RestoreDataChannelsEnabledStateFromSave()
 	}
 
 	using namespace Chaos::VisualDebugger;
-	FChaosVDDataChannelsManager::Get().EnumerateChannels([MiscSettings](const TSharedRef<FChaosVDOptionalDataChannel>& Channel)
+	FChaosVDDataChannelsManager::Get().EnumerateDataChannels([MiscSettings](const TWeakPtr<FChaosVDOptionalDataChannel>& Channel)
 	{
-		if (bool* SavedEnabledState = MiscSettings->DataChannelEnabledState.Find(Channel->GetId().ToString()))
+		if (TSharedPtr<FChaosVDOptionalDataChannel> LockedChannel = Channel.Pin())
 		{
-			Channel->SetChannelEnabled(*SavedEnabledState);
+			if (bool* SavedEnabledState = MiscSettings->DataChannelEnabledState.Find(LockedChannel->GetId().ToString()))
+			{
+				LockedChannel->SetChannelEnabled(*SavedEnabledState);
+			}
 		}
 		return true;
 	});
diff --git a/Engine/Plugins/ChaosVD/Source/ChaosVD/Private/Widgets/ChaosVDRecordingControls.cpp b/Engine/Plugins/ChaosVD/Source/ChaosVD/Private/Widgets/ChaosVDRecordingControls.cpp
--- a/Engine/Plugins/ChaosVD/Source/ChaosVD/Private/Widgets/ChaosVDRecordingControls.cpp
+++ b/Engine/Plugins/ChaosVD/Source/ChaosVD/Private/Widgets/ChaosVDRecordingControls.cpp
@@ -1009,14 +1009,20 @@ void FChaosVDRecordingControls::SetSessionRecordingState(bool bIsRecording, cons
 		RecordingParams.Target = LocalIP->ToString(bAppendPort);
 
 #if WITH_CHAOS_VISUAL_DEBUGGER
-		Chaos::VisualDebugger::FChaosVDDataChannelsManager::Get().EnumerateChannels([&RecordingParams](const TSharedRef<Chaos::VisualDebugger::FChaosVDOptionalDataChannel>& Channel)
 		{
-			if (Channel->IsChannelEnabled())
+			using namespace Chaos::VisualDebugger;
+			FChaosVDDataChannelsManager::Get().EnumerateDataChannels([&RecordingParams](const TWeakPtr<FChaosVDOptionalDataChannel>& Channel)
 			{
-				RecordingParams.DataChannelsEnabledOverrideList.Add(Channel->GetId().ToString());
-			}
-			return true;
-		});
+				if (TSharedPtr<FChaosVDOptionalDataChannel> LockedChannel = Channel.Pin())
+				{
+					if (LockedChannel->IsChannelEnabled())
+					{
+						RecordingParams.DataChannelsEnabledOverrideList.Add(LockedChannel->GetId().ToString());
+					}
+				}
+				return true;
+			});
+		}
 #endif
 
 		if (RecordingParams.TransportMode == EChaosVDTransportMode::Direct)
diff --git a/Engine/Source/Runtime/Experimental/Chaos/Private/ChaosVisualDebugger/ChaosVDOptionalDataChannel.cpp b/Engine/Source/Runtime/Experimental/Chaos/Private/ChaosVisualDebugger/ChaosVDOptionalDataChannel.cpp
--- a/Engine/Source/Runtime/Experimental/Chaos/Private/ChaosVisualDebugger/ChaosVDOptionalDataChannel.cpp
+++ b/Engine/Source/Runtime/Experimental/Chaos/Private/ChaosVisualDebugger/ChaosVDOptionalDataChannel.cpp
@@ -26,9 +26,10 @@ FAutoConsoleCommand SetCVDDataChannelEnabledCommand(
 
 			for (const FString& ChannelName : ParsedChannels)
 			{
-				if (const TSharedPtr<FChaosVDOptionalDataChannel> ChannelInstance = FChaosVDDataChannelsManager::Get().GetChannelById(FName(ChannelName)))
+				TWeakPtr<FChaosVDOptionalDataChannel> ChannelInstance = FChaosVDDataChannelsManager::Get().GetChannelById(FName(ChannelName));
+				if (TSharedPtr<FChaosVDOptionalDataChannel> LockedChannelInstance = ChannelInstance.Pin())
 				{
-					ChannelInstance->SetChannelEnabled(bNewEnabledState);
+					LockedChannelInstance->SetChannelEnabled(bNewEnabledState);
 				}
 			}
 		}
@@ -37,9 +38,14 @@ FAutoConsoleCommand SetCVDDataChannelEnabledCommand(
 
 namespace Chaos::VisualDebugger
 {
+	FChaosVDOptionalDataChannel::~FChaosVDOptionalDataChannel()
+	{
+		FChaosVDDataChannelsManager::Get().UnregisterChannel(*ChannelId);
+	}
+	
 	void FChaosVDOptionalDataChannel::Initialize()
 	{
-		FChaosVDDataChannelsManager::Get().RegisterChannel(AsShared());
+		FChaosVDDataChannelsManager::Get().RegisterChannel(AsWeak());
 	}
 
 	const FText& FChaosVDOptionalDataChannel::GetDisplayName() const
@@ -73,6 +79,14 @@ namespace Chaos::VisualDebugger
 
 		return NewChannel;
 	}
+
+	TSharedPtr<FChaosVDOptionalDataChannel> CreateDataChannelShared(FName InChannelID, const TSharedRef<FText>& InDisplayName, EChaosVDDataChannelInitializationFlags InitializationFlags)
+	{
+		TSharedPtr<FChaosVDOptionalDataChannel> NewChannel = MakeShared<FChaosVDOptionalDataChannel>(MakeShared<FName>(InChannelID), InDisplayName, InitializationFlags);
+		NewChannel->Initialize();
+
+		return NewChannel;
+	}
 }
 
 CVD_DEFINE_OPTIONAL_DATA_CHANNEL(Default, EChaosVDDataChannelInitializationFlags::StartEnabled);
diff --git a/Engine/Source/Runtime/Experimental/Chaos/Private/ChaosVisualDebugger/ChaosVisualDebuggerTrace.cpp b/Engine/Source/Runtime/Experimental/Chaos/Private/ChaosVisualDebugger/ChaosVisualDebuggerTrace.cpp
--- a/Engine/Source/Runtime/Experimental/Chaos/Private/ChaosVisualDebugger/ChaosVisualDebuggerTrace.cpp
+++ b/Engine/Source/Runtime/Experimental/Chaos/Private/ChaosVisualDebugger/ChaosVisualDebuggerTrace.cpp
@@ -1424,16 +1424,19 @@ void FChaosVisualDebuggerTrace::OverrideDefaultEnabledDataChannels(TConstArrayVi
 	UE_LOG(LogChaos, Log, TEXT("[%s] Channel list override provided - Enabling [%d] Requested channels..."), ANSI_TO_TCHAR(__FUNCTION__), EnabledDataChannelsOverrideList.Num());
 
 	using namespace Chaos::VisualDebugger;
-	FChaosVDDataChannelsManager::Get().EnumerateChannels([&EnabledDataChannelsOverrideList](const TSharedRef<FChaosVDOptionalDataChannel>& Channel)
+	FChaosVDDataChannelsManager::Get().EnumerateDataChannels([&EnabledDataChannelsOverrideList](const TWeakPtr<FChaosVDOptionalDataChannel>& Channel)
 	{
-		if (Channel->CanChangeEnabledState())
+		if (TSharedPtr<FChaosVDOptionalDataChannel> LockedChannel = Channel.Pin())
 		{
-			// This is far from efficient, but this will be called once when the recording start command is executed, and we only have a handful of channels
-			const FString ChannelIdAsString = Channel->GetId().ToString();
-			const bool bChannelShouldBeEnabled = EnabledDataChannelsOverrideList.Contains(ChannelIdAsString);
-			Channel->SetChannelEnabled(bChannelShouldBeEnabled);
+			if (LockedChannel->CanChangeEnabledState())
+			{
+				// This is far from efficient, but this will be called once when the recording start command is executed, and we only have a handful of channels
+				const FString ChannelIdAsString = LockedChannel->GetId().ToString();
+				const bool bChannelShouldBeEnabled = EnabledDataChannelsOverrideList.Contains(ChannelIdAsString);
+				LockedChannel->SetChannelEnabled(bChannelShouldBeEnabled);
 
-			UE_LOG(LogChaos, Log, TEXT("[%s] Setting enabled state for channel [%s] to [%s]..."), ANSI_TO_TCHAR(__FUNCTION__), *ChannelIdAsString, bChannelShouldBeEnabled ? TEXT("True") : TEXT("False"));
+				UE_LOG(LogChaos, Log, TEXT("[%s] Setting enabled state for channel [%s] to [%s]..."), ANSI_TO_TCHAR(__FUNCTION__), *ChannelIdAsString, bChannelShouldBeEnabled ? TEXT("True") : TEXT("False"));
+			}
 		}
 		return true;
 	});
diff --git a/Engine/Source/Runtime/Experimental/Chaos/Public/ChaosVisualDebugger/ChaosVDContextProvider.h b/Engine/Source/Runtime/Experimental/Chaos/Public/ChaosVisualDebugger/ChaosVDContextProvider.h
--- a/Engine/Source/Runtime/Experimental/Chaos/Public/ChaosVisualDebugger/ChaosVDContextProvider.h
+++ b/Engine/Source/Runtime/Experimental/Chaos/Public/ChaosVisualDebugger/ChaosVDContextProvider.h
@@ -35,15 +35,25 @@ struct FChaosVDContext
 	int32 Type = INDEX_NONE;
 	int32 Attributes = 0;
 
-	FORCEINLINE void SetDataChannel(const TSharedRef<Chaos::VisualDebugger::FChaosVDOptionalDataChannel>& NewDataChannel)
+	FORCEINLINE void SetDataChannel(const TWeakPtr<Chaos::VisualDebugger::FChaosVDOptionalDataChannel>& NewDataChannel)
 	{
 		CurrentDataChannel = NewDataChannel;
 	}
 
-	FORCEINLINE bool IsDataChannelEnabled() const { return CurrentDataChannel->IsChannelEnabled(); }
+	FORCEINLINE bool IsDataChannelEnabled() const
+	{
+		if (TSharedPtr<Chaos::VisualDebugger::FChaosVDOptionalDataChannel> LockedCurrentDataChannel = CurrentDataChannel.Pin())
+		{
+			return LockedCurrentDataChannel->IsChannelEnabled();
+		}
+		else
+		{
+			return false;
+		}
+	}
 
 private:
-	TSharedRef<Chaos::VisualDebugger::FChaosVDOptionalDataChannel> CurrentDataChannel = CVDDC_Default;
+	TWeakPtr<Chaos::VisualDebugger::FChaosVDOptionalDataChannel> CurrentDataChannel = CVDDC_Default;
 };
 
 UE_DECLARE_THREAD_SINGLETON_TLS(class FChaosVDThreadContext, CHAOS_API)
@@ -141,7 +151,7 @@ struct FChaosVDScopeContext
 
 struct FChaosCVDScopedDataChannelOverride
 {
-	FChaosCVDScopedDataChannelOverride(const TSharedRef<Chaos::VisualDebugger::FChaosVDOptionalDataChannel>& NewDataChannel)
+	FChaosCVDScopedDataChannelOverride(const TWeakPtr<Chaos::VisualDebugger::FChaosVDOptionalDataChannel>& NewDataChannel)
 	{
 		if (const FChaosVDContext* CVDContextData = FChaosVDThreadContext::Get().GetCurrentContext())
 		{
diff --git a/Engine/Source/Runtime/Experimental/Chaos/Public/ChaosVisualDebugger/ChaosVDOptionalDataChannel.h b/Engine/Source/Runtime/Experimental/Chaos/Public/ChaosVisualDebugger/ChaosVDOptionalDataChannel.h
--- a/Engine/Source/Runtime/Experimental/Chaos/Public/ChaosVisualDebugger/ChaosVDOptionalDataChannel.h
+++ b/Engine/Source/Runtime/Experimental/Chaos/Public/ChaosVisualDebugger/ChaosVDOptionalDataChannel.h
@@ -32,6 +32,8 @@ namespace Chaos::VisualDebugger
 			bCanChangeEnabledState = EnumHasAnyFlags(InitializationFlags, EChaosVDDataChannelInitializationFlags::CanChangeEnabledState);
 		}
 
+		virtual ~FChaosVDOptionalDataChannel();
+
 		/** Registers itself with the optional channel data manager */
 		void Initialize();
 
@@ -77,11 +79,29 @@ namespace Chaos::VisualDebugger
 		 * If the callback returns false, the iteration will be stopped.
 		 */
 		template<typename TCallback>
+		UE_DEPRECATED(5.9, "please use EnumerateDataChannels instead")
 		void EnumerateChannels(TCallback Callback)
 		{
-			for (const TPair<FName, TSharedPtr<FChaosVDOptionalDataChannel>>& ChannelWithName : AvailableChannels)
+			for (const TPair<FName, TWeakPtr<FChaosVDOptionalDataChannel>>& ChannelWithName : AvailableChannels)
+			{
+				if (TSharedPtr<FChaosVDOptionalDataChannel> LockedChannel = ChannelWithName.Value.Pin())
+				{
+					if (!Callback(LockedChannel))
+					{
+						return;
+					}
+				}
+			}
+		}
+
+		/** Iterates through all the available data channels and executes the provided callback passing each channel as an argument.
+		 * If the callback returns false, the iteration will be stopped.
+		 */
+		void EnumerateDataChannels(TFunction<bool(const TWeakPtr<FChaosVDOptionalDataChannel>&)> Callback)
+		{
+			for (const TPair<FName, TWeakPtr<FChaosVDOptionalDataChannel>>& ChannelWithName : AvailableChannels)
 			{
-				if (!Callback(ChannelWithName.Value.ToSharedRef()))
+				if (!Callback(ChannelWithName.Value))
 				{
 					return;
 				}
@@ -89,9 +109,9 @@ namespace Chaos::VisualDebugger
 		}
 
 		/** Returns the CVD Data Channel instance for the provided ID*/
-		TSharedPtr<FChaosVDOptionalDataChannel> GetChannelById(FName ChannelId)
+		TWeakPtr<FChaosVDOptionalDataChannel> GetChannelById(FName ChannelId)
 		{
-			if (TSharedPtr<FChaosVDOptionalDataChannel>* ChannelFound = AvailableChannels.Find(ChannelId))
+			if (TWeakPtr<FChaosVDOptionalDataChannel>* ChannelFound = AvailableChannels.Find(ChannelId))
 			{
 				return *ChannelFound;
 			}
@@ -112,12 +132,21 @@ namespace Chaos::VisualDebugger
 	private:
 		
 		/** Registers a CVD Data channel instance - Do not call directly */
-		void RegisterChannel(const TSharedRef<FChaosVDOptionalDataChannel>& NewChannel)
+		void RegisterChannel(const TWeakPtr<FChaosVDOptionalDataChannel>& NewChannel)
 		{
-			AvailableChannels.Add(NewChannel->GetId(), NewChannel);
+			if (TSharedPtr<FChaosVDOptionalDataChannel> LockedNewChannel = NewChannel.Pin())
+			{
+				AvailableChannels.Add(LockedNewChannel->GetId(), NewChannel);
+			}
+		}
+		
+		/** Remove a CVD Data channel instance - Do not call directly */
+		void UnregisterChannel(FName ChannelId)
+		{
+			AvailableChannels.Remove(ChannelId);
 		}
 
-		TMap<FName, TSharedPtr<FChaosVDOptionalDataChannel>> AvailableChannels;
+		TMap<FName, TWeakPtr<FChaosVDOptionalDataChannel>> AvailableChannels;
 
 		friend FChaosVDOptionalDataChannel;
 	};
@@ -125,7 +154,11 @@ namespace Chaos::VisualDebugger
 	void ParseChannelListFromCommandArgument(TArray<FString>& OutParsedChannelList, const FString& InCommandArgument);
 
 	/** Creates a CVD Data Channel Instance - Only to be used by the CVD Macros */
+	UE_DEPRECATED(5.9, "please use CreateDataChannelShared instead to keep the lifetime of the channel explicit");
 	CHAOS_API TSharedRef<FChaosVDOptionalDataChannel> CreateDataChannel(FName InChannelID, const TSharedRef<FText>& InDisplayName, EChaosVDDataChannelInitializationFlags InitializationFlags);
+
+	/** Creates a CVD Data Channel Instance - Only to be used by the CVD Macros */
+	CHAOS_API TSharedPtr<FChaosVDOptionalDataChannel> CreateDataChannelShared(FName InChannelID, const TSharedRef<FText>& InDisplayName, EChaosVDDataChannelInitializationFlags InitializationFlags);
 }
 
 #endif //WITH_CHAOS_VISUAL_DEBUGGER
@@ -151,19 +184,19 @@ namespace Chaos::VisualDebugger
 /** Declares an Optional CVD Data channel to be available globally. The data channel can be accessed by using CVDDC_TheNameOfTheChannelUsedWithThisMacro */
 #ifndef CVD_DECLARE_OPTIONAL_DATA_CHANNEL
 	#define CVD_DECLARE_OPTIONAL_DATA_CHANNEL(DataChannelName) \
-			extern CHAOS_API TSharedRef<Chaos::VisualDebugger::FChaosVDOptionalDataChannel> CVDDC_##DataChannelName;
+			extern CHAOS_API TSharedPtr<Chaos::VisualDebugger::FChaosVDOptionalDataChannel> CVDDC_##DataChannelName;
 #endif
 
 /** Declares an Optional CVD Data channel outside the Chaos Module to be available globally. The data channel can be accessed by using CVDDC_TheNameOfTheChannelUsedWithThisMacro */
 #ifndef CVD_DECLARE_OPTIONAL_DATA_CHANNEL_EXTERN
 	#define CVD_DECLARE_OPTIONAL_DATA_CHANNEL_EXTERN(DataChannelName, ModuleAPI) \
-		extern ModuleAPI TSharedRef<Chaos::VisualDebugger::FChaosVDOptionalDataChannel> CVDDC_##DataChannelName;
+		extern ModuleAPI TSharedPtr<Chaos::VisualDebugger::FChaosVDOptionalDataChannel> CVDDC_##DataChannelName;
 #endif
 
 /** Defines and initializes an Optional CVD Data Channel */
 #ifndef CVD_DEFINE_OPTIONAL_DATA_CHANNEL
 	#define CVD_DEFINE_OPTIONAL_DATA_CHANNEL(DataChannelName, InitializationFlags) \
-			TSharedRef<Chaos::VisualDebugger::FChaosVDOptionalDataChannel> CVDDC_##DataChannelName = Chaos::VisualDebugger::CreateDataChannel(#DataChannelName, MakeShared<FText>(NSLOCTEXT(CVD_STRINGIZE(ChaosVisualDebugger), CVD_STRINGIZE(CVD_CONCAT(DataChannelName,_ChannelName)), #DataChannelName)), InitializationFlags);
+			TSharedPtr<Chaos::VisualDebugger::FChaosVDOptionalDataChannel> CVDDC_##DataChannelName = Chaos::VisualDebugger::CreateDataChannel(#DataChannelName, MakeShared<FText>(NSLOCTEXT(CVD_STRINGIZE(ChaosVisualDebugger), CVD_STRINGIZE(CVD_CONCAT(DataChannelName,_ChannelName)), #DataChannelName)), InitializationFlags);
 #endif
 
 
diff --git a/Engine/Source/Runtime/Experimental/ChaosSolverEngine/Private/Chaos/ChaosVDRemoteSessionsManager.cpp b/Engine/Source/Runtime/Experimental/ChaosSolverEngine/Private/Chaos/ChaosVDRemoteSessionsManager.cpp
--- a/Engine/Source/Runtime/Experimental/ChaosSolverEngine/Private/Chaos/ChaosVDRemoteSessionsManager.cpp
+++ b/Engine/Source/Runtime/Experimental/ChaosSolverEngine/Private/Chaos/ChaosVDRemoteSessionsManager.cpp
@@ -640,9 +640,10 @@ void FChaosVDRemoteSessionsManager::HandleChangeDataChannelStateCommandMessage(c
 	UE_AUTORTFM_ONCOMMIT(InMessage)
 	{
 		using namespace Chaos::VisualDebugger;
-		if (TSharedPtr<FChaosVDOptionalDataChannel> ChannelInstance = FChaosVDDataChannelsManager::Get().GetChannelById(FName(InMessage.NewState.ChannelName)))
+		TWeakPtr<FChaosVDOptionalDataChannel> ChannelInstance = FChaosVDDataChannelsManager::Get().GetChannelById(FName(InMessage.NewState.ChannelName));
+		if (TSharedPtr<FChaosVDOptionalDataChannel> LockedChannelInstance = ChannelInstance.Pin())
 		{
-			ChannelInstance->SetChannelEnabled(InMessage.NewState.bIsEnabled);
+			LockedChannelInstance->SetChannelEnabled(InMessage.NewState.bIsEnabled);
 		}
 	};
 #endif
@@ -677,10 +678,13 @@ void FChaosVDRemoteSessionsManager::HandleFullSessionStateRequestMessage(const F
 		FullSessionStateResponse->bIsRecording = FChaosVDRuntimeModule::Get().IsRecording();
 
 		using namespace Chaos::VisualDebugger;
-		FChaosVDDataChannelsManager::Get().EnumerateChannels([&FullSessionStateResponse](const TSharedRef<FChaosVDOptionalDataChannel>& Channel)
+		FChaosVDDataChannelsManager::Get().EnumerateDataChannels([&FullSessionStateResponse](const TWeakPtr<FChaosVDOptionalDataChannel>& Channel)
 		{
-			FChaosVDChannelStateChangeCommandMessage DataChannelStateMessage = {Channel->GetId().ToString(), Channel->IsChannelEnabled(), Channel->CanChangeEnabledState() };
-			FullSessionStateResponse->DataChannelsStates.Emplace(DataChannelStateMessage.NewState);
+			if (TSharedPtr<FChaosVDOptionalDataChannel> LockedChannel = Channel.Pin())
+			{
+				FChaosVDChannelStateChangeCommandMessage DataChannelStateMessage = { LockedChannel->GetId().ToString(), LockedChannel->IsChannelEnabled(), LockedChannel->CanChangeEnabledState() };
+				FullSessionStateResponse->DataChannelsStates.Emplace(DataChannelStateMessage.NewState);
+			}
 			return true;
 		});
 #endif
