Fix for issue where removing and adding a component to a cluster union in one frame resulted in part not being added again due to operation sorting
Behind a CVar in case we need to revert for any project

diff --git a/Engine/Source/Runtime/Experimental/Chaos/Private/Chaos/ClusterUnionManager.cpp b/Engine/Source/Runtime/Experimental/Chaos/Private/Chaos/ClusterUnionManager.cpp
--- a/Engine/Source/Runtime/Experimental/Chaos/Private/Chaos/ClusterUnionManager.cpp
+++ b/Engine/Source/Runtime/Experimental/Chaos/Private/Chaos/ClusterUnionManager.cpp
@@ -40,6 +40,14 @@ namespace Chaos
 			TEXT("Fix the partial destruction.")
 		);
 
+		// Aim is to prevent re-ordering of operatioins when performing a remove and add in the same frame
+		bool bFixRemoveDuringSort = true;
+		FAutoConsoleVariableRef CVarChaosFixRemoveDuringSort(
+			TEXT("p.Chaos.ClusterUnion.FixRemoveDuringSort"),
+			bFixRemoveDuringSort,
+			TEXT("Fix remove to top during sort.")
+		);
+
 		FRigidTransform3 GetParticleRigidFrameInClusterUnion(FPBDRigidParticleHandle* Child, const FRigidTransform3& ClusterWorldTM)
 		{
 			FRigidTransform3 Frame = FRigidTransform3::Identity;
@@ -442,6 +450,19 @@ namespace Chaos
 				OpMap.Value.KeyStableSort(
 					[](EClusterUnionOperation A, EClusterUnionOperation B)
 					{
+						// remove must be bubbled up to the top (for the case when remove and add are used at the same time)
+						if (bFixRemoveDuringSort)
+						{
+							if (A == EClusterUnionOperation::Remove)
+							{
+								return true;
+							}
+							else if (B == EClusterUnionOperation::Remove)
+							{
+								return false;
+							}
+						}
+
 						return static_cast<int32>(A) < static_cast<int32>(B);
 					}
 				);
