iax: Speculative fix for PGO segfault

#rnx

diff --git a/Engine/Source/Runtime/Experimental/IoStore/HttpClient/Private/TransactionTwo.inl b/Engine/Source/Runtime/Experimental/IoStore/HttpClient/Private/TransactionTwo.inl
--- a/Engine/Source/Runtime/Experimental/IoStore/HttpClient/Private/TransactionTwo.inl
+++ b/Engine/Source/Runtime/Experimental/IoStore/HttpClient/Private/TransactionTwo.inl
@@ -882,6 +882,11 @@ FOutcome TickTwo(FTlsPeer& Peer, void* PeerData)
 	using namespace DetailTwo;
 
 	auto* Driver = (FDriverNg*)PeerData;
+	if (Driver == nullptr)
+	{
+		return FOutcome::Error("ng-driver is unexpectedly null");
+	}
+
 	return Driver->GetTransactId(Peer);
 }
 
