Fix for failing headless chaos unit tests w/ new compact set --

make GeometryCollectionAlgo::TriangulateBoundaries cycle the boundary vertices to try to find a starting vertex that isn't locally co-linear with adjacent vertices

diff --git a/Engine/Source/Runtime/Experimental/Chaos/Private/GeometryCollection/GeometryCollectionAlgo.cpp b/Engine/Source/Runtime/Experimental/Chaos/Private/GeometryCollection/GeometryCollectionAlgo.cpp
--- a/Engine/Source/Runtime/Experimental/Chaos/Private/GeometryCollection/GeometryCollectionAlgo.cpp
+++ b/Engine/Source/Runtime/Experimental/Chaos/Private/GeometryCollection/GeometryCollectionAlgo.cpp
@@ -926,18 +926,48 @@ namespace GeometryCollectionAlgo
 				{
 					continue;
 				}
-				int32 First = Boundary[0];
+
+				// Note: technically this is 4x the squared area, but matches how Chaos generally filters degenerate triangles
+				auto GetAreaSq = [GeometryCollection](int32 A, int32 B, int32 C) -> double
+				{
+					FVector p10(GeometryCollection->Vertex[B] - GeometryCollection->Vertex[A]);
+					FVector p20(GeometryCollection->Vertex[C] - GeometryCollection->Vertex[A]);
+					FVector Cross = FVector::CrossProduct(p20, p10);
+					return Cross.SizeSquared();
+				};
+
+				// Try to find a start index that isn't a locally degenerate/collinear chain (if we have a positive degenerate tri threshold)
+				int32 UseStartIdx = 0;
+				if (MinTriangleAreaSq > 0 && Boundary.Num() > 3)
+				{
+					int32 NumLocalNonCollinear = 0;
+					for (int32 Prev2 = Boundary.Num() - 2, Prev = Boundary.Num() - 1, Idx = 0; Idx < Boundary.Num(); Prev2 = Prev, Prev = Idx++)
+					{
+						if (GetAreaSq(Boundary[Prev2], Boundary[Prev], Boundary[Idx]) < MinTriangleAreaSq)
+						{
+							NumLocalNonCollinear = 0;
+						}
+						else
+						{
+							++NumLocalNonCollinear;
+						}
+						if (NumLocalNonCollinear == 3)
+						{
+							UseStartIdx = Prev2;
+							break;
+						}
+					}
+				}
+
+				int32 First = Boundary[UseStartIdx];
 				
 				for (int32 BoundaryIdx = 1; BoundaryIdx + 1 < Boundary.Num(); BoundaryIdx++)
 				{
-					int32 A = First, B = Boundary[BoundaryIdx], C = Boundary[BoundaryIdx + 1];
+					int32 A = First, B = Boundary[(UseStartIdx + BoundaryIdx) % Boundary.Num()], C = Boundary[(UseStartIdx + BoundaryIdx + 1) % Boundary.Num()];
 
 					if (MinTriangleAreaSq > 0)
 					{
-						FVector p10(GeometryCollection->Vertex[B] - GeometryCollection->Vertex[A]);
-						FVector p20(GeometryCollection->Vertex[C] - GeometryCollection->Vertex[A]);
-						FVector Cross = FVector::CrossProduct(p20, p10);
-						if (Cross.SizeSquared() < MinTriangleAreaSq)
+						if (GetAreaSq(A, B, C) < MinTriangleAreaSq)
 						{
 							continue;
 						}
