Log spam cleanup

diff --git a/Engine/Source/Runtime/Experimental/IoStore/OnDemand/Private/OnDemandContentInstaller.cpp b/Engine/Source/Runtime/Experimental/IoStore/OnDemand/Private/OnDemandContentInstaller.cpp
--- a/Engine/Source/Runtime/Experimental/IoStore/OnDemand/Private/OnDemandContentInstaller.cpp
+++ b/Engine/Source/Runtime/Experimental/IoStore/OnDemand/Private/OnDemandContentInstaller.cpp
@@ -1084,7 +1084,7 @@ void FOnDemandContentInstaller::ExectuteFlushLastAccessRequest(FRequest& Request
 
 	const FRequest::FFlushLastAccess& FlushLastAccessReqest = Request.AsFlushLastAccess();
 
-	UE_LOG(LogIoStoreOnDemand, Log, TEXT("Executing flush last access request"));
+	UE_LOG(LogIoStoreOnDemand, VeryVerbose, TEXT("Executing flush last access request"));
 
 	if (Request.IsCancelled())
 	{
@@ -1114,6 +1114,7 @@ void FOnDemandContentInstaller::ExectuteFlushLastAccessRequest(FRequest& Request
 
 	if (bLastAccessDirty)
 	{
+		UE_LOG(LogIoStoreOnDemand, Log, TEXT("Flushing last access to journal"));
 		Request.Result = IoStore.InstallCache->FlushLastAccess();
 	}
 
@@ -1665,8 +1666,15 @@ void FOnDemandContentInstaller::CompleteFlushLastAccessRequest(FRequest& Request
 {
 	const double DurationInSeconds = FPlatformTime::ToSeconds64(FPlatformTime::Cycles64() - Request.StartTimeCycles);
 
-	UE_LOG(LogIoStoreOnDemand, Log, TEXT("Flush last access request completed, Result='%s', Duration=%d ms"),
-		*Request.GetErrorString(), Private::SecondsToMillicseconds(DurationInSeconds));
+	static constexpr const TCHAR FmtStr[] = TEXT("Flush last access request completed, Result='%s', Duration=%d ms");
+	if (Request.Result.HasError())
+	{
+		UE_LOG(LogIoStoreOnDemand, Error, FmtStr, *Request.GetErrorString(), Private::SecondsToMillicseconds(DurationInSeconds));
+	}
+	else
+	{
+		UE_LOG(LogIoStoreOnDemand, VeryVerbose, FmtStr, *Request.GetErrorString(), Private::SecondsToMillicseconds(DurationInSeconds));
+	}
 
 	FOnDemandFlushLastAccessResult FlushResult;
 	FlushResult.DurationInSeconds	= DurationInSeconds;
