Collision Filter Refactor Part 5.1.0 Replacing GetCollisionChannel with HasCollisionChannel in the editor to more properly support being in more than one collision channel.

diff --git a/Engine/Source/Runtime/Engine/Private/Collision/PhysicsFiltering.cpp b/Engine/Source/Runtime/Engine/Private/Collision/PhysicsFiltering.cpp
--- a/Engine/Source/Runtime/Engine/Private/Collision/PhysicsFiltering.cpp
+++ b/Engine/Source/Runtime/Engine/Private/Collision/PhysicsFiltering.cpp
@@ -185,6 +185,16 @@ ECollisionChannel GetCollisionChannel(const Chaos::Filter::FQueryFilterData& Que
 	return (ECollisionChannel)QueryFilterData.GetCollisionChannelIndex();
 }
 
+bool HasCollisionChannel(const Chaos::Filter::FShapeFilterData& ShapeFilterData, ECollisionChannel Channel)
+{
+	return ShapeFilterData.IsCollisionChannelSet(Channel);
+}
+
+bool HasCollisionChannel(const Chaos::Filter::FQueryFilterData& QueryFilterData, ECollisionChannel Channel)
+{
+	return QueryFilterData.IsCollisionChannelSet(Channel);
+}
+
 FCollisionResponseContainer ExtractCollisionResponseContainer(const Chaos::Filter::FShapeFilterData& ShapeFilterData)
 {
 	FCollisionResponseContainer CollisionResponseContainer;
diff --git a/Engine/Source/Runtime/Engine/Public/Physics/PhysicsFiltering.h b/Engine/Source/Runtime/Engine/Public/Physics/PhysicsFiltering.h
--- a/Engine/Source/Runtime/Engine/Public/Physics/PhysicsFiltering.h
+++ b/Engine/Source/Runtime/Engine/Public/Physics/PhysicsFiltering.h
@@ -206,6 +206,8 @@ inline void CreateShapeFilterData(
 
 ENGINE_API ECollisionChannel GetCollisionChannel(const Chaos::Filter::FShapeFilterData& ShapeFilterData);
 ENGINE_API ECollisionChannel GetCollisionChannel(const Chaos::Filter::FQueryFilterData& QueryFilterData);
+ENGINE_API bool HasCollisionChannel(const Chaos::Filter::FShapeFilterData& ShapeFilterData, ECollisionChannel Channel);
+ENGINE_API bool HasCollisionChannel(const Chaos::Filter::FQueryFilterData& QueryFilterData, ECollisionChannel Channel);
 
 inline ECollisionChannel GetCollisionChannel(uint32 Word3)
 {
diff --git a/Engine/Source/Runtime/Experimental/Chaos/Private/Chaos/CollisionFilterData.cpp b/Engine/Source/Runtime/Experimental/Chaos/Private/Chaos/CollisionFilterData.cpp
--- a/Engine/Source/Runtime/Experimental/Chaos/Private/Chaos/CollisionFilterData.cpp
+++ b/Engine/Source/Runtime/Experimental/Chaos/Private/Chaos/CollisionFilterData.cpp
@@ -178,6 +178,16 @@ namespace Chaos::Filter
 		return Chaos::Filter::GetCollisionChannelMask64(Word3);
 	}
 
+	bool FShapeFilterData::IsCollisionChannelSet(const uint32 ChannelIndex) const
+	{
+		if (ChannelIndex < 64)
+		{
+			const uint64 Mask = 1LLU << ChannelIndex;
+			return GetCollisionChannelMask() & Mask;
+		}
+		return false;
+	}
+
 	uint64 FShapeFilterData::GetQueryBlockChannels() const
 	{
 		return QueryBlockChannels;
@@ -243,6 +253,16 @@ namespace Chaos::Filter
 		return Chaos::Filter::GetCollisionChannelMask64(Word3);
 	}
 
+	bool FQueryFilterData::IsCollisionChannelSet(const uint32 ChannelIndex) const
+	{
+		if (ChannelIndex < 64)
+		{
+			const uint64 Mask = 1LLU << ChannelIndex;
+			return GetCollisionChannelMask() & Mask;
+		}
+		return false;
+	}
+
 	uint64 FQueryFilterData::GetBlockChannels() const
 	{
 		return Word1;
diff --git a/Engine/Source/Runtime/Experimental/Chaos/Public/Chaos/CollisionFilterData.h b/Engine/Source/Runtime/Experimental/Chaos/Public/Chaos/CollisionFilterData.h
--- a/Engine/Source/Runtime/Experimental/Chaos/Public/Chaos/CollisionFilterData.h
+++ b/Engine/Source/Runtime/Experimental/Chaos/Public/Chaos/CollisionFilterData.h
@@ -125,6 +125,7 @@ namespace Chaos::Filter
 
 		CHAOS_API uint8 GetCollisionChannelIndex() const;
 		CHAOS_API uint64 GetCollisionChannelMask() const;
+		CHAOS_API bool IsCollisionChannelSet(const uint32 ChannelIndex) const;
 
 		CHAOS_API uint64 GetQueryBlockChannels() const;
 		CHAOS_API uint64 GetQueryOverlapChannels() const;
@@ -173,6 +174,7 @@ namespace Chaos::Filter
 		// Channel API
 		CHAOS_API uint8 GetCollisionChannelIndex() const;
 		CHAOS_API uint64 GetCollisionChannelMask() const;
+		CHAOS_API bool IsCollisionChannelSet(const uint32 ChannelIndex) const;
 		CHAOS_API uint64 GetBlockChannels() const;
 		CHAOS_API uint64 GetOverlapChannels() const;
 
