iax: Better version selection that works with ZenServer as a host.

#rnx

diff --git a/Engine/Source/Runtime/Experimental/IoStore/OnDemand/Private/OnDemandHttpClient.cpp b/Engine/Source/Runtime/Experimental/IoStore/OnDemand/Private/OnDemandHttpClient.cpp
--- a/Engine/Source/Runtime/Experimental/IoStore/OnDemand/Private/OnDemandHttpClient.cpp
+++ b/Engine/Source/Runtime/Experimental/IoStore/OnDemand/Private/OnDemandHttpClient.cpp
@@ -685,8 +685,11 @@ TUniquePtr<HTTP::FConnectionPool> FMultiEndpointHttpClient::CreateConnection(FAn
 		Params.SendBufSize = Config.SendBufferSize;
 	}
 
+	bool bUseHttpTwo = (GIaxHttpVersion == 2);
+	bUseHttpTwo &= HostUrl.StartsWith("https://");
+	Params.HttpVersion = bUseHttpTwo ? EHttpVersion::Two : EHttpVersion::One;
+
 	Params.ConnectionCount = uint16(Config.MaxConnectionCount);
-	Params.HttpVersion = (GIaxHttpVersion == 2) ? EHttpVersion::Two : EHttpVersion::One;
 	Params.MaxInflight = uint8(FMath::Clamp(GIaxHttpMaxInflight, 1, 64));
 	if (Params.HttpVersion == EHttpVersion::Two)
 	{
diff --git a/Engine/Source/Runtime/Experimental/IoStore/OnDemand/Private/OnDemandHttpThread.cpp b/Engine/Source/Runtime/Experimental/IoStore/OnDemand/Private/OnDemandHttpThread.cpp
--- a/Engine/Source/Runtime/Experimental/IoStore/OnDemand/Private/OnDemandHttpThread.cpp
+++ b/Engine/Source/Runtime/Experimental/IoStore/OnDemand/Private/OnDemandHttpThread.cpp
@@ -98,7 +98,7 @@ int32 GIaxHttpVersion = 2;
 static FAutoConsoleVariableRef CVar_IaxHttpVersion(
 	TEXT("iax.HttpVersion"),
 	GIaxHttpVersion,
-	TEXT("Protocol version to use, either '1' or '2'"),
+	TEXT("Protocol version to prefer, either '1' or '2'. The latter requires https."),
 	FConsoleVariableDelegate::CreateStatic(&OnHttpClientCVarChanged)
 );
 
diff --git a/Engine/Source/Runtime/Experimental/IoStore/OnDemand/Private/Statistics.cpp b/Engine/Source/Runtime/Experimental/IoStore/OnDemand/Private/Statistics.cpp
--- a/Engine/Source/Runtime/Experimental/IoStore/OnDemand/Private/Statistics.cpp
+++ b/Engine/Source/Runtime/Experimental/IoStore/OnDemand/Private/Statistics.cpp
@@ -27,7 +27,6 @@ LLM_DEFINE_TAG(Ias);
 namespace UE::IoStore
 {
 
-extern int32 GIaxHttpVersion;
 extern int32 GHttpIoDispatcherMaxRangeKiB;
 
 int32 CalculateRequestQueueLengthMax();
@@ -871,7 +870,6 @@ void FOnDemandIoBackendStats::PrintOnScreenStatistics(FCoreDelegates::FSeverityM
 
 			TStringBuilder<256> Builder;
 			Builder << TEXT("IAX Http:");
-			Builder << TEXT(" v") << GIaxHttpVersion;
 			Builder << TEXT(" - Downloaded: ") << FText::AsMemory(uint64(DownKiB) << 10).ToString();
 			Builder << TEXT(" - Rate/limit: ") << RateKiBps << TEXT(" / ") << RateLimit << TEXT(" KiB/s");
 			Builder << TEXT(" - TTFB: ") << TtfbMs << TEXT("ms");
