iax: Resolve the pool early so the description string contains more information.

#rnx

diff --git a/Engine/Source/Runtime/Experimental/IoStore/OnDemand/Private/LatencyTesting.cpp b/Engine/Source/Runtime/Experimental/IoStore/OnDemand/Private/LatencyTesting.cpp
--- a/Engine/Source/Runtime/Experimental/IoStore/OnDemand/Private/LatencyTesting.cpp
+++ b/Engine/Source/Runtime/Experimental/IoStore/OnDemand/Private/LatencyTesting.cpp
@@ -30,6 +30,7 @@ void LatencyTest(FAnsiStringView InUrl, FAnsiStringView InPath, uint32 TimeOutMs
 	}
 	
 	FConnectionPool Pool(PoolParams);
+	Pool.Resolve();
 
 	TAnsiStringBuilder<512> ConnectionDesc;
 	Pool.Describe(ConnectionDesc);
