Include  system error code when serializing FCasErrorContext.

#rnx

diff --git a/Engine/Source/Runtime/Experimental/IoStore/OnDemandCore/Public/IO/OnDemandError.h b/Engine/Source/Runtime/Experimental/IoStore/OnDemandCore/Public/IO/OnDemandError.h
--- a/Engine/Source/Runtime/Experimental/IoStore/OnDemandCore/Public/IO/OnDemandError.h
+++ b/Engine/Source/Runtime/Experimental/IoStore/OnDemandCore/Public/IO/OnDemandError.h
@@ -190,6 +190,7 @@ inline void SerializeForLog(FCbWriter& Writer, const UE::UnifiedError::IoStoreOn
 	Writer.AddInteger(ANSITEXTVIEW("ErrorCode"), uint32(Ctx.ErrorCode));
 	Writer.AddString(ANSITEXTVIEW("ErrorCodeText"), ToString(Ctx.ErrorCode));
 	Writer.AddInteger(ANSITEXTVIEW("IoErrorCode"), uint32(Ctx.IoErrorCode));
+	Writer.AddInteger(ANSITEXTVIEW("SystemErrorCode"), Ctx.SystemErrorCode);
 	Writer.AddString(ANSITEXTVIEW("IoErrorCodeText"), GetIoErrorText(Ctx.IoErrorCode));
 	Writer.AddString(ANSITEXTVIEW("ErrorMessage"), Ctx.ErrorMessage); 
 	Writer.EndObject();
