Fix crash in FHttpIoDispatcher when connection drops.

In rare cases where the HTTP connection drops while the dispatcher is handling multiple requests, the FIoBuffer passed to CompleteRequest may go out of scope, causing a client crash. Changing the Body argument to use a reference-counted copy of the original Body prevents this issue.
Updated CompleteRequest to also take a copy of the headers per request to make BugHawk happy. Headers are never used when batching multiple request together.

#rnx

diff --git a/Engine/Source/Runtime/Experimental/IoStore/OnDemand/Private/OnDemandHttpIoDispacher.cpp b/Engine/Source/Runtime/Experimental/IoStore/OnDemand/Private/OnDemandHttpIoDispacher.cpp
--- a/Engine/Source/Runtime/Experimental/IoStore/OnDemand/Private/OnDemandHttpIoDispacher.cpp
+++ b/Engine/Source/Runtime/Experimental/IoStore/OnDemand/Private/OnDemandHttpIoDispacher.cpp
@@ -776,7 +776,7 @@ private:
 	void							CancelRemainingHttpRequests();
 	FHttpRequest*					DequeueHttpRequests();
 	void							CompleteHttpRequest(FHttpRequest& Request, FMultiEndpointHttpClientResponse&& HttpResponse, FIASHostGroup& HostGroup, const FIoHttpRange& TotalRange);
-	void							CompleteRequest(FHttpRequest& Request, TArray<FAnsiString>&& Headers, FIoBuffer& Body, const FIoHttpRange& TotalRange, uint32 StatusCode, bool bCached);
+	void							CompleteRequest(FHttpRequest& Request, TArray<FAnsiString>&& Headers, FIoBuffer Body, const FIoHttpRange& TotalRange, uint32 StatusCode, bool bCached);
 	void							CompleteRequest(FHttpRequest& Request, FIoBuffer& Body, const FIoHttpRange& TotalRange, uint32 StatusCode, bool bCached)
 									{
 										CompleteRequest(Request, TArray<FAnsiString>(), Body, TotalRange, StatusCode, bCached);
@@ -1562,7 +1562,7 @@ void FHttpDispatcher::CompleteHttpRequest(FHttpRequest& Request, FMultiEndpointH
 	CompleteRequest(Request, MoveTemp(HttpResponse.Headers), HttpResponse.Body, TotalRange, HttpResponse.StatusCode, false);
 }
 
-void FHttpDispatcher::CompleteRequest(FHttpRequest& Request, TArray<FAnsiString>&& Headers, FIoBuffer& Body, const FIoHttpRange& TotalRange, const uint32 StatusCode, bool bCached)
+void FHttpDispatcher::CompleteRequest(FHttpRequest& Request, TArray<FAnsiString>&& Headers, FIoBuffer Body, const FIoHttpRange& TotalRange, const uint32 StatusCode, bool bCached)
 {
 #if UE_ALLOW_CACHE_POISONING
 	if (GIasPoisonCache && bCached)
@@ -1617,7 +1617,9 @@ void FHttpDispatcher::CompleteRequest(FHttpRequest& Request, TArray<FAnsiString>
 
 		const EIoHttpResponseFlags Flags	= bCached ? EIoHttpResponseFlags::Cached : EIoHttpResponseFlags::None;
 		FIoHttpRequestCompleted OnCompleted = MoveTemp(ToComplete->OnCompleted);
-		OnCompleted(FIoHttpResponse(ToComplete->CacheKey, FIoHttpHeaders::Create(MoveTemp(Headers)), PartialBody, CompletionStatus, RequestStatusCode, Flags));
+		TArray<FAnsiString> HeadersCopy		= Headers; // TODO: Only used for generic HTTP request but still bad 
+
+		OnCompleted(FIoHttpResponse(ToComplete->CacheKey, FIoHttpHeaders::Create(MoveTemp(HeadersCopy)), PartialBody, CompletionStatus, RequestStatusCode, Flags));
 
 		ToComplete->Flags.Add(EHttpRequestFlags::Completed);
 		ToComplete->CompletionStatus.store(CompletionStatus, std::memory_order_seq_cst);
