Collision Filter Refactor Part 4.4: Updating filtering to use the new structures.
- Added a UnionShapeFilter for particles with multiple shapes.
- Added Broad/Narrow Filter functions and tests.
- Updated the accel handle to use the union filter and methods
- Updated sim (DoCollide) and query (QueryFilterCallback) to use the NarrowFilter methods

diff --git a/Engine/Source/Programs/HeadlessChaos/Private/HeadlessChaosTestPhysicsScene.cpp b/Engine/Source/Programs/HeadlessChaos/Private/HeadlessChaosTestPhysicsScene.cpp
--- a/Engine/Source/Programs/HeadlessChaos/Private/HeadlessChaosTestPhysicsScene.cpp
+++ b/Engine/Source/Programs/HeadlessChaos/Private/HeadlessChaosTestPhysicsScene.cpp
@@ -144,12 +144,9 @@ namespace ChaosTest {
 
 	bool PrePreFilter(const Chaos::Filter::FQueryFilterData& QueryFilter, const Chaos::Filter::FShapeFilterData& ShapeFilter)
 	{
-		PRAGMA_DISABLE_INTERNAL_WARNINGS
-		const FCollisionFilterData LegacyQueryFilter = Chaos::Filter::FQueryFilterBuilder::GetLegacyQueryFilter(QueryFilter);
-		Chaos::Filter::FCombinedShapeFilterData CombinedFilter(ShapeFilter, Chaos::Filter::FInstanceData());
-		const FCollisionFilterData LegacyShapeFilter = Chaos::Filter::FShapeFilterBuilder::GetLegacyShapeQueryFilter(CombinedFilter);
-		PRAGMA_ENABLE_INTERNAL_WARNINGS
-		return PrePreQueryFilterImp(LegacyQueryFilter, LegacyShapeFilter);
+		Chaos::Filter::FShapeUnionFilterData ParticleFilter;
+		ParticleFilter.Combine(ShapeFilter);
+		return ParticleFilter.BroadFilter(QueryFilter) == Chaos::Filter::EBroadFilterResult::Reject;
 	}
 
 	GTEST_TEST(EngineInterface, CreateAndReleaseActor)
diff --git a/Engine/Source/Runtime/Engine/Private/PhysicsEngine/CollisionQueryFilterCallback.cpp b/Engine/Source/Runtime/Engine/Private/PhysicsEngine/CollisionQueryFilterCallback.cpp
--- a/Engine/Source/Runtime/Engine/Private/PhysicsEngine/CollisionQueryFilterCallback.cpp
+++ b/Engine/Source/Runtime/Engine/Private/PhysicsEngine/CollisionQueryFilterCallback.cpp
@@ -8,81 +8,46 @@
 #include "Physics/Experimental/ChaosInterfaceWrapper.h"
 #include "Chaos/GeometryParticles.h"
 
-ECollisionQueryHitType FCollisionQueryFilterCallback::CalcQueryHitType(const FQueryFilterData& InQueryFilter, const FShapeFilterData& InShapeFilter, bool bPreFilter)
+ECollisionQueryHitType ConvertToHitType(const Chaos::Filter::ENarrowFilterResult FilterResult)
 {
-	const FCollisionFilterData QueryFilter = Chaos::Filter::FQueryFilterBuilder::GetLegacyQueryFilter(InQueryFilter);
-	Chaos::Filter::FCombinedShapeFilterData CombinedFilter(InShapeFilter, Chaos::Filter::FInstanceData());
-	const FCollisionFilterData ShapeFilter = Chaos::Filter::FShapeFilterBuilder::GetLegacyShapeQueryFilter(CombinedFilter);
-
-	ECollisionQuery QueryType = (ECollisionQuery)QueryFilter.Word0;
-	FMaskFilter QuerierMaskFilter;
-	PRAGMA_DISABLE_DEPRECATION_WARNINGS
-	const ECollisionChannel QuerierChannel = GetCollisionChannelAndExtraFilter(QueryFilter.Word3, QuerierMaskFilter);
-
-	FMaskFilter ShapeMaskFilter;
-	const ECollisionChannel ShapeChannel = GetCollisionChannelAndExtraFilter(ShapeFilter.Word3, ShapeMaskFilter);
-	PRAGMA_ENABLE_DEPRECATION_WARNINGS
+	switch (FilterResult)
+	{
+		case Chaos::Filter::ENarrowFilterResult::None: return ECollisionQueryHitType::None;
+		case Chaos::Filter::ENarrowFilterResult::Overlap: return ECollisionQueryHitType::Touch;
+		case Chaos::Filter::ENarrowFilterResult::Block: return ECollisionQueryHitType::Block;
+		default: return ECollisionQueryHitType::None;
+	}
+}
 
-	if ((QuerierMaskFilter & ShapeMaskFilter) != 0)	//If ignore mask hit something, ignore it
+ECollisionQueryHitType FCollisionQueryFilterCallback::CalcQueryHitType(const FQueryFilterData& QueryFilter, const FShapeFilterData& ShapeFilter, bool bPreFilter)
+{
+	const ECollisionQueryHitType HitType = ConvertToHitType(QueryFilter.NarrowFilter(ShapeFilter));
+	if (HitType == ECollisionQueryHitType::None)
 	{
-		return ECollisionQueryHitType::None;
+		return HitType;
 	}
 
-	const uint32 ShapeBit = ECC_TO_BITFIELD(ShapeChannel);
-	if (QueryType == ECollisionQuery::ObjectQuery)
+	// There's some edge cases to handle...
+	const FQueryFilterData::EQueryType QueryType = QueryFilter.GetQueryType();
+	if (QueryType == FQueryFilterData::EQueryType::ObjectType)
 	{
-		const int32 MultiTrace = (int32)QuerierChannel;
-		// do I belong to one of objects of interest?
-		if (ShapeBit & QueryFilter.Word1)
+		// The regular object type filter logic just chooses None/Block.
+		// For pre-filtering we want all results, so convert block to touch in the multi case.
+		if (bPreFilter && QueryFilter.IsMultiQuery())
 		{
-			if (bPreFilter)	//In the case of an object query we actually want to return all object types (or first in single case). So in PreFilter we have to trick physx by not blocking in the multi case, and blocking in the single case.
-			{
-
-				return MultiTrace ? ECollisionQueryHitType::Touch: ECollisionQueryHitType::Block;
-			}
-			else
-			{
-				return ECollisionQueryHitType::Block;	//In the case where an object query is being resolved for the user we just return a block because object query doesn't have the concept of overlap at all and block seems more natural
-			}
+			return ECollisionQueryHitType::Touch;
 		}
 	}
 	else
 	{
-		// if query channel is Touch All, then just return touch
-		if (QuerierChannel == ECC_OverlapAll_Deprecated)
+		// If query channel is Touch All, then just return touch.
+		// Will have to figure out this special case when channel index is converted to registered channel masks...
+		if (QueryFilter.GetCollisionChannelIndex() == ECC_OverlapAll_Deprecated)
 		{
 			return ECollisionQueryHitType::Touch;
 		}
-		// @todo delete once we fix up object/trace APIs to work separate
-
-		uint32 const QuerierBit = ECC_TO_BITFIELD(QuerierChannel);
-		ECollisionQueryHitType QuerierHitType = ECollisionQueryHitType::None;
-		ECollisionQueryHitType ShapeHitType = ECollisionQueryHitType::None;
-
-		// check if Querier wants a hit
-		if ((QuerierBit & ShapeFilter.Word1) != 0)
-		{
-			QuerierHitType = ECollisionQueryHitType::Block;
-		}
-		else if ((QuerierBit & ShapeFilter.Word2) != 0)
-		{
-			QuerierHitType = ECollisionQueryHitType::Touch;
-		}
-
-		if ((ShapeBit & QueryFilter.Word1) != 0)
-		{
-			ShapeHitType = ECollisionQueryHitType::Block;
-		}
-		else if ((ShapeBit & QueryFilter.Word2) != 0)
-		{
-			ShapeHitType = ECollisionQueryHitType::Touch;
-		}
-
-		// return minimum agreed-upon interaction
-		return FMath::Min(QuerierHitType, ShapeHitType);
 	}
-
-	return ECollisionQueryHitType::None;
+	return HitType;
 }
 
 ECollisionQueryHitType FCollisionQueryFilterCallback::PostFilter(const FQueryFilterData& FilterData, const ChaosInterface::FQueryHit& Hit)
diff --git a/Engine/Source/Runtime/Experimental/Chaos/Private/Chaos/CollisionFilterData.cpp b/Engine/Source/Runtime/Experimental/Chaos/Private/Chaos/CollisionFilterData.cpp
--- a/Engine/Source/Runtime/Experimental/Chaos/Private/Chaos/CollisionFilterData.cpp
+++ b/Engine/Source/Runtime/Experimental/Chaos/Private/Chaos/CollisionFilterData.cpp
@@ -141,6 +141,11 @@ namespace Chaos::Filter
 		return QueryBlockChannels != 0 || QueryOverlapChannels != 0 || SimBlockChannels != 0 || Word3 != 0 || SimWord3 != 0;
 	}
 
+	bool FShapeFilterData::IsSimValid() const
+	{
+		return SimBlockChannels != 0 || SimWord3 != 0;
+	}
+
 	EFilterFlags FShapeFilterData::GetFlags() const
 	{
 		return (EFilterFlags)Chaos::Filter::GetFlags(Word3);
@@ -203,6 +208,17 @@ namespace Chaos::Filter
 		return SimBlockChannels;
 	}
 
+	ENarrowFilterResult FShapeFilterData::NarrowFilter(const FShapeFilterData& OtherShape) const
+	{
+		const uint64 ChannelMask0 = GetCollisionChannelMask();
+		const uint64 ChannelMask1 = OtherShape.GetCollisionChannelMask();
+		const uint64 BlockMask0 = GetSimBlockChannels();
+		const uint64 BlockMask1 = OtherShape.GetSimBlockChannels();
+
+		const bool bBlocking = (ChannelMask0 & BlockMask1) && (ChannelMask1 & BlockMask0);
+		return bBlocking ? ENarrowFilterResult::Block : ENarrowFilterResult::None;
+	}
+
 	FString FShapeFilterData::ToString() const
 	{
 		return FString::Format(TEXT("SimBlock(%x) QueryBlock(%x) QueryOverlap(%x) CollisionChannel(%x) MaskFilter(%x) Flags(%x)"), { GetSimBlockChannels(), GetQueryBlockChannels(), GetQueryOverlapChannels(), GetCollisionChannelIndex(), GetMaskFilter(), (uint8)GetFlags() });
@@ -213,6 +229,44 @@ namespace Chaos::Filter
 		return Word3;
 	}
 
+	void FShapeUnionFilterData::Combine(const FShapeFilterData& ShapeFilter, bool bQueryEnabled, bool bSimEnabled)
+	{
+		if (bQueryEnabled)
+		{
+			QueryBlockChannels |= ShapeFilter.QueryBlockChannels;
+			QueryOverlapChannels |= ShapeFilter.QueryOverlapChannels;
+			QueryChannelMask |= (uint32)ShapeFilter.GetCollisionChannelMask();
+		}
+		if (bSimEnabled)
+		{
+			SimBlockChannels |= ShapeFilter.SimBlockChannels;
+			SimChannelMask |= (uint32)ShapeFilter.GetCollisionChannelMask();
+		}
+	}
+
+	EBroadFilterResult FShapeUnionFilterData::BroadFilter(const FShapeUnionFilterData& OtherFilter) const
+	{
+		const uint64 ChannelMask0 = SimChannelMask;
+		const uint64 ChannelMask1 = OtherFilter.SimChannelMask;
+		const uint64 BlockMask0 = SimBlockChannels;
+		const uint64 BlockMask1 = OtherFilter.SimBlockChannels;
+
+		const bool bCanCollide = (ChannelMask0 & BlockMask1) && (ChannelMask1 & BlockMask0);
+		return bCanCollide ? EBroadFilterResult::Accept : EBroadFilterResult::Reject;
+	}
+
+	EBroadFilterResult FShapeUnionFilterData::BroadFilter(const FQueryFilterData& QueryFilter) const
+	{
+		if (QueryFilter.GetQueryType() == FQueryFilterData::EQueryType::Channel)
+		{
+			const uint64 QueryFilterChannelMask = QueryFilter.GetCollisionChannelMask();
+			const uint64 CombinedMask = QueryBlockChannels | QueryOverlapChannels;
+			const bool bIsHit = QueryFilterChannelMask & CombinedMask;
+			return bIsHit ? EBroadFilterResult::Accept : EBroadFilterResult::Reject;
+		}
+		return EBroadFilterResult::Accept;
+	}
+
 	bool FQueryFilterData::IsValid() const
 	{
 		return Word0 != 0 || Word1 != 0 || Word2 != 0 || Word3 != 0;
@@ -284,6 +338,32 @@ namespace Chaos::Filter
 		return MultiTrace != 0;
 	}
 
+	ENarrowFilterResult FQueryFilterData::NarrowFilter(const FShapeFilterData& ShapeFilter) const
+	{
+		return NarrowFilter(ShapeFilter, false);
+	}
+
+	ENarrowFilterResult FQueryFilterData::NarrowFilter(const FShapeFilterData& ShapeFilter, bool bSimFilter) const
+	{
+		const FMaskFilter QuerierMaskFilter = GetIgnoreMask();
+		const FMaskFilter ShapeMaskFilter = ShapeFilter.GetMaskFilter();
+
+		// Check if the ignore masks say to skip
+		if ((QuerierMaskFilter & ShapeMaskFilter) != 0)
+		{
+			return ENarrowFilterResult::None;
+		}
+
+		if (GetQueryType() == FQueryFilterData::EQueryType::ObjectType)
+		{
+			return ObjectTypeNarrowFilter(ShapeFilter);
+		}
+		else
+		{
+			return ChannelTypeNarrowFilter(ShapeFilter, bSimFilter);
+		}
+	}
+
 	CHAOS_API FString FQueryFilterData::ToString() const
 	{
 		if (GetQueryType() == EQueryType::Channel)
@@ -301,6 +381,55 @@ namespace Chaos::Filter
 		Word3 = ((uint32)IgnoreMask << MaskFilterBitsOffset) | ((uint32)ChannelIndex << CollisionChannelBitsOffset) | (uint32)FilterFlags;
 	}
 
+	ENarrowFilterResult FQueryFilterData::ChannelTypeNarrowFilter(const FShapeFilterData& ShapeFilter, bool bSimBlock) const
+	{
+		const uint64 QueryChannelMask = GetCollisionChannelMask();
+		const uint64 QueryBlockChannels = GetBlockChannels();
+		const uint64 QueryOverlapChannels = GetOverlapChannels();
+
+		const uint64 ShapeChannelMask = ShapeFilter.GetCollisionChannelMask();
+		const uint64 ShapeBlockChannels = bSimBlock ? ShapeFilter.GetSimBlockChannels() : ShapeFilter.GetQueryBlockChannels();
+		const uint64 ShapeOverlapChannels = ShapeFilter.GetQueryOverlapChannels();
+
+		ENarrowFilterResult QuerierHitType = ENarrowFilterResult::None;
+		ENarrowFilterResult ShapeHitType = ENarrowFilterResult::None;
+
+		// Check query type vs. shape masks
+		if ((QueryChannelMask & ShapeBlockChannels) != 0)
+		{
+			QuerierHitType = ENarrowFilterResult::Block;
+		}
+		else if ((QueryChannelMask & ShapeOverlapChannels) != 0)
+		{
+			QuerierHitType = ENarrowFilterResult::Overlap;
+		}
+
+		// Check shape type vs. query masks
+		if ((ShapeChannelMask & QueryBlockChannels) != 0)
+		{
+			ShapeHitType = ENarrowFilterResult::Block;
+		}
+		else if ((ShapeChannelMask & QueryOverlapChannels) != 0)
+		{
+			ShapeHitType = ENarrowFilterResult::Overlap;
+		}
+
+		// Combine the results by taking the min
+		return FMath::Min(QuerierHitType, ShapeHitType);
+	}
+
+	ENarrowFilterResult FQueryFilterData::ObjectTypeNarrowFilter(const FShapeFilterData& ShapeFilter) const
+	{
+		const uint64 QueryObjectTypesMask = GetObjectTypesToQueryMask();
+		const uint64 ShapeChannelMask = ShapeFilter.GetCollisionChannelMask();
+		// do I belong to one of objects of interest?
+		if (ShapeChannelMask & QueryObjectTypesMask)
+		{
+			return ENarrowFilterResult::Block;
+		}
+		return ENarrowFilterResult::None;
+	}
+
 	FCombinedShapeFilterData::FCombinedShapeFilterData(const FShapeFilterData& InShapeFilter, const FInstanceData& InInstanceData)
 		: ShapeFilterData(InShapeFilter)
 		, InstanceData(InInstanceData)
diff --git a/Engine/Source/Runtime/Experimental/Chaos/Public/Chaos/Collision/CollisionFilter.h b/Engine/Source/Runtime/Experimental/Chaos/Public/Chaos/Collision/CollisionFilter.h
--- a/Engine/Source/Runtime/Experimental/Chaos/Public/Chaos/Collision/CollisionFilter.h
+++ b/Engine/Source/Runtime/Experimental/Chaos/Public/Chaos/Collision/CollisionFilter.h
@@ -98,6 +98,7 @@ namespace Chaos
 	/**
 	* A filter set to all zeroes means we do not filter
 	*/
+	UE_DEPRECATED(5.8, "Use FShapeFilterData.IsValid instead")
 	inline bool IsFilterValid(const FCollisionFilterData& Filter)
 	{
 		return Filter.Word0 || Filter.Word1 || Filter.Word2 || Filter.Word3;
@@ -108,9 +109,7 @@ namespace Chaos
 	*/
 	inline bool FilterHasSimEnabled(const FPerShapeData* Shape)
 	{
-		PRAGMA_DISABLE_DEPRECATION_WARNINGS
-		return (!Shape || (Shape->GetSimEnabled() && IsFilterValid(Shape->GetSimData())));
-		PRAGMA_ENABLE_DEPRECATION_WARNINGS
+		return (!Shape || (Shape->GetSimEnabled() && Shape->GetShapeFilterData().IsSimValid()));
 	}
 
 	inline bool DoCollide(EImplicitObjectType ImplicitType, const FPerShapeData* Shape)
@@ -156,18 +155,10 @@ namespace Chaos
 		//
 		if (Shape0 && Shape1)
 		{
-			PRAGMA_DISABLE_DEPRECATION_WARNINGS
-			if (IsFilterValid(Shape0->GetSimData()) && IsFilterValid(Shape1->GetSimData()))
-			{
-				FMaskFilter Filter0Mask, Filter1Mask;
-				const uint32 Filter0Channel = GetChaosCollisionChannelAndExtraFilter(Shape0->GetSimData().Word3, Filter0Mask);
-				const uint32 Filter1Channel = GetChaosCollisionChannelAndExtraFilter(Shape1->GetSimData().Word3, Filter1Mask);
-
-				const uint32 Filter1Bit = 1 << (Filter1Channel); // SIMDATA_TO_BITFIELD
-				uint32 const Filter0Bit = 1 << (Filter0Channel); // SIMDATA_TO_BITFIELD
-				return (Filter0Bit & Shape1->GetSimData().Word1) && (Filter1Bit & Shape0->GetSimData().Word1);
-			}
-			PRAGMA_ENABLE_DEPRECATION_WARNINGS
+			const Filter::FShapeFilterData& ShapeFilter0 = Shape0->GetShapeFilterData();
+			const Filter::FShapeFilterData& ShapeFilter1 = Shape1->GetShapeFilterData();
+			const Filter::ENarrowFilterResult HitType = ShapeFilter0.NarrowFilter(ShapeFilter1);
+			return HitType != Filter::ENarrowFilterResult::None;
 		}
 
 		return true;
diff --git a/Engine/Source/Runtime/Experimental/Chaos/Public/Chaos/Collision/SpatialAccelerationBroadPhase.h b/Engine/Source/Runtime/Experimental/Chaos/Public/Chaos/Collision/SpatialAccelerationBroadPhase.h
--- a/Engine/Source/Runtime/Experimental/Chaos/Public/Chaos/Collision/SpatialAccelerationBroadPhase.h
+++ b/Engine/Source/Runtime/Experimental/Chaos/Public/Chaos/Collision/SpatialAccelerationBroadPhase.h
@@ -243,11 +243,6 @@ namespace Chaos
 				, ParticleUniqueIdx(ParticleHandle ? ParticleHandle->UniqueIdx() : FUniqueIdx(0))
 				, AccelerationHandle(ParticleHandle)
 			{
-				// Should never be null. Also suppresses a warning about possible null ptr.
-				check(ParticleHandle != nullptr);
-				PRAGMA_DISABLE_INTERNAL_WARNINGS
-				FAccelerationStructureHandle::ComputeParticleSimFilterDataFromShapes(*ParticleHandle, SimFilterData);
-				PRAGMA_ENABLE_INTERNAL_WARNINGS
 			}
 
 			UE_DEPRECATED(5.7, "Use the constructor that does not take the sim filter.")
@@ -282,7 +277,7 @@ namespace Chaos
 			const void* GetQueryData() const { return nullptr; }
 
 			UE_DEPRECATED(5.8, "This has been deprecated in favor of PrePreFilter")
-			const void* GetSimData() const { return &SimFilterData; }
+			const void* GetSimData() const { return nullptr; }
 
 			bool ShouldIgnore(const TSpatialVisitorData<FAccelerationStructureHandle>& Instance) const
 			{
@@ -305,15 +300,11 @@ namespace Chaos
 				{
 					return true;
 				}
-				PRAGMA_DISABLE_INTERNAL_WARNINGS
-				return Payload.PrePreSimFilter(SimFilterData);
-				PRAGMA_ENABLE_INTERNAL_WARNINGS
+				return Payload.PrePreFilter(AccelerationHandle);
 			}
 
 		private:
 			Private::FBroadPhaseContext& Context;
-			// @todo(chaos): cache this on the particle?
-			FCollisionFilterData SimFilterData;
 			FUniqueIdx ParticleUniqueIdx; // unique id of the particle visiting, used to skip self intersection as early as possible
 
 			/** Handle to be stored to retrieve the payload on which we are querying the acceleration structure*/
diff --git a/Engine/Source/Runtime/Experimental/Chaos/Public/Chaos/CollisionFilterData.h b/Engine/Source/Runtime/Experimental/Chaos/Public/Chaos/CollisionFilterData.h
--- a/Engine/Source/Runtime/Experimental/Chaos/Public/Chaos/CollisionFilterData.h
+++ b/Engine/Source/Runtime/Experimental/Chaos/Public/Chaos/CollisionFilterData.h
@@ -77,6 +77,20 @@ namespace Chaos::Filter
 	struct FCombinedShapeFilterData;
 	struct FQueryFilterData;
 
+	enum class ENarrowFilterResult
+	{
+		None = 0,
+		Overlap = 1,
+		Block = 2,
+	};
+	enum class EBroadFilterResult
+	{
+		// The object should be skipped from further processing.
+		Reject,
+		// The object should go through the rest of the detection pipeline.
+		Accept,
+	};
+
 	struct FInstanceData
 	{
 		FInstanceData() = default;
@@ -115,6 +129,7 @@ namespace Chaos::Filter
 		bool operator!=(const FShapeFilterData&) const = default;
 
 		CHAOS_API bool IsValid() const;
+		CHAOS_API bool IsSimValid() const;
 
 		CHAOS_API EFilterFlags GetFlags() const;
 		CHAOS_API void SetFlags(EFilterFlags InFlags);
@@ -131,6 +146,9 @@ namespace Chaos::Filter
 		CHAOS_API uint64 GetQueryOverlapChannels() const;
 		CHAOS_API uint64 GetSimBlockChannels() const;
 
+		// Filter API
+		CHAOS_API ENarrowFilterResult NarrowFilter(const FShapeFilterData& OtherShape) const;
+
 		CHAOS_API FString ToString() const;
 
 		UE_INTERNAL uint32 GetLegacyWord3() const;
@@ -139,6 +157,7 @@ namespace Chaos::Filter
 		friend class FShapeFilterBuilder;
 		friend FQueryFilterData;
 		friend FCombinedShapeFilterData;
+		friend struct FShapeUnionFilterData;
 
 		uint32 QueryBlockChannels = 0;
 		uint32 QueryOverlapChannels = 0;
@@ -148,6 +167,28 @@ namespace Chaos::Filter
 		uint32 SimWord3 = 0;
 	};
 
+	struct FShapeUnionFilterData
+	{
+		FShapeUnionFilterData() = default;
+
+		bool operator==(const FShapeUnionFilterData&) const = default;
+		bool operator!=(const FShapeUnionFilterData&) const = default;
+
+		CHAOS_API void Combine(const FShapeFilterData& ShapeFilter, bool bQueryEnabled = true, bool bSimEnabled = true);
+
+		// Filter API
+		CHAOS_API EBroadFilterResult BroadFilter(const FShapeUnionFilterData& OtherFilter) const;
+		CHAOS_API EBroadFilterResult BroadFilter(const FQueryFilterData& QueryFilter) const;
+
+	private:
+		friend class FShapeFilterBuilder;
+		uint32 QueryBlockChannels = 0;
+		uint32 QueryOverlapChannels = 0;
+		uint32 QueryChannelMask = 0;
+		uint32 SimBlockChannels = 0;
+		uint32 SimChannelMask = 0;
+	};
+
 	struct FQueryFilterData
 	{
 		enum EQueryType
@@ -182,6 +223,10 @@ namespace Chaos::Filter
 		CHAOS_API uint64 GetObjectTypesToQueryMask() const;
 		CHAOS_API bool IsMultiQuery() const;
 
+		// Filter API
+		CHAOS_API ENarrowFilterResult NarrowFilter(const FShapeFilterData& ShapeFilter) const;
+		UE_INTERNAL CHAOS_API ENarrowFilterResult NarrowFilter(const FShapeFilterData& ShapeFilter, bool bSimFilter) const;
+
 		CHAOS_API FString ToString() const;
 
 	private:
@@ -190,6 +235,8 @@ namespace Chaos::Filter
 		friend class FQueryTraceFilterBuilder;
 
 		void SetMaskChannelAndFlags(const uint8 IgnoreMask, const uint8 ChannelIndex, const EFilterFlags FilterFlags);
+		ENarrowFilterResult ChannelTypeNarrowFilter(const FShapeFilterData& ShapeFilter, bool bSimBlock) const;
+		ENarrowFilterResult ObjectTypeNarrowFilter(const FShapeFilterData& ShapeFilter) const;
 
 		uint32 Word0 = 0;
 		uint32 Word1 = 0;
diff --git a/Engine/Source/Runtime/Experimental/Chaos/Public/Chaos/ParticleHandle.h b/Engine/Source/Runtime/Experimental/Chaos/Public/Chaos/ParticleHandle.h
--- a/Engine/Source/Runtime/Experimental/Chaos/Public/Chaos/ParticleHandle.h
+++ b/Engine/Source/Runtime/Experimental/Chaos/Public/Chaos/ParticleHandle.h
@@ -161,6 +161,7 @@ bool GeometryParticleSleeping(const FConcrete& Concrete)
 
 //Used to filter out at the acceleration structure layer using Query data
 //Returns true when there is no way a later PreFilter will succeed. Avoid virtuals etc..
+UE_DEPRECATED(5.8, "Use FShapeUnionFilterData::BroadFilter(const FQueryFilterData&) instead")
 FORCEINLINE_DEBUGGABLE bool PrePreQueryFilterImp(const FCollisionFilterData& QueryFilterData, const FCollisionFilterData& UnionFilterData)
 {
 	//HACK: need to replace all these hard-coded values with proper enums, bad modules are not setup for it right now
@@ -183,6 +184,7 @@ FORCEINLINE_DEBUGGABLE bool PrePreQueryFilterImp(const FCollisionFilterData& Que
 	return false;
 }
 
+UE_DEPRECATED(5.8, "Use the member functions on FShapeFilterData instead")
 FORCEINLINE_DEBUGGABLE uint32 GetChaosCollisionChannelAndExtraFilter(uint32 Word3, uint8& OutMaskFilter)
 {
 	uint32 ChannelMask = (Word3 << ChaosNumExtraFilterBits) >> (32 - ChaosNumCollisionChannelBits);
@@ -192,15 +194,18 @@ FORCEINLINE_DEBUGGABLE uint32 GetChaosCollisionChannelAndExtraFilter(uint32 Word
 
 //Used to filter out at the acceleration structure layer using Simdata
 //Returns true when there is no way a later PreFilter will succeed. Avoid virtuals etc..
+UE_DEPRECATED(5.8, "Use FShapeUnionFilterData::BroadFilter(const FShapeFilterData&) instead")
 FORCEINLINE_DEBUGGABLE bool PrePreSimFilterImp(const FCollisionFilterData& SimFilterData, const FCollisionFilterData& OtherSimFilterData)
 {
 	//HACK: need to replace all these hard-coded values with proper enums, bad modules are not setup for it right now
 	//since we're taking the union of shapes we can only support trace channel
 	//const ECollisionChannel QuerierChannel = GetCollisionChannel(QueryFilter.Word3);
+	PRAGMA_DISABLE_DEPRECATION_WARNINGS
 	uint8  QuerierMaskFilter;
 	const uint32 QuerierChannel = GetChaosCollisionChannelAndExtraFilter(SimFilterData.Word3, QuerierMaskFilter);
 	uint8  OtherMaskFilter;
 	const uint32 OtherChannel = GetChaosCollisionChannelAndExtraFilter(OtherSimFilterData.Word3, OtherMaskFilter);
+	PRAGMA_ENABLE_DEPRECATION_WARNINGS
 
 	//uint32 const QuerierBit = ECC_TO_BITFIELD(QuerierChannel);
 	const uint32 QuerierBit = (1 << (QuerierChannel));
@@ -248,74 +253,70 @@ public:
 		return CachedUniqueIdx;
 	}
 
+	UE_DEPRECATED(5.8, "This has been deprecated in favor of PrePreFilter")
 	bool PrePreQueryFilter(const void* QueryData) const
 	{
 		if(bCanPrePreFilter)
 		{
 			if (const FCollisionFilterData* QueryFilterData = static_cast<const FCollisionFilterData*>(QueryData))
 			{
-				return PrePreQueryFilterImp(*QueryFilterData, UnionQueryFilterData);
+				PRAGMA_DISABLE_INTERNAL_WARNINGS
+				const Chaos::Filter::FQueryFilterData QuerierFilter = Chaos::Filter::FQueryFilterBuilder::BuildFromLegacyQueryFilter(*QueryFilterData);
+				PRAGMA_ENABLE_INTERNAL_WARNINGS
+				return UnionFilter.BroadFilter(QuerierFilter) == Chaos::Filter::EBroadFilterResult::Reject;
 			}
 		}
 		
 		return false;
 	}
 
-
+	UE_DEPRECATED(5.8, "This has been deprecated in favor of PrePreFilter")
 	bool PrePreSimFilter(const void* SimData) const
 	{
 		if (bCanPrePreFilter)
 		{
 			if (const FCollisionFilterData* SimFilterData = static_cast<const FCollisionFilterData*>(SimData))
 			{
-				return PrePreSimFilterImp(*SimFilterData, UnionSimFilterData);
+				PRAGMA_DISABLE_INTERNAL_WARNINGS
+				const Chaos::Filter::FShapeFilterData ShapeFilter = Chaos::Filter::FShapeFilterBuilder::BuildFromLegacySimFilter(*SimFilterData);
+				PRAGMA_ENABLE_INTERNAL_WARNINGS
+				Chaos::Filter::FShapeUnionFilterData SimUnionFilter;
+				SimUnionFilter.Combine(ShapeFilter);
+				return UnionFilter.BroadFilter(SimUnionFilter) == Chaos::Filter::EBroadFilterResult::Reject;
 			}
 		}
 
 		return false;
 	}
 
-	UE_INTERNAL bool PrePreSimFilter(const FCollisionFilterData& QuerierFilterData) const
+	bool PrePreFilter(const Chaos::Filter::FQueryFilterData& QueryFilterData) const
 	{
 		if (bCanPrePreFilter)
 		{
-			const uint32 QuerierBlockMask = QuerierFilterData.Word1;
-			const uint32 QuerierChannelMask = QuerierFilterData.Word2;
-			const uint32 SimBlockMask = UnionSimFilterData.Word1;
-			const uint32 SimChannelMask = UnionSimFilterData.Word2;
-
-			// check if they can collide ( same logic as DoCollide in CollisionResolution.cpp )
-			const bool CanCollide = (QuerierChannelMask & SimBlockMask) && (SimChannelMask & QuerierBlockMask);
-			return !CanCollide;
+			return UnionFilter.BroadFilter(QueryFilterData) == Chaos::Filter::EBroadFilterResult::Reject;
 		}
 
 		return false;
 	}
 
-	bool PrePreFilter(const Chaos::Filter::FQueryFilterData& QueryFilterData) const
+	bool PrePreFilter(const Chaos::Filter::FShapeUnionFilterData& QuerierFilterData) const
 	{
 		if (bCanPrePreFilter)
 		{
-			PRAGMA_DISABLE_INTERNAL_WARNINGS
-			const FCollisionFilterData LegacyQueryFilter = Chaos::Filter::FQueryFilterBuilder::GetLegacyQueryFilter(QueryFilterData);
-			PRAGMA_ENABLE_INTERNAL_WARNINGS
-			return PrePreQueryFilterImp(LegacyQueryFilter, UnionQueryFilterData);
+			return UnionFilter.BroadFilter(QuerierFilterData) == Chaos::Filter::EBroadFilterResult::Reject;
 		}
-		
+
 		return false;
 	}
 
-	void UpdateFrom(const FAccelerationStructureHandle& InOther)
+	bool PrePreFilter(const FAccelerationStructureHandle& OtherHandle) const
 	{
-		UnionQueryFilterData.Word0 = InOther.UnionQueryFilterData.Word0;
-		UnionQueryFilterData.Word1 = InOther.UnionQueryFilterData.Word1;
-		UnionQueryFilterData.Word2 = InOther.UnionQueryFilterData.Word2;
-		UnionQueryFilterData.Word3 = InOther.UnionQueryFilterData.Word3;
+		return PrePreFilter(OtherHandle.UnionFilter);
+	}
 
-		UnionSimFilterData.Word0 = InOther.UnionSimFilterData.Word0;
-		UnionSimFilterData.Word1 = InOther.UnionSimFilterData.Word1;
-		UnionSimFilterData.Word2 = InOther.UnionSimFilterData.Word2;
-		UnionSimFilterData.Word3 = InOther.UnionSimFilterData.Word3;
+	void UpdateFrom(const FAccelerationStructureHandle& InOther)
+	{
+		UnionFilter = InOther.UnionFilter;
 	}
 
 public:
@@ -324,25 +325,20 @@ public:
 	* NOTE: may not initialize OutQueryFilterData, so it must be cleared before calling if necessary (e.g., when changing overlapping channels)
 	**/
 	template <typename TParticle>
+	UE_DEPRECATED(5.8, "Use ComputeFilterDataFromShapes instead")
 	static void ComputeParticleQueryFilterDataFromShapes(const TParticle& Particle, FCollisionFilterData& OutQueryFilterData);
 
-	/**
-	* compute the aggregated sim collision filter from all associated shapes
-	* NOTE: may not initialize OutSimFilterData, so it must be cleared before calling if necessary (e.g., when changing blocking channels).
-	* NOTE: This sim filter is slightly different than the individual filters as word2 is stolen to be the decompressed channel bitmask instead of the component ID.
-	**/
+	// Computes union filter data from the associated shapes.
 	template <typename TParticle>
-	UE_INTERNAL static void ComputeParticleSimFilterDataFromShapes(const TParticle& Particle, FCollisionFilterData& OutSimFilterData);
+	static void ComputeFilterDataFromShapes(const TParticle& Particle, Chaos::Filter::FShapeUnionFilterData& OutFilter);
 
 private:
 	FGeometryParticle* ExternalGeometryParticle;
 	FGeometryParticleHandle* GeometryParticleHandle;
 
 	FUniqueIdx CachedUniqueIdx;
-	FCollisionFilterData UnionQueryFilterData;
-	// Note: The union data is a binary or of all of the shapes in the particle. Some data in the normal sim filter is not valid if or'd, such as the channel index in word3, the component id in word2, and the body index in word0.
-	// In order to make pre-pre filtering work, the sim filter's word2 is replaced with the unpacked collision channel mask which can be or'd.
-	FCollisionFilterData UnionSimFilterData;
+	Chaos::Filter::FShapeUnionFilterData UnionFilter;
+
 	bool bCanPrePreFilter;
 
 	template <typename TParticle>
@@ -4008,40 +4004,22 @@ template <typename TParticle>
 }
 
 template <typename TParticle>
-/* static */ void FAccelerationStructureHandle::ComputeParticleSimFilterDataFromShapes(const TParticle& Particle, FCollisionFilterData& OutSimFilterData)
+/* static */ void FAccelerationStructureHandle::ComputeFilterDataFromShapes(const TParticle& Particle, Chaos::Filter::FShapeUnionFilterData& OutFilter)
 {
-	const auto& Shapes = Particle.ShapesArray();
-	OutSimFilterData.Word2 = 0;
+	const FShapesArray& Shapes = Particle.ShapesArray();
+	OutFilter = Chaos::Filter::FShapeUnionFilterData();
 	for (const auto& Shape : Shapes)
 	{
-		// In the simulation we can ignore query-only shapes, but need to know about sim- and probe-enabled shapes
-		if (Shape->GetSimEnabled() || Shape->GetIsProbe())
-		{
-			PRAGMA_DISABLE_DEPRECATION_WARNINGS
-			const FCollisionFilterData& ShapeSimData = Shape->GetSimData();
-			PRAGMA_ENABLE_DEPRECATION_WARNINGS
-
-			// Word3 contains the channel index, this cannot be or'd together. Instead, expand the index to the bitmask and or that. 
-			// No one currently uses word2 (the component ID) which would also be invalid to or, so re-use that memory.
-			uint8 UnusedMaskFilter;
-			const uint32 ShapeChannel = GetChaosCollisionChannelAndExtraFilter(ShapeSimData.Word3, UnusedMaskFilter);
-			const uint32 ShapeChannelMask = 1 << ShapeChannel;
-
-			OutSimFilterData.Word0 |= ShapeSimData.Word0;
-			OutSimFilterData.Word1 |= ShapeSimData.Word1;
-			OutSimFilterData.Word2 |= ShapeChannelMask;
-			OutSimFilterData.Word3 |= ShapeSimData.Word3;
-		}
+		const bool bSimEnabled = Shape->GetSimEnabled() || Shape->GetIsProbe();
+		const bool bQueryEnabled = Shape->GetQueryEnabled();
+		OutFilter.Combine(Shape->GetShapeFilterData(), bQueryEnabled, bSimEnabled);
 	}
 }
 
 template <typename TParticle>
 FORCEINLINE_DEBUGGABLE void FAccelerationStructureHandle::UpdatePrePreFilter(const TParticle& Particle)
 {
-	ComputeParticleQueryFilterDataFromShapes(Particle, UnionQueryFilterData);
-	PRAGMA_DISABLE_INTERNAL_WARNINGS
-	ComputeParticleSimFilterDataFromShapes(Particle, UnionSimFilterData);
-	PRAGMA_ENABLE_INTERNAL_WARNINGS
+	ComputeFilterDataFromShapes(Particle, UnionFilter);
 	bCanPrePreFilter = true;
 }
 
