[Chaos] Adding ensures to try to track down an invalid aabb on a particle.

diff --git a/Engine/Source/Runtime/Experimental/Chaos/Public/Chaos/AABBTree.h b/Engine/Source/Runtime/Experimental/Chaos/Public/Chaos/AABBTree.h
--- a/Engine/Source/Runtime/Experimental/Chaos/Public/Chaos/AABBTree.h
+++ b/Engine/Source/Runtime/Experimental/Chaos/Public/Chaos/AABBTree.h
@@ -2365,28 +2365,7 @@ public:
 	// Returns true if bounds appear valid. Returns false if extremely large values, contains NaN, or is empty.
 	FORCEINLINE_DEBUGGABLE bool ValidateBounds(const TAABB<T, 3>& Bounds)
 	{
-		const TVec3<T>& Min = Bounds.Min();
-		const TVec3<T>& Max = Bounds.Max();
-
-		for (int32 i = 0; i < 3; ++i)
-		{
-			const T& MinComponent = Min[i];
-			const T& MaxComponent = Max[i];
-
-			// Are we an empty aabb?
-			if (MinComponent > MaxComponent)
-			{
-				return false;
-			}
-
-			// Are we NaN/Inf?
-			if (!FMath::IsFinite(MinComponent) || !FMath::IsFinite(MaxComponent))
-			{
-				return false;
-			}
-		}
-
-		return true;
+		return Bounds.IsValid();
 	}
 
 	virtual void Serialize(FChaosArchive& Ar) override
diff --git a/Engine/Source/Runtime/Experimental/Chaos/Public/Chaos/GeometryParticles.h b/Engine/Source/Runtime/Experimental/Chaos/Public/Chaos/GeometryParticles.h
--- a/Engine/Source/Runtime/Experimental/Chaos/Public/Chaos/GeometryParticles.h
+++ b/Engine/Source/Runtime/Experimental/Chaos/Public/Chaos/GeometryParticles.h
@@ -458,6 +458,8 @@ namespace Chaos
 			}
 
 			MWorldSpaceInflatedBounds[Index] = TAABB<T, d>(WorldBounds).ThickenSymmetrically(BoundsExpansion);
+			ensureMsgf(MWorldSpaceInflatedBounds[Index].IsValid(), TEXT("Invalid world space inflated bounds. %s, BoundsExpansion(%s)"),
+				*MWorldSpaceInflatedBounds[Index].ToString(), *BoundsExpansion.ToString());
 		}
 
 		void UpdateWorldSpaceStateSwept(const int32 Index, const FRigidTransform3& WorldTransform, const FVec3& BoundsExpansion, const FVec3& DeltaX)
@@ -467,6 +469,8 @@ namespace Chaos
 
 			// Apply the swept bounds delta
 			MWorldSpaceInflatedBounds[Index].GrowByVector(DeltaX);
+			ensureMsgf(MWorldSpaceInflatedBounds[Index].IsValid(), TEXT("Invalid world space inflated bounds. %s, BoundsExpansion(%s), DeltaX(%s)"),
+				*MWorldSpaceInflatedBounds[Index].ToString(), *BoundsExpansion.ToString(), *DeltaX.ToString());
 		}
 
 		typedef FGeometryParticleHandle THandleType;
diff --git a/Engine/Source/Runtime/Experimental/ChaosCore/Public/Chaos/AABB.h b/Engine/Source/Runtime/Experimental/ChaosCore/Public/Chaos/AABB.h
--- a/Engine/Source/Runtime/Experimental/ChaosCore/Public/Chaos/AABB.h
+++ b/Engine/Source/Runtime/Experimental/ChaosCore/Public/Chaos/AABB.h
@@ -646,6 +646,12 @@ namespace Chaos
 			return (MMin.X == TNumericLimits<T>::Max());
 		}
 
+		// AABB is valid if it doesn't contain nan and it's non empty (min <= max)
+		FORCEINLINE_DEBUGGABLE bool IsValid() const
+		{
+			return MMin <= MMax && !MMin.ContainsNaN() && !MMax.ContainsNaN();
+		}
+
 		// Shrink this AABB to being empty
 		FORCEINLINE void Clear()
 		{
