Add ways to update GT particles when set world transform is called on the proxy
* behind a cvar that is disabled by default

diff --git a/Engine/Source/Runtime/Experimental/Chaos/Private/PhysicsProxy/GeometryCollectionPhysicsProxy.cpp b/Engine/Source/Runtime/Experimental/Chaos/Private/PhysicsProxy/GeometryCollectionPhysicsProxy.cpp
--- a/Engine/Source/Runtime/Experimental/Chaos/Private/PhysicsProxy/GeometryCollectionPhysicsProxy.cpp
+++ b/Engine/Source/Runtime/Experimental/Chaos/Private/PhysicsProxy/GeometryCollectionPhysicsProxy.cpp
@@ -151,6 +151,12 @@ FAutoConsoleVariableRef CVarEmptyUniqueIndicesOnGT(
 	bEmptyUniqueIndicesOnGT,
 	TEXT("If enabled, immediately free particle unique indices after particle initialisation"));
 
+bool bUpdateGTParticlesOnWorldTransformUpdate = false;
+FAutoConsoleVariableRef CVarUpdateGTParticlesOnWorldTransformUpdate(
+	TEXT("p.GeometryCollection.UpdateGTParticlesOnWorldTransformUpdate"),
+	bUpdateGTParticlesOnWorldTransformUpdate,
+	TEXT("If enabled, SetWorldTransform_External will immediately update the GT particle to reflect the new transform instead of waiting for a full GT -> PT -> GT sync to do so"));
+
 enum EOverrideGCCollisionSetupForTraces
 {
 	GCCSFT_Property   = -1,  // Default: do what property says
@@ -3369,6 +3375,31 @@ void FGeometryCollectionPhysicsProxy::SetWorldTransform_External(const FTransfor
 		PreviousWorldTransform_External = WorldTransform_External;
 		WorldTransform_External = WorldTransform;
 
+		// update GT particle if needed - experimental
+		if (bUpdateGTParticlesOnWorldTransformUpdate)
+		{
+			const TManagedArray<bool>* AnimationsActive = GameThreadCollection.GetAnimateTransformAttribute();
+
+			const int32 NumParticles = GTParticles.Num();
+			for (int32 ParticleIndex = 0; ParticleIndex < NumParticles; ParticleIndex++)
+			{
+				if (const TUniquePtr<FParticle>& GTParticlePtr = GTParticles[ParticleIndex])
+				{
+					const int32 TransformIndex = FromParticleToTransformIndex[ParticleIndex];
+					const bool bAnimatingWhileDisabled = AnimationsActive ? (*AnimationsActive)[TransformIndex] : false;
+					const bool bActive = GameThreadCollection.Active[TransformIndex];
+					if (bActive || bAnimatingWhileDisabled)
+					{
+						const FTransform PreviousParticleTransform(GTParticlePtr->R(), GTParticlePtr->X());
+						const FTransform RelativeParticleTransform = PreviousParticleTransform.GetRelativeTransform(PreviousWorldTransform_External);
+						const FTransform NewParticleTransform = RelativeParticleTransform * WorldTransform_External;
+						GTParticlePtr->SetX(NewParticleTransform.GetTranslation());
+						GTParticlePtr->SetR(NewParticleTransform.GetRotation());
+					}
+				}
+			}
+		}
+
 		// Because SetWorldTransform_Internal is running on the physics thread we need to capture this parameter 
 		// (it is transient and may changed after calling this function)
 		const bool bSkipChildToParentUpdate = GetSkipChildToParentUpdateWhenInClusterUnion();
