Island manager crash: Add ensure for dynamic nodes without island or non-dynamic nodes with islands

diff --git a/Engine/Source/Runtime/Experimental/Chaos/Private/Chaos/Island/IslandManager.cpp b/Engine/Source/Runtime/Experimental/Chaos/Private/Chaos/Island/IslandManager.cpp
--- a/Engine/Source/Runtime/Experimental/Chaos/Private/Chaos/Island/IslandManager.cpp
+++ b/Engine/Source/Runtime/Experimental/Chaos/Private/Chaos/Island/IslandManager.cpp
@@ -522,7 +522,12 @@ namespace Chaos::Private
 
 	void FPBDIsland::Trash()
 	{
-		check(Nodes.IsEmpty());
+		// Print more information about remaining nodes to track down UE-296297
+		checkf(Nodes.IsEmpty(), 
+			TEXT("Number of nodes: %d, Nodes[0]: %s, Number of edges: %d"), 
+			Nodes.Num(), 
+			*(Nodes[0]->GetParticle()->GetDebugName()), 
+			Nodes[0]->Edges.Num());
 		check(NumEdges == 0);
 		check(ArrayIndex == INDEX_NONE);
 
@@ -1023,6 +1028,13 @@ namespace Chaos::Private
 				// all other loops are over islands and won't visit this node
 				DestroyGraphNode(Node);
 			}
+			else
+			{
+				// Every dynamic node needs to be in an island, non-dynamic nodes cannot be in an island.
+				ensureMsgf((Node->Flags.bIsDynamic && (Node->Island != nullptr)) || (!Node->Flags.bIsDynamic && (Node->Island == nullptr)),
+					TEXT("Node->Flags.bIsDynamic: %d, Node->Island: %d, Node->Edges.IsEmpty(): %d, DebugName: %s"), 
+					Node->Flags.bIsDynamic, Node->Island != nullptr, Node->Edges.IsEmpty(), *(Node->GetParticle()->GetDebugName()));
+			}
 		}
 	}
 
@@ -1078,6 +1090,12 @@ namespace Chaos::Private
 			// Reverse loop because we remove as we iterate
 			for (int32 EdgeIndex = Node->Edges.Num() - 1; EdgeIndex >= 0; --EdgeIndex)
 			{
+				// Print more information about remaining nodes to track down UE-297690
+				checkf(EdgeIndex < Node->Edges.Num(),
+					TEXT("Array index out of bounds: %d into an array of size %d, Node->GetDebugName() %s."),
+					EdgeIndex, Node->Edges.Num(), 
+					*(Node->GetParticle()->GetDebugName()));
+
 				FPBDIslandConstraint* Edge = Node->Edges[EdgeIndex];
 
 				// NOTE: Destroys the edge and possibly this and other nodes (if this was their last constraint)
@@ -1356,7 +1374,28 @@ namespace Chaos::Private
 			if (FPBDIslandParticle* Node = Edge->Nodes[NodeIndex])
 			{
 				const int32 ArrayIndex = Edge->NodeArrayIndices[NodeIndex];
-				check(Node->Edges[ArrayIndex] == Edge);
+
+				// These checks are supposed to help print out more information about the crashes in UE-215521, UE-296295, UE-314497
+				// If the array index is out of bounds, we print out information about "Edge", i.e. the edge we intended to remove. 
+				checkf(ArrayIndex < Node->Edges.Num(), 
+					TEXT("Array index out of bounds: %d into an array of size %d, Node: %s, Edge->Type: %s, Edge->Node[0]: %s, Edge->Node[1]: %s."),
+					ArrayIndex, Node->Edges.Num(),
+					*(Node->GetParticle()->GetDebugName()), 
+					*(Edge->GetConstraint()->GetType().ToString()), 
+					Edge->Nodes[0] ? *(Edge->Nodes[0]->GetParticle()->GetDebugName()) : *(FString()),
+					Edge->Nodes[1] ? *(Edge->Nodes[1]->GetParticle()->GetDebugName()) : *(FString()));
+				// If the array index is within the bounds but does not correspond to the constraint we intended to remove, 
+				// we print out information about "Edge", i.e. the edge we intended to remove, and the edge we currently reference "Node->Edges[ArrayIndex]".
+				checkf(Node->Edges[ArrayIndex] == Edge, 
+					TEXT("Selected edge does not correspond to ArrayIndex. Node: %s, Edge->Type: %s, Edge->Node[0]: %s, Edge->Node[1]: %s, Node->Edges[ArrayIndex]->Type: %s, Edges[ArrayIndex]->Node[0]: %s, Edges[ArrayIndex]->Node[1]: %s."), 
+					*(Node->GetParticle()->GetDebugName()), 
+					*(Edge->GetConstraint()->GetType().ToString()), 
+					Edge->Nodes[0] ? *(Edge->Nodes[0]->GetParticle()->GetDebugName()) : *(FString()),
+					Edge->Nodes[1] ? *(Edge->Nodes[1]->GetParticle()->GetDebugName()) : *(FString()),
+					*(Node->Edges[ArrayIndex]->GetConstraint()->GetType().ToString()), 
+					Node->Edges[ArrayIndex]->Nodes[0] ? *(Node->Edges[ArrayIndex]->Nodes[0]->GetParticle()->GetDebugName()) : *(FString()),
+					Node->Edges[ArrayIndex]->Nodes[1] ? *(Node->Edges[ArrayIndex]->Nodes[1]->GetParticle()->GetDebugName()) : *(FString()));
+				
 
 				// Remove the edge from the node
 				Node->Edges.RemoveAtSwap(ArrayIndex, EAllowShrinking::No);
@@ -1369,20 +1408,36 @@ namespace Chaos::Private
 					// We don't know which of the edges nodes we are...
 					if (Node->Edges[ArrayIndex]->Nodes[0] == Node)
 					{
-						// The node we swapped in must have been at the end of the array.
-						check(Node->Edges[ArrayIndex]->NodeArrayIndices[0] == Node->Edges.Num());
+						checkf(Node->Edges[ArrayIndex]->NodeArrayIndices[0] == Node->Edges.Num(), 
+							TEXT("The node we swapped in was not at the end of the array. Node->Edges[ArrayIndex]->NodeArrayIndices[0]: %d, Node->Edges.Num(): %d, Node->DebugName: %s, Node->Edges[ArrayIndex]->Type: %s, Edges[ArrayIndex]->Node[0]: %s, Edges[ArrayIndex]->Node[1]: %s."),
+							Node->Edges[ArrayIndex]->NodeArrayIndices[0], Node->Edges.Num(), 
+							*(Node->GetParticle()->GetDebugName()),
+							*(Node->Edges[ArrayIndex]->GetConstraint()->GetType().ToString()),
+							Node->Edges[ArrayIndex]->Nodes[0] ? *(Node->Edges[ArrayIndex]->Nodes[0]->GetParticle()->GetDebugName()) : *(FString()),
+							Node->Edges[ArrayIndex]->Nodes[1] ? *(Node->Edges[ArrayIndex]->Nodes[1]->GetParticle()->GetDebugName()) : *(FString()));
+
 						Node->Edges[ArrayIndex]->NodeArrayIndices[0] = ArrayIndex;
 					}
 					else if (Node->Edges[ArrayIndex]->Nodes[1] == Node)
 					{
-						// The node we swapped in must have been at the end of the array.
-						check(Node->Edges[ArrayIndex]->NodeArrayIndices[1] == Node->Edges.Num());
+						checkf(Node->Edges[ArrayIndex]->NodeArrayIndices[1] == Node->Edges.Num(),
+							TEXT("The node we swapped in was not at the end of the array. Node->Edges[ArrayIndex]->NodeArrayIndices[1]: %d, Node->Edges.Num(): %d, Node: %s, Node->Edges[ArrayIndex]->Type: %s, Edges[ArrayIndex]->Node[0]: %s, Edges[ArrayIndex]->Node[1]: %s."),
+							Node->Edges[ArrayIndex]->NodeArrayIndices[1], Node->Edges.Num(), 
+							*(Node->GetParticle()->GetDebugName()),
+							*(Node->Edges[ArrayIndex]->GetConstraint()->GetType().ToString()),
+							Node->Edges[ArrayIndex]->Nodes[0] ? *(Node->Edges[ArrayIndex]->Nodes[0]->GetParticle()->GetDebugName()) : *(FString()),
+							Node->Edges[ArrayIndex]->Nodes[1] ? *(Node->Edges[ArrayIndex]->Nodes[1]->GetParticle()->GetDebugName()) : *(FString()));
+
 						Node->Edges[ArrayIndex]->NodeArrayIndices[1] = ArrayIndex;
 					}
 					else
 					{
 						// If we get here, the list Node->Edges must be corrupted.
-						checkf(false, TEXT("Incorrect node referenced."));
+						checkf(false, 
+							TEXT("Node %s was not found between Node->Edges[ArrayIndex]->Nodes[0] %s and Node->Edges[ArrayIndex]->Nodes[1] %s."),
+							*(Node->GetParticle()->GetDebugName()), 
+							*(Node->Edges[ArrayIndex]->Nodes[0]->GetParticle()->GetDebugName()),
+							*(Node->Edges[ArrayIndex]->Nodes[1]->GetParticle()->GetDebugName()));
 					}
 				}
 			}
