Http errors from httpiodispatcher are a unique error code now.  Which means we can translate errors at a higher level when they are related to http.

diff --git a/Engine/Source/Runtime/Core/Private/Experimental/UnifiedError/UnifiedError.cpp b/Engine/Source/Runtime/Core/Private/Experimental/UnifiedError/UnifiedError.cpp
--- a/Engine/Source/Runtime/Core/Private/Experimental/UnifiedError/UnifiedError.cpp
+++ b/Engine/Source/Runtime/Core/Private/Experimental/UnifiedError/UnifiedError.cpp
@@ -95,6 +95,11 @@ FAnsiString FError::GetModuleIdString() const
 
 
 FText FError::GetErrorMessage(bool bIncludeContext) const
+{
+	return GetErrorMessage(FText::Join(FText::FromString(": "), NSLOCTEXT("UnifiedError", "DefaultLogPrepend", "{Root/ModuleIdString}.{Root/ErrorCodeString}"), GetFormatErrorText()), bIncludeContext);
+}
+
+FText FError::GetErrorMessage(const FText& InFormatText, bool bIncludeContext) const 
 {
 
 	EDetailFilter LogDetailFilter = EDetailFilter::IncludeInLogMessage;
@@ -104,12 +109,12 @@ FText FError::GetErrorMessage(bool bIncludeContext) const
 	}
 
 	TArray<FText> AdditionalContext;
-	AdditionalContext.Add(FText::Join(FText::FromString(": "), NSLOCTEXT("UnifiedError", "DefaultLogPrepend", "{Root/ModuleIdString}.{Root/ErrorCodeString}"), GetFormatErrorText()));
+	AdditionalContext.Add(InFormatText);
 
 	const IErrorDetails* DetailsIt = ErrorDetails.GetReference();
 	while (DetailsIt != nullptr)
 	{
-		if (DetailsIt->ShouldInclude(EDetailFilter::IncludeInContextLogMessage))
+		if (DetailsIt->ShouldInclude(LogDetailFilter))
 		{
 			AdditionalContext.Add(FText::FromString(FUtf8String::Printf("(%s:{%s})", DetailsIt->GetErrorDetailsTypeNameForLog().GetData(), DetailsIt->GetErrorDetailsTypeNameForLog().GetData())));
 		}
diff --git a/Engine/Source/Runtime/Core/Private/IO/IoDispatcher.cpp b/Engine/Source/Runtime/Core/Private/IO/IoDispatcher.cpp
--- a/Engine/Source/Runtime/Core/Private/IO/IoDispatcher.cpp
+++ b/Engine/Source/Runtime/Core/Private/IO/IoDispatcher.cpp
@@ -66,7 +66,8 @@ static const TCHAR* const GetIoErrorText_ErrorCodeTextArray[] =
 	TEXT("FileFlush Failed"),
 	TEXT("FileMove Failed"),
 	TEXT("FileClose Failed"),
-	TEXT("File Corrupt")
+	TEXT("File Corrupt"),
+	TEXT("Http Read Error")
 };
 static_assert(std::size(GetIoErrorText_ErrorCodeTextArray) == static_cast<uint32>(EIoErrorCode::Last), "GetIoErrorText_ErrorCodeTextArray is missing entries to match EIoErrorCode. Please update the array.");
 
diff --git a/Engine/Source/Runtime/Core/Public/Experimental/UnifiedError/UnifiedError.h b/Engine/Source/Runtime/Core/Public/Experimental/UnifiedError/UnifiedError.h
--- a/Engine/Source/Runtime/Core/Public/Experimental/UnifiedError/UnifiedError.h
+++ b/Engine/Source/Runtime/Core/Public/Experimental/UnifiedError/UnifiedError.h
@@ -145,6 +145,9 @@ constexpr uint64 CompileTimeHashString(const char* String)
 
 // TODO Daniel: StaticDetailsTypeId is intended to register the creation functions with the FErrorDetailsRegistry for support of serialization / deserialization of error details
 
+#define UE_ERROR_DETAILS_PARAM_NAME(DetailsNamespace, DetailsTypeName, ParamName) \
+	#DetailsNamespace"_"#DetailsTypeName"/"#ParamName
+
 #define UE_DECLARE_ERROR_DETAILS_INTERNAL(DetailsNamespace, TypeName)\
 	inline static constexpr FAnsiStringView StaticGetErrorDetailsTypeName() \
 	{\
@@ -466,6 +469,9 @@ namespace UE::UnifiedError
 		CORE_API FText GetErrorMessage(bool bAppendContext = false) const;
 		CORE_API FText GetFormatErrorText() const;
 
+		// returns an error message, using the Formatstring provided to format the error message
+		CORE_API FText GetErrorMessage(const FText& OverrideFormatString, bool bAppendContext = false) const;
+
 		template<typename T>
 		void PushErrorContext(T&& InErrorStruct, const EDetailFilter& InDetailFilter = EDetailFilter::Default)
 		{
diff --git a/Engine/Source/Runtime/Core/Public/IO/IoStatus.h b/Engine/Source/Runtime/Core/Public/IO/IoStatus.h
--- a/Engine/Source/Runtime/Core/Public/IO/IoStatus.h
+++ b/Engine/Source/Runtime/Core/Public/IO/IoStatus.h
@@ -45,6 +45,7 @@ enum class EIoErrorCode
 	FileMoveFailed,
 	FileCloseFailed,
 	FileCorrupt,
+	HttpReadError,
 	Last
 };
 
diff --git a/Engine/Source/Runtime/Experimental/IoStore/OnDemand/Private/OnDemandContentInstaller.cpp b/Engine/Source/Runtime/Experimental/IoStore/OnDemand/Private/OnDemandContentInstaller.cpp
--- a/Engine/Source/Runtime/Experimental/IoStore/OnDemand/Private/OnDemandContentInstaller.cpp
+++ b/Engine/Source/Runtime/Experimental/IoStore/OnDemand/Private/OnDemandContentInstaller.cpp
@@ -41,10 +41,17 @@ namespace UE::IoStore
 
 		FString IoStoreDebugRequest;
 		static FAutoConsoleVariableRef CVar_IoStoreDebugRequest(
-			TEXT("iostore.DebugRequest"),
-			IoStoreDebugRequest,
-			TEXT("When the request with a debug name partially matching this cvar is found iostore will print a message to the log on which a breakpoint can be placed.")
-		);
+				TEXT("iostore.DebugRequest"),
+				IoStoreDebugRequest,
+				TEXT("When the request with a debug name partially matching this cvar is found iostore will print a message to the log on which a breakpoint can be placed.")
+			);
+
+		int32 IoStoreSimulateHttpError = -1;
+		static FAutoConsoleVariableRef CVar_IoStoreSimulateHttpError(
+				TEXT("iostore.SimulateHttpError"),
+				IoStoreSimulateHttpError,
+				TEXT("Simulate http error code in iostore.")
+			);
 #endif // !UE_BUILD_SHIPPING
 
 		static bool bContentInstallerFlushOnIdle = false;
@@ -987,6 +994,12 @@ void FOnDemandContentInstaller::ExecuteInstallRequest(FRequest& Request)
 					Request.Priority,
 					[this, ChunkHash, RequestId = PendingChunkRequest->RequestId](uint32 HttpStatusCode, FStringView ErrorReason, FIoBuffer&& Chunk)
 					{
+#if !UE_BUILD_SHIPPING
+						if (IoStore::CVars::IoStoreSimulateHttpError >= 0)
+						{
+							HttpStatusCode = (uint32)IoStore::CVars::IoStoreSimulateHttpError;
+						}
+#endif
 						EIoErrorCode CompletionStatus = (HttpStatusCode > 199 && HttpStatusCode < 300) ? EIoErrorCode::Ok : EIoErrorCode::ReadError;
 						if (ErrorReason.Contains(TEXTVIEW("cancelled"), ESearchCase::IgnoreCase) ||
 							ErrorReason.Contains(TEXTVIEW("canceled"),  ESearchCase::IgnoreCase)) // Yes we spell it both ways (╯°□°）╯︵ ┻━┻
