[NetPhysics] Stop stuttering when resimulating objects that doesn't use resimulation as physics replication mode.

diff --git a/Engine/Source/Runtime/Engine/Private/PhysicsEngine/PhysicsReplication.cpp b/Engine/Source/Runtime/Engine/Private/PhysicsEngine/PhysicsReplication.cpp
--- a/Engine/Source/Runtime/Engine/Private/PhysicsEngine/PhysicsReplication.cpp
+++ b/Engine/Source/Runtime/Engine/Private/PhysicsEngine/PhysicsReplication.cpp
@@ -19,6 +19,7 @@
 #include "Chaos/DebugDrawQueue.h"
 #include "Chaos/Particles.h"
 #include "Chaos/Island/IslandManager.h"
+#include "RewindData.h"
 
 
 namespace CharacterMovementCVars
@@ -129,9 +130,12 @@ namespace PhysicsReplicationCVars
 		bool bRuntimeCorrectConnectedBodies = true;
 		static FAutoConsoleVariableRef CVarResimRuntimeCorrectConnectedBodies(TEXT("np2.Resim.RuntimeCorrectConnectedBodies"), bRuntimeCorrectConnectedBodies, TEXT("If true runtime position and rotation correction will also shift transform of any connected physics objects. Used if RuntimeCorrectionEnabled is true."));
 
-		bool bDisableReplicationOnInteraction = true;
+		bool bDisableReplicationOnInteraction = false;
 		static FAutoConsoleVariableRef CVarResimDisableReplicationOnInteraction(TEXT("np2.Resim.DisableReplicationOnInteraction"), bDisableReplicationOnInteraction, TEXT("If a resim object interacts with another object not running resimulation, deactivate that objects replication until interaction stops."));
 
+		bool bKeepResimStateForNonResimReplicatedObjects = false;
+		static FAutoConsoleVariableRef CVarResimKeepResimStateForNonResimReplicatedObjects(TEXT("np2.Resim.KeepResimStateForNonResimReplicatedObjects"), bKeepResimStateForNonResimReplicatedObjects, TEXT("When false, move objects back to their pre-resim state at the end of a resimulation if they are not set to use resimulation as replication mode."));
+
 		float PosStabilityMultiplier = 0.5f;
 		static FAutoConsoleVariableRef CVarResimPosStabilityMultiplier(TEXT("np2.Resim.PosStabilityMultiplier"), PosStabilityMultiplier, TEXT("Recommended range between 0.0-1.0. Lower value means more stable positional corrections."));
 
@@ -1000,6 +1004,59 @@ void FPhysicsReplicationAsync::OnPreSimulate_Internal()
 	ApplyTargetStatesAsync(GetDeltaTime_Internal());
 }
 
+void FPhysicsReplicationAsync::OnPostSolve_Internal()
+{
+	if (FPhysicsReplication::ShouldSkipPhysicsReplication() || PhysicsReplicationCVars::ResimulationCVars::bKeepResimStateForNonResimReplicatedObjects)
+	{
+		return;
+	}
+
+	Chaos::FPBDRigidsSolver* RigidsSolver = static_cast<Chaos::FPBDRigidsSolver*>(GetSolver());
+	check(RigidsSolver);
+
+	// On final resimulation frame, revert non-resim replicated objects to their pre-resim state
+	Chaos::FRewindData* RewindData = RigidsSolver->GetRewindData();
+	if (RewindData && RewindData->IsFinalResim())
+	{
+		Chaos::FReadPhysicsObjectInterface_Internal Interface = Chaos::FPhysicsObjectInternalInterface::GetRead();
+		for (auto Itr = ObjectToTarget.CreateIterator(); Itr; ++Itr)
+		{
+			FReplicatedPhysicsTargetAsync& Target = Itr.Value();
+			if (Target.RepMode != EPhysicsReplicationMode::Resimulation)
+			{
+				Chaos::FConstPhysicsObjectHandle& PhysicsObject = Itr.Key();
+				if (Chaos::FGeometryParticleHandle* GeometryParticle = Interface.GetParticle(PhysicsObject))
+				{
+					Chaos::FParticleSimpleState* PreResimState = nullptr;
+					RewindData->GetPreResimState(GeometryParticle, PreResimState);
+					if (PreResimState)
+					{
+						if (Chaos::FPBDRigidParticleHandle* RigidParticle = GeometryParticle->CastToRigidParticle())
+						{
+							// Note, PostSolve happens before rigids update XR from PQ
+							RigidParticle->SetP(PreResimState->GetX());
+							RigidParticle->SetQ(PreResimState->GetR());
+							RigidParticle->SetV(PreResimState->GetV());
+							RigidParticle->SetW(PreResimState->GetW());
+						}
+						else if (Chaos::FKinematicGeometryParticleHandle* KinematicParticle = GeometryParticle->CastToKinematicParticle())
+						{
+							// Kinematic particles doesn't have PR so set XR directly
+							KinematicParticle->SetX(PreResimState->GetX());
+							KinematicParticle->SetR(PreResimState->GetR());
+							KinematicParticle->SetV(PreResimState->GetV());
+							KinematicParticle->SetW(PreResimState->GetW());
+						}
+
+						// Remove pre resim state
+						RewindData->RemovePreResimState(GeometryParticle);
+					}
+				}
+			}
+		}
+	}
+}
+
 FReplicatedPhysicsTargetAsync* FPhysicsReplicationAsync::AddObjectToReplication(Chaos::FConstPhysicsObjectHandle PhysicsObject)
 {
 	if (ensure(PhysicsObject))
diff --git a/Engine/Source/Runtime/Engine/Public/PhysicsReplication.h b/Engine/Source/Runtime/Engine/Public/PhysicsReplication.h
--- a/Engine/Source/Runtime/Engine/Public/PhysicsReplication.h
+++ b/Engine/Source/Runtime/Engine/Public/PhysicsReplication.h
@@ -189,11 +189,12 @@ class FPhysicsReplicationAsync : public IPhysicsReplicationAsync,
 	public Chaos::TSimCallbackObject<
 	FPhysicsReplicationAsyncInput,
 	Chaos::FSimCallbackNoOutput,
-	Chaos::ESimCallbackOptions::Presimulate | Chaos::ESimCallbackOptions::PhysicsObjectUnregister>
+	Chaos::ESimCallbackOptions::Presimulate | Chaos::ESimCallbackOptions::PostSolve | Chaos::ESimCallbackOptions::PhysicsObjectUnregister>
 {
 	virtual FName GetFNameForStatId() const override;
 	virtual void OnPostInitialize_Internal() override;
 	virtual void OnPreSimulate_Internal() override;
+	virtual void OnPostSolve_Internal() override;
 	virtual void OnPhysicsObjectUnregistered_Internal(Chaos::FConstPhysicsObjectHandle PhysicsObject) override;
 
 	virtual void ApplyTargetStatesAsync(const float DeltaSeconds);
diff --git a/Engine/Source/Runtime/Experimental/Chaos/Private/RewindData.cpp b/Engine/Source/Runtime/Experimental/Chaos/Private/RewindData.cpp
--- a/Engine/Source/Runtime/Experimental/Chaos/Private/RewindData.cpp
+++ b/Engine/Source/Runtime/Experimental/Chaos/Private/RewindData.cpp
@@ -904,10 +904,6 @@ void FRewindData::RewindToFrame(int32 RewindFrame)
 		
 			FGeometryParticleStateBase& History = DirtyParticleInfo.GetHistory(); //non-const in case we need to record what's at head for a rewind (CurFrame has already been increased to the next frame)
 
-			PRAGMA_DISABLE_DEPRECATION_WARNINGS
-			History.CachePreCorrectionState(*GeometryParticle); // Deprecated UE 5.6
-			PRAGMA_ENABLE_DEPRECATION_WARNINGS
-
 			if (!ResimConsoleVars::bIsBubbleResimEnabled || GeometryParticle->ResimType() != EResimType::FrozenDuringResim)
 			{
 				CachePreResimState(*GeometryParticle);
@@ -1598,18 +1594,38 @@ void FRewindData::CachePreResimState(FGeometryParticleHandle& Handle)
 	}
 
 	// Find or add pre-resim error for the particle that has an error
-	auto PreErrorInfo = [this](FGeometryParticleHandle* ParticleHandle) -> FDirtyParticleErrorInfo&
+	auto GetPreResimState = [this](FGeometryParticleHandle* ParticleHandle) -> FParticleSimpleState&
 		{
-			if (FDirtyParticleErrorInfo* Found = DirtyParticlePreResimState.Find(ParticleHandle))
+			if (FParticleSimpleState* Found = DirtyParticlePreResimState.Find(ParticleHandle))
 			{
 				return *Found;
 			}
 
-			return DirtyParticlePreResimState.Add(ParticleHandle, FDirtyParticleErrorInfo(*ParticleHandle));
+			return DirtyParticlePreResimState.Add(ParticleHandle, FParticleSimpleState(*ParticleHandle));
 		};
 
-	// Cache dirty particle XR before a resimulation
-	PreErrorInfo(&Handle).AccumulateError(Handle.GetX(), Handle.GetR());
+	auto CacheState = [&](FGeometryParticleHandle* ParticleHandle)
+		{
+			if (!ParticleHandle)
+			{
+				return;
+			}
+
+			FParticleSimpleState& PreResimState = GetPreResimState(ParticleHandle);
+
+			// Cache dirty particle XR before a resimulation
+			PreResimState.CacheXR(ParticleHandle->GetX(), ParticleHandle->GetR());
+
+			/** Cache dirty kinematic or rigid particle VW before a resimulation
+			* NOTE: Rigids are based on kinematic so can always cast to kinematic */
+			if (FKinematicGeometryParticleHandle* KinematicHandle = ParticleHandle->CastToKinematicParticle())
+			{
+				PreResimState.CacheVW(KinematicHandle->GetV(), KinematicHandle->GetW());
+			}
+		};
+
+	// Cache dirty particle state before a resimulation
+	CacheState(&Handle);
 
 	// If particle is a Cluster Union, also cache child particles
 	if (PhysicsProxy->GetType() == EPhysicsProxyType::ClusterUnionProxy)
@@ -1627,7 +1643,7 @@ void FRewindData::CachePreResimState(FGeometryParticleHandle& Handle)
 						{
 							if (FGeometryParticleHandle* ChildHandle = ChildProxy->GetHandle_LowLevel())
 							{
-								PreErrorInfo(ChildHandle).AccumulateError(ChildHandle->GetX(), ChildHandle->GetR());
+								CacheState(ChildHandle);
 							}
 						}
 						break;
@@ -1638,7 +1654,7 @@ void FRewindData::CachePreResimState(FGeometryParticleHandle& Handle)
 						{
 							if (FGeometryParticleHandle* ChildHandle = ChildProxy->GetParticle_Internal())
 							{
-								PreErrorInfo(ChildHandle).AccumulateError(ChildHandle->GetX(), ChildHandle->GetR());
+								CacheState(ChildHandle);
 							}
 						}
 						break;
@@ -1649,7 +1665,7 @@ void FRewindData::CachePreResimState(FGeometryParticleHandle& Handle)
 						{
 							if (FGeometryParticleHandle* ChildHandle = ChildProxy->GetInitialRootParticle_Internal())
 							{
-								PreErrorInfo(ChildHandle).AccumulateError(ChildHandle->GetX(), ChildHandle->GetR());
+								CacheState(ChildHandle);
 							}
 						}
 						break;
@@ -1658,21 +1674,30 @@ void FRewindData::CachePreResimState(FGeometryParticleHandle& Handle)
 			}
 		}
 	}
+}
 
+void FRewindData::GetPreResimState(Chaos::FGeometryParticleHandle* GeometryParticle, FParticleSimpleState*& OutPreResimState)
+{
+	OutPreResimState = DirtyParticlePreResimState.Find(GeometryParticle);
+}
+
+void FRewindData::RemovePreResimState(Chaos::FGeometryParticleHandle* GeometryParticle)
+{
+	DirtyParticlePreResimState.Remove(GeometryParticle, EAllowShrinking::No);
 }
 
 template<>
 void FRewindData::AccumulateErrorIfNecessary(FGeometryParticleHandle& Obj, const FFrameAndPhase FrameAndPhase)
 {
-	FDirtyParticleErrorInfo* PreErrorInfo = DirtyParticlePreResimState.Find(&Obj);
-	if (!PreErrorInfo)
+	FParticleSimpleState* PreResimState = DirtyParticlePreResimState.Find(&Obj);
+	if (!PreResimState)
 	{
 		return;
 	}
 
 	// Get the error offset after a correction
-	const FVec3 ErrorX = PreErrorInfo->GetErrorX() - Obj.GetX();
-	FQuat ErrorR = Obj.GetR().Inverse() * PreErrorInfo->GetErrorR(); // ErrorR in local space 
+	const FVec3 ErrorX = PreResimState->GetX() - Obj.GetX();
+	FQuat ErrorR = Obj.GetR().Inverse() * PreResimState->GetR(); // ErrorR in local space 
 	ErrorR.EnforceShortestArcWith(FQuat::Identity);
 	ErrorR.Normalize();
 
diff --git a/Engine/Source/Runtime/Experimental/Chaos/Public/RewindData.h b/Engine/Source/Runtime/Experimental/Chaos/Public/RewindData.h
--- a/Engine/Source/Runtime/Experimental/Chaos/Public/RewindData.h
+++ b/Engine/Source/Runtime/Experimental/Chaos/Public/RewindData.h
@@ -1311,14 +1311,6 @@ struct FGeometryParticleStateBase
 	void SyncSimWritablePropsFromSim(FDirtyPropData Manager,const TPBDRigidParticleHandle<FReal,3>& Rigid);
 	void SyncDirtyDynamics(FDirtyPropData& DestManager,const FDirtyChaosProperties& Dirty,const FConstDirtyPropData& SrcManager);
 	
-	template<typename TParticle>
-	UE_DEPRECATED(5.6, "Deprecated, use FRewindData::CachePreResimState instead. Not all moving particles are marked as dirty, for example GeometryCollection children of a ClusterUnion.")
-	void CachePreCorrectionState(const TParticle& Particle)
-	{
-		PreCorrectionXR.SetX(Particle.GetX());
-		PreCorrectionXR.SetR(Particle.GetR());
-	}
-
 	/** Setting bNoEntryIsHead = false in TParticlePropertyBuffer allows us to not cache the data each frame and it will only return entries for a specific frame when fetched 
 	* For example, we don't need to cache a kinematic target that is set to None, instead we only cache when it's not set to None and if we don't have an entry for a specific frame we can expect it to be a None kinematic target.
 	* Note that bNoEntryIsHead = true (which is default) will return the latest cached entry closest to the requested frame and phase without returning an earlier entry. */
@@ -1335,9 +1327,38 @@ struct FGeometryParticleStateBase
 	TParticlePropertyBuffer<FParticleDynamicMisc, EChaosProperty::DynamicMisc, /*bNoEntryIsHead=*/false> TargetStates;
 
 	FShapesArrayStateBase ShapesArrayState;
+};
+
+struct FParticleSimpleState
+{
+private:
+	FGeometryParticleHandle* HandlePtr;
+	FVec3 X = { 0,0,0 };
+	FQuat R = FQuat::Identity;
+	FVec3 V = { 0,0,0 };
+	FVec3 W = { 0,0,0 };
+
+public:
+	FParticleSimpleState(FGeometryParticleHandle& InHandle) : HandlePtr(&InHandle)
+	{
+	}
 
-	UE_DEPRECATED(5.6, "Deprecated, use FRewindData::DirtyParticlePreResimState")
-	FParticlePositionRotation PreCorrectionXR;
+	void CacheXR(FVec3 InX, FQuat InR)
+	{
+		X = InX;
+		R = InR;
+	}
+	void CacheVW(FVec3 InV, FVec3 InW)
+	{
+		V = InV;
+		W = InW;
+	}
+
+	FGeometryParticleHandle* GetObjectPtr() const { return HandlePtr; }
+	FVec3 GetX() const { return X; }
+	FQuat GetR() const { return R; }
+	FVec3 GetV() const { return V; }
+	FVec3 GetW() const { return W; }
 };
 
 class FGeometryParticleState
@@ -1945,6 +1966,12 @@ public:
 	* NOTE: This only happens when the client has desynced and is behind the server, so we receive server states for frames not yet simulated */
 	const int32 CHAOS_API CompareTargetsToLastFrame();
 
+	/** Get cached pre resimulation state for particle, if there is one */
+	void CHAOS_API GetPreResimState(Chaos::FGeometryParticleHandle* GeometryParticle, FParticleSimpleState*& OutPreResimState);
+
+	/** Remove cached pre resimulation state for particle*/
+	void CHAOS_API RemovePreResimState(Chaos::FGeometryParticleHandle* GeometryParticle);
+
 	static bool CHAOS_API CheckVectorThreshold(FVec3 A, FVec3 B, float Threshold);
 	static bool CHAOS_API CheckQuaternionThreshold(FQuat A, FQuat B, float ThresholdDegrees);
 
@@ -2159,7 +2186,7 @@ private:
 
 	TDirtyObjects<FDirtyParticleInfo> DirtyParticles;
 	TDirtyObjects<FDirtyJointInfo> DirtyJoints;
-	TDirtyObjects<FDirtyParticleErrorInfo> DirtyParticlePreResimState;
+	TDirtyObjects<FParticleSimpleState> DirtyParticlePreResimState;
 	TDirtyObjects<FDirtyParticleErrorInfo> DirtyParticleErrors;
 
 	TArray<TWeakPtr<FBaseRewindHistory>> InputHistories; // Todo, deprecate in favor of Particle->Input map?
