Fix issue where http cancellation logging was always being printed even though it was supposed to be VeryVerybose only.

#rnx

- The cancellation logging uses FMsg::Logf so that it can use the dynamic log category provided by the FMultiEndpointHttpClient but was not checking if the verbosity level of that log category was met, as FMsg does not do that.
- Now that we no longer have multiple instances of FMultiEndpointHttpClient but a single shared instance used by all systems it makes no sense for us to continue with the dynamic LogCategory/Verbosity.
- This has now been removed and the log category 'LogHttpIoDispatcher' is used everywhere.
- The dynamic verbosity of FMultiEndpointHttpClient::Log has been left alone for now.

diff --git a/Engine/Source/Runtime/Experimental/IoStore/OnDemand/Private/OnDemandHttpClient.cpp b/Engine/Source/Runtime/Experimental/IoStore/OnDemand/Private/OnDemandHttpClient.cpp
--- a/Engine/Source/Runtime/Experimental/IoStore/OnDemand/Private/OnDemandHttpClient.cpp
+++ b/Engine/Source/Runtime/Experimental/IoStore/OnDemand/Private/OnDemandHttpClient.cpp
@@ -8,6 +8,7 @@
 #include "HAL/FileManager.h"
 #include "HAL/IConsoleManager.h"
 #include "HAL/PlatformTime.h"
+#include "IO/HttpIoDispatcher.h"
 #include "IO/IoStoreOnDemandInternals.h"
 #include "Misc/ConfigCacheIni.h"
 #include "Misc/Paths.h"
@@ -329,7 +330,7 @@ void FMultiEndpointHttpClient::CancelRequest(FHttpTicketId TicketId)
 		if (Ticket->HttpTicket != 0)
 		{
 #if !NO_LOGGING
-			FMsg::Logf(__FILE__, __LINE__, Config.LogCategory->GetCategoryName(), ELogVerbosity::Verbose, TEXT("Canceling FHttpTicketId %u with FTicket: %llu"), TicketId, Ticket->HttpTicket);
+			UE_LOG(LogHttpIoDispatcher, Verbose, TEXT("Canceling FHttpTicketId %u with FTicket: %llu"), TicketId, Ticket->HttpTicket);
 #endif // !NO_LOGGING
 
 			EventLoop.Cancel(Ticket->HttpTicket);
@@ -752,14 +753,14 @@ void FMultiEndpointHttpClient::Log(const FRequest& Request, const HTTP::FTicketS
 
 	using namespace UE::IoStore::HTTP;
 
-	ELogVerbosity::Type Verbosity = Config.LogVerbosity;
+	ELogVerbosity::Type Verbosity = ELogVerbosity::VeryVerbose;
 	// We count  a request as failed if it's status code is bad AND the request was not canceled
 	if (GIaxWarnOnFailure && !IsHttpStatusOk(Request.StatusCode) && TicketStatus.GetId() != FTicketStatus::EId::Cancelled)
 	{
 		Verbosity = ELogVerbosity::Type::Warning;
 	}
 
-	if (Config.LogCategory == nullptr || Config.LogCategory->IsSuppressed(Verbosity))
+	if (LogHttpIoDispatcher.IsSuppressed(Verbosity))
 	{
 		return;
 	}
@@ -781,7 +782,7 @@ void FMultiEndpointHttpClient::Log(const FRequest& Request, const HTTP::FTicketS
 		Reason << TEXT(" ") << Request.ResponseMessage;
 	}
 
-	FMsg::Logf(__FILE__, __LINE__, Config.LogCategory->GetCategoryName(), Verbosity,
+	FMsg::Logf(__FILE__, __LINE__, LogHttpIoDispatcher.GetCategoryName(), Verbosity,
 		TEXT("http-%3u: %5" UINT64_FMT "ms %5" UINT64_FMT "KiB [%4s] %s (Attempt %u/%u)%s"),
 		Request.StatusCode,
 		DurationMs,
diff --git a/Engine/Source/Runtime/Experimental/IoStore/OnDemand/Private/OnDemandHttpClient.h b/Engine/Source/Runtime/Experimental/IoStore/OnDemand/Private/OnDemandHttpClient.h
--- a/Engine/Source/Runtime/Experimental/IoStore/OnDemand/Private/OnDemandHttpClient.h
+++ b/Engine/Source/Runtime/Experimental/IoStore/OnDemand/Private/OnDemandHttpClient.h
@@ -68,9 +68,6 @@ struct FMultiEndpointHttpClientConfig
 	bool						bAllowChunkedTransfer = true;
 	/** When true the client will return the response body for requests outside of the 2XX range. When false only requests within the 2XX range will return the response body */
 	bool						bResponseBodyOnError = false;
-	/** Logging will be disabled if this is set to a nullptr, it is up to the calling system to assign a LogCategory */
-	const UE_LOG_CATEGORY_TYPE*	LogCategory = nullptr;
-	ELogVerbosity::Type			LogVerbosity = ELogVerbosity::Log;
 };
 
 enum class EMultiEndpointRequestFlags : uint8
diff --git a/Engine/Source/Runtime/Experimental/IoStore/OnDemand/Private/OnDemandHttpIoDispacher.cpp b/Engine/Source/Runtime/Experimental/IoStore/OnDemand/Private/OnDemandHttpIoDispacher.cpp
--- a/Engine/Source/Runtime/Experimental/IoStore/OnDemand/Private/OnDemandHttpIoDispacher.cpp
+++ b/Engine/Source/Runtime/Experimental/IoStore/OnDemand/Private/OnDemandHttpIoDispacher.cpp
@@ -1229,9 +1229,7 @@ bool FHttpDispatcher::TryCreateHttpClient()
 			.MaxRetryCount = GIasHttpRetryCount,
 			.TimeoutMs = GIasHttpFailTimeOutMs,
 			.Redirects = EHttpRedirects::Follow,
-			.bAllowChunkedTransfer = true,
-			.LogCategory = &LogHttpIoDispatcher,
-			.LogVerbosity = ELogVerbosity::VeryVerbose
+			.bAllowChunkedTransfer = true
 		});
 
 	return HttpClient.IsValid();
diff --git a/Engine/Source/Runtime/Experimental/IoStore/OnDemand/Private/OnDemandHttpThread.cpp b/Engine/Source/Runtime/Experimental/IoStore/OnDemand/Private/OnDemandHttpThread.cpp
--- a/Engine/Source/Runtime/Experimental/IoStore/OnDemand/Private/OnDemandHttpThread.cpp
+++ b/Engine/Source/Runtime/Experimental/IoStore/OnDemand/Private/OnDemandHttpThread.cpp
@@ -655,9 +655,7 @@ bool FOnDemandHttpThread::TryCreateHttpClient()
 			.TimeoutMs = GIasHttpFailTimeOutMs,
 			.Redirects = EHttpRedirects::Disabled,
 			.bEnableThreadSafetyChecks = true,
-			.bAllowChunkedTransfer = GIasHttpAllowChunkedXfer,
-			.LogCategory = &LogIas,
-			.LogVerbosity = ELogVerbosity::VeryVerbose
+			.bAllowChunkedTransfer = GIasHttpAllowChunkedXfer
 		});
 
 	return HttpClient.IsValid();
@@ -747,8 +745,6 @@ void FOnDemandHttpThread::OnTogglePause(bool bPause, EHttpRequestTypeFilter Filt
 	if (bPause)
 	{
 		EnumAddFlags(PauseFilter, Filter);
-
-
 	}
 	else
 	{
