Add a new debug cvar command 'ias.RunRequestTest' to run latency tests on all urls in all IAX hostgroups.

#rnx

diff --git a/Engine/Source/Runtime/Experimental/IoStore/OnDemand/Private/DebugCommands.cpp b/Engine/Source/Runtime/Experimental/IoStore/OnDemand/Private/DebugCommands.cpp
--- a/Engine/Source/Runtime/Experimental/IoStore/OnDemand/Private/DebugCommands.cpp
+++ b/Engine/Source/Runtime/Experimental/IoStore/OnDemand/Private/DebugCommands.cpp
@@ -7,10 +7,11 @@
 #include "IO/IoChunkId.h"
 #include "IO/IoDispatcher.h"
 #include "IO/IoStoreOnDemand.h"
+#include "IasHostGroup.h"
+#include "LatencyTesting.h"
 #include "Misc/AssertionMacros.h"
 #include "OnDemandIoStore.h"
 #include "String/Numeric.h"
-#include "IasHostGroup.h"
 
 #if !UE_BUILD_SHIPPING
 
@@ -51,6 +52,13 @@ void FOnDemandDebugCommands::BindConsoleCommands()
 			TEXT("[optional Number of requests to make] Creates a number of requests for chunks that can be found in that IAS system"),
 			FConsoleCommandWithArgsDelegate::CreateRaw(this, &FOnDemandDebugCommands::RunRequestTest),
 			ECVF_Default));
+
+	DynamicConsoleCommands.Emplace(
+		IConsoleManager::Get().RegisterConsoleCommand(
+			TEXT("iax.LatencyTests"),
+			TEXT(""),
+			FConsoleCommandDelegate::CreateRaw(this, &FOnDemandDebugCommands::RunLatencyTests),
+			ECVF_Default));
 #endif // !NO_CVARS
 }
 
@@ -147,6 +155,31 @@ void FOnDemandDebugCommands::RunRequestTest(const TArray<FString>& Args) const
 	UE_CLOG(bQuitAfterTest, LogIas, Display, TEXT("Process will quit once the requests have been completed"));
 }
 
+void FOnDemandDebugCommands::RunLatencyTests() const
+{
+	bool bLoggedHostGroup = false;
+
+	FHostGroupManager::Get().ForEachHostGroup([&bLoggedHostGroup](const FIASHostGroup& HostGroup)
+	{
+		if (!HostGroup.IsResolved() || HostGroup.GetHostUrls().IsEmpty())
+		{
+			return;
+		}
+
+		bLoggedHostGroup = true;
+
+		UE_LOG(LogIas, Log, TEXT("--------------------------------------------------------------------------------"));
+		UE_LOG(LogIas, Log, TEXT("Testing Hostgroup '%s':"), *HostGroup.GetName().ToString());
+
+		for (const FAnsiString& Url : HostGroup.GetHostUrls())
+		{
+			ConnectionTest(Url, HostGroup.GetTestPath(), 0);
+		}
+	});
+
+	UE_CLOG(bLoggedHostGroup, LogIas, Log, TEXT("--------------------------------------------------------------------------------"));
+}
+
 } // namespace UE::IoStore
 
 #endif // !UE_BUILD_SHIPPING
diff --git a/Engine/Source/Runtime/Experimental/IoStore/OnDemand/Private/DebugCommands.h b/Engine/Source/Runtime/Experimental/IoStore/OnDemand/Private/DebugCommands.h
--- a/Engine/Source/Runtime/Experimental/IoStore/OnDemand/Private/DebugCommands.h
+++ b/Engine/Source/Runtime/Experimental/IoStore/OnDemand/Private/DebugCommands.h
@@ -29,6 +29,7 @@ private:
 	void UnbindConsoleCommands();
 
 	void RunRequestTest(const TArray<FString>& Args) const;
+	void RunLatencyTests() const;
 
 	FOnDemandIoStore* OnDemandIoStore = nullptr;
 
