Collision Filter Refactor Part 4.5: Updating SimRaycast to properly use the new API. Also removed some old internal functions.

diff --git a/Engine/Source/Programs/HeadlessChaos/Private/HeadlessChaosTestSimSweep.cpp b/Engine/Source/Programs/HeadlessChaos/Private/HeadlessChaosTestSimSweep.cpp
--- a/Engine/Source/Programs/HeadlessChaos/Private/HeadlessChaosTestSimSweep.cpp
+++ b/Engine/Source/Programs/HeadlessChaos/Private/HeadlessChaosTestSimSweep.cpp
@@ -156,7 +156,7 @@ namespace ChaosTest
 		{
 			Private::FSimRaycastHit Hits[2];
 
-			Filter::FQueryFilterData QueryFilter = Filter::FQueryFilterBuilder::CreateLegacyObjectTypeFilter(1 << 0, false);
+			Filter::FQueryFilterData QueryFilter = Filter::FQueryObjectFilterBuilder().SetObjectTypes(1 << 0).Build();
 
 			bool bAnyHits = Private::SimRaycastBatchFirstHits(
 				Test.GetEvolution().GetSpatialAcceleration(),
@@ -178,7 +178,7 @@ namespace ChaosTest
 		{
 			Private::FSimRaycastHit Hits[2];
 
-			Filter::FQueryFilterData QueryFilter = Filter::FQueryFilterBuilder::CreateLegacyObjectTypeFilter(1 << 1, false);
+			Filter::FQueryFilterData QueryFilter = Filter::FQueryObjectFilterBuilder().SetObjectTypes(1 << 1).Build();
 
 			bool bAnyHits = Private::SimRaycastBatchFirstHits(
 				Test.GetEvolution().GetSpatialAcceleration(),
@@ -216,7 +216,7 @@ namespace ChaosTest
 
 		Private::FSimRaycastHit Hit;
 
-		Filter::FQueryFilterData QueryFilter = Filter::FQueryFilterBuilder::CreateLegacyObjectTypeFilter(0xFFFFFFFF, false);
+		Filter::FQueryFilterData QueryFilter = Filter::FQueryObjectFilterBuilder().SetObjectTypes(0xFFFFFFFF).Build();
 
 		bool bAnyHits = Private::SimRaycastBatchFirstHits(
 			Test.GetEvolution().GetSpatialAcceleration(),
@@ -261,7 +261,7 @@ namespace ChaosTest
 
 		Private::FSimRaycastHit Hits[2];
 
-		Filter::FQueryFilterData QueryFilter = Filter::FQueryFilterBuilder::CreateLegacyObjectTypeFilter(0xFFFFFFFF, false);
+		Filter::FQueryFilterData QueryFilter = Filter::FQueryObjectFilterBuilder().SetObjectTypes(0xFFFFFFFF).Build();
 
 		bool bAnyHits = Private::SimRaycastBatchFirstHits(
 			Test.GetEvolution().GetSpatialAcceleration(),
diff --git a/Engine/Source/Runtime/Experimental/Chaos/Private/Chaos/Collision/SimRaycast.cpp b/Engine/Source/Runtime/Experimental/Chaos/Private/Chaos/Collision/SimRaycast.cpp
--- a/Engine/Source/Runtime/Experimental/Chaos/Private/Chaos/Collision/SimRaycast.cpp
+++ b/Engine/Source/Runtime/Experimental/Chaos/Private/Chaos/Collision/SimRaycast.cpp
@@ -77,7 +77,7 @@ namespace Chaos::Private
 	}
 
 	FSimQueryDataParticleFilter::FSimQueryDataParticleFilter(const Filter::FQueryFilterData& InFilterData)
-		: FilterData(Filter::FQueryFilterBuilder::GetLegacyQueryFilter(InFilterData))
+		: FilterData(InFilterData)
 	{
 	}
 
@@ -87,7 +87,7 @@ namespace Chaos::Private
 	}
 
 	FSimQueryDataShapeFilter::FSimQueryDataShapeFilter(const Filter::FQueryFilterData& InFilterData, bool bInDetectSimple, bool bInDetectComplex)
-		: FilterData(Filter::FQueryFilterBuilder::GetLegacyQueryFilter(InFilterData))
+		: FilterData(InFilterData)
 	{
 		Flags.bDetectSimple = bInDetectSimple;
 		Flags.bDetectComplex = bInDetectComplex;
@@ -107,10 +107,9 @@ namespace Chaos::Private
 		}
 
 		// Run the channel filter
-		// NOTE: PrePreSimFilterImp return true if we want to skip
-		PRAGMA_DISABLE_DEPRECATION_WARNINGS
-		return !PrePreSimFilterImp(FilterData, Shape->GetSimData());
-		PRAGMA_ENABLE_DEPRECATION_WARNINGS
+		// NOTE: We return true if we want to keep the result
+		const Filter::ENarrowFilterResult HitType = FilterData.NarrowFilter(Shape->GetShapeFilterData(), true);
+		return HitType != Filter::ENarrowFilterResult::None;
 	}
 
 	bool SimRaycastBatchFirstHits(
diff --git a/Engine/Source/Runtime/Experimental/Chaos/Private/Chaos/CollisionFilterData.cpp b/Engine/Source/Runtime/Experimental/Chaos/Private/Chaos/CollisionFilterData.cpp
--- a/Engine/Source/Runtime/Experimental/Chaos/Private/Chaos/CollisionFilterData.cpp
+++ b/Engine/Source/Runtime/Experimental/Chaos/Private/Chaos/CollisionFilterData.cpp
@@ -541,25 +541,6 @@ namespace Chaos::Filter
 		return Builder.Build();
 	}
 
-	FShapeFilterData FShapeFilterBuilder::BuildLegacyShapeFilter(const uint8 InChannelIndex, const uint32 InQueryOverlapChannels, const uint32 InQueryBlockChannels, const uint32 InSimBlockChannels, const uint8 InMaskFilter, const EFilterFlags InFilterFlags)
-	{
-		// (Mask Filter 6 bits) | (Channel Index 5 bits) | (Flags 21 bits)
-		FShapeFilterData Result;
-		Result.QueryBlockChannels = InQueryBlockChannels;
-		Result.QueryOverlapChannels = InQueryOverlapChannels;
-		Result.SimBlockChannels = InSimBlockChannels;
-		Filter::SetFlags(Result.Word3, (uint32)InFilterFlags);
-		Filter::SetCollisionChannel(Result.Word3, (uint32)InChannelIndex);
-		Filter::SetMaskFilter(Result.Word3, InMaskFilter);
-		Result.SimWord3 = Result.Word3;
-		return Result;
-	}
-
-	FShapeFilterData FShapeFilterBuilder::BuildLegacyBlockAllSimShapeFilter(const uint8 InChannelIndex, const uint8 InMaskFilter, const EFilterFlags InFilterFlags)
-	{
-		return BuildLegacyShapeFilter(InChannelIndex, 0, 0, (uint32)-1, InMaskFilter, InFilterFlags);
-	}
-
 	FShapeFilterData FShapeFilterBuilder::BuildFromLegacyQueryFilter(const FCollisionFilterData& QueryFilter)
 	{
 		FShapeFilterData Result;
@@ -697,27 +678,6 @@ namespace Chaos::Filter
 		OutCombinedShapeFilterData.ShapeFilterData.SimWord3 = Word3;
 	}
 
-	FQueryFilterData FQueryFilterBuilder::CreateLegacyObjectTypeFilter(const uint32 ObjectTypesToQuery, const bool bMultiTrace, const uint8 IgnoreMask, const EFilterFlags FilterFlags)
-	{
-		const uint8 MultiTraceValue = bMultiTrace ? 0x1 : 0x0;
-		FQueryFilterData Result;
-		Result.Word0 = 0x0;
-		Result.Word1 = ObjectTypesToQuery;
-		Result.Word2 = 0;
-		Result.SetMaskChannelAndFlags(IgnoreMask, MultiTraceValue, FilterFlags);
-		return Result;
-	}
-
-	FQueryFilterData FQueryFilterBuilder::CreateLegacyTraceFilter(const uint8 ChannelIndex, const uint32 BlockChannelMask, const uint32 OverlapChannelMask, const uint8 IgnoreMask, const EFilterFlags FilterFlags)
-	{
-		FQueryFilterData Result;
-		Result.Word0 = 0x1;
-		Result.Word1 = BlockChannelMask;
-		Result.Word2 = OverlapChannelMask;
-		Result.SetMaskChannelAndFlags(IgnoreMask, ChannelIndex, FilterFlags);
-		return Result;
-	}
-
 	FQueryFilterData FQueryFilterBuilder::BuildFromLegacyQueryFilter(const FCollisionFilterData& QueryFilterData)
 	{
 		FQueryFilterData Result;
diff --git a/Engine/Source/Runtime/Experimental/Chaos/Public/Chaos/Collision/SimRaycast.h b/Engine/Source/Runtime/Experimental/Chaos/Public/Chaos/Collision/SimRaycast.h
--- a/Engine/Source/Runtime/Experimental/Chaos/Public/Chaos/Collision/SimRaycast.h
+++ b/Engine/Source/Runtime/Experimental/Chaos/Public/Chaos/Collision/SimRaycast.h
@@ -44,7 +44,7 @@ namespace Chaos::Private
 	private:
 		CHAOS_API bool ApplyFilter(const FGeometryParticleHandle* InParticleHandle) const;
 
-		FCollisionFilterData FilterData;
+		Filter::FQueryFilterData FilterData;
 	};
 
 	/*
@@ -74,7 +74,7 @@ namespace Chaos::Private
 			uint8 Bits;
 		};
 
-		FCollisionFilterData FilterData;
+		Filter::FQueryFilterData FilterData;
 		FFlags Flags;
 	};
 
diff --git a/Engine/Source/Runtime/Experimental/Chaos/Public/Chaos/CollisionFilterData.h b/Engine/Source/Runtime/Experimental/Chaos/Public/Chaos/CollisionFilterData.h
--- a/Engine/Source/Runtime/Experimental/Chaos/Public/Chaos/CollisionFilterData.h
+++ b/Engine/Source/Runtime/Experimental/Chaos/Public/Chaos/CollisionFilterData.h
@@ -286,9 +286,6 @@ namespace Chaos::Filter
 		CHAOS_API static FShapeFilterData BuildBlockAll(const EFilterFlags InFilterFlags = EFilterFlags::None);
 		CHAOS_API static FShapeFilterData BuildSimBlockAll(const EFilterFlags InFilterFlags = EFilterFlags::None);
 
-		UE_INTERNAL CHAOS_API static FShapeFilterData BuildLegacyShapeFilter(const uint8 InChannelIndex, const uint32 InQueryOverlapChannels, const uint32 InQueryBlockChannels, const uint32 InSimBlockChannels, const uint8 InMaskFilter = 0, const EFilterFlags InFilterFlags = EFilterFlags::None);
-		UE_INTERNAL CHAOS_API static FShapeFilterData BuildLegacyBlockAllSimShapeFilter(const uint8 InChannelIndex = 0, const uint8 InMaskFilter = 0, const EFilterFlags InFilterFlags = EFilterFlags::All);
-
 		UE_INTERNAL CHAOS_API static FShapeFilterData BuildFromLegacyQueryFilter(const FCollisionFilterData& QueryFilter);
 		UE_INTERNAL CHAOS_API static FShapeFilterData BuildFromLegacySimFilter(const FCollisionFilterData& SimFilter);
 		UE_INTERNAL CHAOS_API static FCombinedShapeFilterData BuildFromLegacyShapeFilter(const FCollisionFilterData& QueryFilterData, const FCollisionFilterData& SimFilterData);
@@ -312,9 +309,6 @@ namespace Chaos::Filter
 	class UE_INTERNAL FQueryFilterBuilder
 	{
 	public:
-		UE_INTERNAL CHAOS_API static FQueryFilterData CreateLegacyObjectTypeFilter(const uint32 ObjectTypesToQuery, const bool bMultiTrace, const uint8 IgnoreMask = 0, const EFilterFlags FilterFlags = EFilterFlags::None);
-		UE_INTERNAL CHAOS_API static FQueryFilterData CreateLegacyTraceFilter(const uint8 ChannelIndex, const uint32 BlockChannelMask, const uint32 OverlapChannelMask, const uint8 IgnoreMask = 0, const EFilterFlags FilterFlags = EFilterFlags::None);
-
 		UE_INTERNAL CHAOS_API static FQueryFilterData BuildFromLegacyQueryFilter(const FCollisionFilterData& QueryFilterData);
 		UE_INTERNAL CHAOS_API static FCollisionFilterData GetLegacyQueryFilter(const FQueryFilterData& QueryFilterData);
 	private:
