iax: separate out a block of enum-to-string code

#rnx

diff --git a/Engine/Source/Runtime/Experimental/IoStore/HttpClient/Private/Api.inl b/Engine/Source/Runtime/Experimental/IoStore/HttpClient/Private/Api.inl
--- a/Engine/Source/Runtime/Experimental/IoStore/HttpClient/Private/Api.inl
+++ b/Engine/Source/Runtime/Experimental/IoStore/HttpClient/Private/Api.inl
@@ -7,6 +7,22 @@ namespace UE::IoStore::HTTP
 
 // {{{1 request ................................................................
 
+////////////////////////////////////////////////////////////////////////////////
+static FAnsiStringView MimeTypeToString(EMimeType MimeType)
+{
+	switch (MimeType)
+	{
+	case EMimeType::Text:				return "text/html";
+	case EMimeType::Binary:				return "application/octet-stream";
+	case EMimeType::Json:				return "application/json";
+	case EMimeType::Xml:				return "application/xml";
+	case EMimeType::CbObject:			return "application/x-ue-cb";
+	case EMimeType::CbPackage:			return "application/x-ue-pkg";
+	case EMimeType::CompressedBuffer:	return "application/x-ue-comp";
+	}
+	return {};
+}
+
 ////////////////////////////////////////////////////////////////////////////////
 FRequest::FRequest(FRequest&& Rhs)
 {
@@ -34,22 +50,17 @@ FRequest::~FRequest()
 ////////////////////////////////////////////////////////////////////////////////
 void FRequest::Accept(EMimeType MimeType)
 {
-	switch (MimeType)
-	{
-	case EMimeType::Text:				return Header("Accept", "text/html");
-	case EMimeType::Binary:				return Header("Accept", "application/octet-stream");
-	case EMimeType::Json:				return Header("Accept", "application/json");
-	case EMimeType::Xml:				return Header("Accept", "application/xml");
-	case EMimeType::CbObject:			return Header("Accept", "application/x-ue-cb");
-	case EMimeType::CbPackage:			return Header("Accept", "application/x-ue-pkg");
-	case EMimeType::CompressedBuffer:	return Header("Accept", "application/x-ue-comp");
-	}
+	FAnsiStringView AsString = MimeTypeToString(MimeType);
+	Accept(AsString);
 }
 
 ////////////////////////////////////////////////////////////////////////////////
 void FRequest::Accept(FAnsiStringView MimeType)
 {
-	Header("Accept", MimeType);
+	if (!MimeType.IsEmpty())
+	{
+		Header("Accept", MimeType);
+	}
 }
 
 ////////////////////////////////////////////////////////////////////////////////
