iax: surface address from a resolved pool.

#rnx

diff --git a/Engine/Source/Runtime/Experimental/IoStore/HttpClient/Internal/IO/Http/Client.h b/Engine/Source/Runtime/Experimental/IoStore/HttpClient/Internal/IO/Http/Client.h
--- a/Engine/Source/Runtime/Experimental/IoStore/HttpClient/Internal/IO/Http/Client.h
+++ b/Engine/Source/Runtime/Experimental/IoStore/HttpClient/Internal/IO/Http/Client.h
@@ -109,6 +109,7 @@ public:
 							FConnectionPool(FConnectionPool&& Rhs)	{ *this = MoveTemp(Rhs); }
 	FConnectionPool&		operator = (FConnectionPool&& Rhs)		{ Swap(Ptr, Rhs.Ptr); return *this; }
 	UE_API bool				Resolve();
+	UE_API uint32			GetIps(uint32* Out, uint32 Count) const;
 	UE_API void				Describe(FAnsiStringBuilderBase&) const;
 
 	static UE_API bool		IsValidHostUrl(FAnsiStringView Url);
diff --git a/Engine/Source/Runtime/Experimental/IoStore/HttpClient/Private/ConnectionPool.inl b/Engine/Source/Runtime/Experimental/IoStore/HttpClient/Private/ConnectionPool.inl
--- a/Engine/Source/Runtime/Experimental/IoStore/HttpClient/Private/ConnectionPool.inl
+++ b/Engine/Source/Runtime/Experimental/IoStore/HttpClient/Private/ConnectionPool.inl
@@ -286,9 +286,32 @@ bool FConnectionPool::Resolve()
 	return Ptr->ResolveHostName().IsOk();
 }
 
+////////////////////////////////////////////////////////////////////////////////
+uint32 FConnectionPool::GetIps(uint32* Out, uint32 Count) const
+{
+	if (Ptr == nullptr)
+	{
+		return 0;
+	}
+
+	uint32 IpAddress = Ptr->GetIpAddress();
+	if (IpAddress == 0)
+	{
+		return 0;
+	}
+
+	*Out = IpAddress;
+	return 1;
+}
+
 ////////////////////////////////////////////////////////////////////////////////
 void FConnectionPool::Describe(FAnsiStringBuilderBase& OutString) const
 {
+	if (Ptr == nullptr)
+	{
+		return;
+	}
+
 	const FAnsiStringView HostName = Ptr->GetHostName();
 	OutString.Appendf("%.*s", HostName.Len(), HostName.GetData());
 	if (!!Ptr->IsResolved())
