iax: allow streaming of known-sized content

#rnx

diff --git a/Engine/Source/Runtime/Experimental/IoStore/HttpClient/Private/Activity.inl b/Engine/Source/Runtime/Experimental/IoStore/HttpClient/Private/Activity.inl
--- a/Engine/Source/Runtime/Experimental/IoStore/HttpClient/Private/Activity.inl
+++ b/Engine/Source/Runtime/Experimental/IoStore/HttpClient/Private/Activity.inl
@@ -526,10 +526,9 @@ FOutcome FActivity::RecvMessage(FHttpPeer& Peer)
 	}
 	else if (!bChunked && DestSize < ContentLength)
 	{
-		// todo: support piece-wise transfer of content (a la chunked).
-		Outcome = FOutcome::Error("Destination buffer too small");
-		Fail(Outcome);
-		return Outcome;
+		// Destination is too small for all the content. The user did however
+		// provide a size so we'll use that to deliver content piecewise.
+		bChunked = true;
 	}
 	else if (enum { MinStreamBuf = 256 }; bChunked && DestSize < MinStreamBuf)
 	{
