iax: latency test would skip every other probe on some CDN providers.

#rnx

diff --git a/Engine/Source/Runtime/Experimental/IoStore/HttpClient/Private/TransactionTwo.inl b/Engine/Source/Runtime/Experimental/IoStore/HttpClient/Private/TransactionTwo.inl
--- a/Engine/Source/Runtime/Experimental/IoStore/HttpClient/Private/TransactionTwo.inl
+++ b/Engine/Source/Runtime/Experimental/IoStore/HttpClient/Private/TransactionTwo.inl
@@ -447,8 +447,15 @@ FOutcome FDriverNg::TickInner(FTlsPeer& Peer)
 	RecvRemaining = RecvFrame.GetSize();
 	if (RecvFrame.IsPertinent())
 	{
-		uint32 Id = RecvFrame.GetId();
-		return FOutcome::Ok(Id);
+		// We'll not expect higher levels to come to read zero-sized data.
+		if (RecvFrame.Type != FFrame::EType::Data || RecvRemaining)
+		{
+			uint32 Id = RecvFrame.GetId();
+			return FOutcome::Ok(Id);
+		}
+
+		check(RecvFrame.Flags & FFrame::Flag_EndStream);
+		RecvFrame = {};
 	}
 
 	if (RecvRemaining == 0)
