iax: removed an unused cvar.

#rnx

diff --git a/Engine/Source/Runtime/Experimental/IoStore/HttpClient/Private/Misc.inl b/Engine/Source/Runtime/Experimental/IoStore/HttpClient/Private/Misc.inl
--- a/Engine/Source/Runtime/Experimental/IoStore/HttpClient/Private/Misc.inl
+++ b/Engine/Source/Runtime/Experimental/IoStore/HttpClient/Private/Misc.inl
@@ -65,8 +65,7 @@ IOSTOREHTTPCLIENT_API const void*	GetIaxTraceChannel()	{ return nullptr; }
 	)
 
 ////////////////////////////////////////////////////////////////////////////////
-static IAS_CVAR(int32, RecvWorkThresholdKiB,80,		"Threshold of data remaining at which next request is sent (in KiB)");
-static IAS_CVAR(int32, IdleMs,				50'000,	"Time in milliseconds to close idle connections or fail waits");
+static IAS_CVAR(int32, IdleMs, 50'000, "Time in milliseconds to close idle connections or fail waits");
 
 
 
