Added functions for dataflow

diff --git a/Engine/Source/Runtime/Experimental/Dataflow/Engine/Private/Dataflow/DataflowBlueprintLibrary.cpp b/Engine/Source/Runtime/Experimental/Dataflow/Engine/Private/Dataflow/DataflowBlueprintLibrary.cpp
--- a/Engine/Source/Runtime/Experimental/Dataflow/Engine/Private/Dataflow/DataflowBlueprintLibrary.cpp
+++ b/Engine/Source/Runtime/Experimental/Dataflow/Engine/Private/Dataflow/DataflowBlueprintLibrary.cpp
@@ -6,6 +6,7 @@
 #include "Dataflow/DataflowInstance.h"
 #include "StructUtils/PropertyBag.h"
 
+
 #include UE_INLINE_GENERATED_CPP_BY_NAME(DataflowBlueprintLibrary)
 
 void UDataflowBlueprintLibrary::EvaluateTerminalNodeByName(UDataflow* Dataflow, FName TerminalNodeName, UObject* ResultAsset)
@@ -120,9 +121,104 @@ bool UDataflowBlueprintLibrary::OverrideDataflowVariableObjectArray(UObject* Ass
 
 
 
-
-
-
-
-
-
+TArray<FDataflowVariable>  UDataflowBlueprintLibrary::GetDataflowVariableList(UObject* Asset)
+{
+	TArray<FDataflowVariable> PropArray;
+	if (IDataflowInstanceInterface* DataflowInterface = Cast<IDataflowInstanceInterface>(Asset))
+	{
+		FInstancedPropertyBag& AssetBag = DataflowInterface->GetDataflowInstance().GetVariableOverrides().GetVariables();
+		TConstArrayView<FPropertyBagPropertyDesc> Descs = AssetBag.GetPropertyBagStruct()->GetPropertyDescs();
+		for (FPropertyBagPropertyDesc desc : Descs)
+		{
+			FDataflowVariable NewPropVar;
+			NewPropVar.VariableName = desc.Name;
+			NewPropVar.Type = desc.ValueType;
+			NewPropVar.ContainerType = desc.ContainerTypes.GetFirstContainerType();
+			if (desc.ValueType == EPropertyBagPropertyType::Int64)
+			{
+				const TValueOrError<int64, EPropertyBagResult> res = AssetBag.GetValueInt64(desc.Name);
+				if (res.HasValue())
+				{
+					NewPropVar.Value = FString::FromInt(res.GetValue());
+					PropArray.Add(NewPropVar);
+				}
+				else
+				{
+					EPropertyBagResult Err = res.GetError();
+					UE_LOG(LogTemp, Warning, TEXT("Could not read int"), (int32)Err);
+				}
+			}
+			else if (desc.ValueType == EPropertyBagPropertyType::Int32)
+			{
+				const TValueOrError<int32, EPropertyBagResult> res = AssetBag.GetValueInt32(desc.Name);
+				if (res.HasValue())
+				{
+					NewPropVar.Value = FString::FromInt(res.GetValue());
+					PropArray.Add(NewPropVar);
+				}
+				else
+				{
+					EPropertyBagResult Err = res.GetError();
+					UE_LOG(LogTemp, Warning, TEXT("Could not read int"), (int32)Err);
+				}
+			}
+			else if (desc.ValueType == EPropertyBagPropertyType::Float)
+			{
+				const TValueOrError<float, EPropertyBagResult> res = AssetBag.GetValueFloat(desc.Name);
+				if (res.HasValue())
+				{
+					NewPropVar.Value = LexToString(res.GetValue());
+					PropArray.Add(NewPropVar);
+				}
+				else
+				{
+					EPropertyBagResult Err = res.GetError();
+					UE_LOG(LogTemp, Warning, TEXT("Could not read float"), (int32)Err);
+				}
+			}
+			else if (desc.ValueType == EPropertyBagPropertyType::Double)
+			{
+				const TValueOrError<double, EPropertyBagResult> res = AssetBag.GetValueDouble(desc.Name);
+				if (res.HasValue())
+				{
+					NewPropVar.Value = LexToString(res.GetValue());
+					PropArray.Add(NewPropVar);
+				}
+				else
+				{
+					EPropertyBagResult Err = res.GetError();
+					UE_LOG(LogTemp, Warning, TEXT("Could not read double"), (int32)Err);
+				}
+			}
+			else if (desc.ValueType == EPropertyBagPropertyType::Object)
+			{
+				const TValueOrError<UObject*, EPropertyBagResult> res = AssetBag.GetValueObject(desc.Name, UObject::StaticClass());
+				if (res.HasValue())
+				{
+					NewPropVar.Value = res.GetValue()->GetOutermost()->GetName();
+					PropArray.Add(NewPropVar);
+				}
+				else
+				{
+					EPropertyBagResult Err = res.GetError();
+					UE_LOG(LogTemp, Warning, TEXT("Could not read Object"), (int32)Err);
+				}
+			}
+			else if (desc.ValueType == EPropertyBagPropertyType::Bool)
+			{
+				const TValueOrError<bool, EPropertyBagResult> res = AssetBag.GetValueBool(desc.Name);
+				if (res.HasValue())
+				{
+					NewPropVar.Value = LexToString(res.GetValue());
+					PropArray.Add(NewPropVar);
+				}
+				else
+				{
+					EPropertyBagResult Err = res.GetError();
+					UE_LOG(LogTemp, Warning, TEXT("Could not read bool"), (int32)Err);
+				}
+			}
+		}
+	}
+	return PropArray;
+}
\ No newline at end of file
diff --git a/Engine/Source/Runtime/Experimental/Dataflow/Engine/Public/Dataflow/DataflowBlueprintLibrary.h b/Engine/Source/Runtime/Experimental/Dataflow/Engine/Public/Dataflow/DataflowBlueprintLibrary.h
--- a/Engine/Source/Runtime/Experimental/Dataflow/Engine/Public/Dataflow/DataflowBlueprintLibrary.h
+++ b/Engine/Source/Runtime/Experimental/Dataflow/Engine/Public/Dataflow/DataflowBlueprintLibrary.h
@@ -7,6 +7,30 @@
 
 #include "DataflowBlueprintLibrary.generated.h"
 
+
+
+USTRUCT(BlueprintType)
+struct FDataflowVariable
+{
+	GENERATED_BODY()
+
+public:
+
+	UPROPERTY(EditAnywhere, BlueprintReadWrite, Category = "Dataflow")
+	FName VariableName;
+
+	UPROPERTY(EditAnywhere, BlueprintReadWrite, Category = "Dataflow")
+	EPropertyBagPropertyType Type;
+
+	UPROPERTY(EditAnywhere, BlueprintReadWrite, Category = "Dataflow")
+	EPropertyBagContainerType ContainerType;
+
+	UPROPERTY(EditAnywhere, BlueprintReadWrite, Category = "Dataflow")
+	FString Value;
+
+};
+
+
 UCLASS(MinimalAPI)
 class UDataflowBlueprintLibrary : public UBlueprintFunctionLibrary
 {
@@ -81,4 +105,7 @@ public:
 	*/
 	UFUNCTION(BlueprintCallable, Category = "Dataflow", meta = (Keywords = "Dataflow graph"))
 	static DATAFLOWENGINE_API bool OverrideDataflowVariableObjectArray(UObject* Asset, FName VariableName, const TArray<UObject*>& VariableArrayValue);
+
+	UFUNCTION(BlueprintCallable, Category = "Dataflow", meta = (Keywords = "Dataflow graph"))
+	static DATAFLOWENGINE_API TArray<FDataflowVariable> GetDataflowVariableList(UObject* Asset);
 };
diff --git a/Engine/Source/Runtime/Experimental/GeometryCollectionEngine/Private/GeometryCollection/GeometryCollectionBlueprintLibrary.cpp b/Engine/Source/Runtime/Experimental/GeometryCollectionEngine/Private/GeometryCollection/GeometryCollectionBlueprintLibrary.cpp
--- a/Engine/Source/Runtime/Experimental/GeometryCollectionEngine/Private/GeometryCollection/GeometryCollectionBlueprintLibrary.cpp
+++ b/Engine/Source/Runtime/Experimental/GeometryCollectionEngine/Private/GeometryCollection/GeometryCollectionBlueprintLibrary.cpp
@@ -4,7 +4,11 @@
 
 #include "GeometryCollection/GeometryCollectionComponent.h"
 #include "GeometryCollection/GeometryCollectionExternalRenderInterface.h"
-
+#if WITH_EDITOR
+#include "Engine/SCS_Node.h"
+#include "Editor/EditorEngine.h"
+#include "Engine/SimpleConstructionScript.h"
+#endif
 #include UE_INLINE_GENERATED_CPP_BY_NAME(GeometryCollectionBlueprintLibrary)
 
 void UGeometryCollectionBlueprintLibrary::SetCustomInstanceDataByIndex(UGeometryCollectionComponent* GeometryCollectionComponent, int32 CustomDataIndex, float CustomDataValue)
@@ -43,3 +47,26 @@ void UGeometryCollectionBlueprintLibrary::SetISMPoolCustomInstanceData(UGeometry
 {
 	SetCustomInstanceDataByIndex(GeometryCollectionComponent, CustomDataIndex, CustomDataValue);
 }
+
+
+bool UGeometryCollectionBlueprintLibrary::AddGeometryCollectionsToBlueprint(UBlueprint* Blueprint, const TArray<UObject*>& Collections)
+{
+#if WITH_EDITOR
+	for (UObject* Collection : Collections)
+	{
+		USCS_Node* NewNode = Blueprint->SimpleConstructionScript->CreateNode(Collection->GetClass(), Collection->GetFName());
+		TArray<USCS_Node*> AllNodes = Blueprint->SimpleConstructionScript->GetAllNodes();
+		if (AllNodes.Num() == 0)
+		{
+			Blueprint->SimpleConstructionScript->AddNode(NewNode);
+		}
+		else 
+		{
+			AllNodes[0]->AddChildNode(NewNode);
+		}
+		UEditorEngine::CopyPropertiesForUnrelatedObjects(Collection, NewNode->ComponentTemplate);
+	}
+	Blueprint->Modify();
+#endif
+	return true;
+}
\ No newline at end of file
diff --git a/Engine/Source/Runtime/Experimental/GeometryCollectionEngine/Public/GeometryCollection/GeometryCollectionBlueprintLibrary.h b/Engine/Source/Runtime/Experimental/GeometryCollectionEngine/Public/GeometryCollection/GeometryCollectionBlueprintLibrary.h
--- a/Engine/Source/Runtime/Experimental/GeometryCollectionEngine/Public/GeometryCollection/GeometryCollectionBlueprintLibrary.h
+++ b/Engine/Source/Runtime/Experimental/GeometryCollectionEngine/Public/GeometryCollection/GeometryCollectionBlueprintLibrary.h
@@ -38,6 +38,10 @@ public:
 	UE_DEPRECATED(5.5, "Use SetCustomInstanceDataByIndex() instead")
 	UFUNCTION(BlueprintCallable, Category = "ChaosPhysics", meta = (DeprecatedFunction, DeprecationMessage = "Use SetCustomInstanceDataByIndex() instead"))
 	static UE_API void SetISMPoolCustomInstanceData(UGeometryCollectionComponent* GeometryCollectionComponent, int32 CustomDataIndex, float CustomDataValue);
+
+	UFUNCTION(BlueprintCallable, Category = "ChaosPhysics")
+	static UE_API bool AddGeometryCollectionsToBlueprint(UBlueprint* Blueprint, const TArray<UObject*>& Collections);
 };
 
 #undef UE_API
+ 
\ No newline at end of file
