Geometry collection : make set collision response call to work on geometry collection
* make sure we fall back to applying the component level collision settings from the body instance

diff --git a/Engine/Source/Runtime/Experimental/GeometryCollectionEngine/Private/GeometryCollection/GeometryCollectionComponent.cpp b/Engine/Source/Runtime/Experimental/GeometryCollectionEngine/Private/GeometryCollection/GeometryCollectionComponent.cpp
--- a/Engine/Source/Runtime/Experimental/GeometryCollectionEngine/Private/GeometryCollection/GeometryCollectionComponent.cpp
+++ b/Engine/Source/Runtime/Experimental/GeometryCollectionEngine/Private/GeometryCollection/GeometryCollectionComponent.cpp
@@ -4871,6 +4871,10 @@ void UGeometryCollectionComponent::LoadCollisionProfiles()
 			{
 				Data = AbandonedData;
 			}
+			else // last resort use the body instance set collision settings
+			{
+				Data = CreateOrGetCollisionProfileData(DefaultCollisionProfileName);
+			}
 
 			if (Data)
 			{
