Chaos: Fix issue where midphase data is not properly stored in MidPhaseModifierAccessor

diff --git a/Engine/Source/Runtime/Experimental/Chaos/Private/Chaos/PBDRigidsEvolutionGBF.cpp b/Engine/Source/Runtime/Experimental/Chaos/Private/Chaos/PBDRigidsEvolutionGBF.cpp
--- a/Engine/Source/Runtime/Experimental/Chaos/Private/Chaos/PBDRigidsEvolutionGBF.cpp
+++ b/Engine/Source/Runtime/Experimental/Chaos/Private/Chaos/PBDRigidsEvolutionGBF.cpp
@@ -1297,6 +1297,10 @@ void FPBDRigidsEvolutionGBF::ApplyMidPhaseModifier(FReal Dt)
 {
 	if (MidPhaseModifiers != nullptr)
 	{
+		// Consolidate midphases from parallel threads into the main lists held by the allocator.
+		// This makes them accessible to MidPhaseModifiers.
+		CollisionConstraints.GetConstraintAllocator().ProcessNewMidPhases();
+
 		FMidPhaseModifierAccessor MidPhaseModifier(*this);
 		for (ISimCallbackObject* ModifierCallback : *MidPhaseModifiers)
 		{
