Geometry collection : add regsitry tags ( dataflow + Gc )
- Dataflow registry tags can be set from the Dataflow instance util function library

diff --git a/Engine/Source/Runtime/Experimental/Dataflow/Engine/Private/Dataflow/DataflowInstance.cpp b/Engine/Source/Runtime/Experimental/Dataflow/Engine/Private/Dataflow/DataflowInstance.cpp
--- a/Engine/Source/Runtime/Experimental/Dataflow/Engine/Private/Dataflow/DataflowInstance.cpp
+++ b/Engine/Source/Runtime/Experimental/Dataflow/Engine/Private/Dataflow/DataflowInstance.cpp
@@ -3,6 +3,7 @@
 #include "Dataflow/DataflowInstance.h"
 #include "Dataflow/DataflowObject.h"
 #include "Dataflow/DataflowObjectInterface.h"
+#include "UObject/AssetRegistryTagsContext.h"
 
 #include UE_INLINE_GENERATED_CPP_BY_NAME(DataflowInstance)
 
@@ -109,6 +110,18 @@ namespace UE::Dataflow::InstanceUtils
 		}
 		return TerminalNodeNames;
 	}
+
+	void GetAssetRegistryTags(const UObject* BoundAsset, FAssetRegistryTagsContext& Context)
+	{
+		const UDataflow* DataflowObject = UE::Dataflow::InstanceUtils::GetDataflowAssetFromObject(BoundAsset);
+
+		const bool bHasDataflow = (DataflowObject != nullptr);
+		Context.AddTag(UObject::FAssetRegistryTag(TEXT("HasDataflow"), bHasDataflow ? TEXT("True") : TEXT("False"), UObject::FAssetRegistryTag::TT_Alphabetical));
+
+		const bool bIsDataflowEmbedded = DataflowObject ? (!DataflowObject->IsAsset()) : false;
+		Context.AddTag(UObject::FAssetRegistryTag(TEXT("IsDataflowEmbedded"), bIsDataflowEmbedded ? TEXT("True") : TEXT("False"), UObject::FAssetRegistryTag::TT_Alphabetical));
+	}
+
 }
 
 
diff --git a/Engine/Source/Runtime/Experimental/Dataflow/Engine/Public/Dataflow/DataflowInstance.h b/Engine/Source/Runtime/Experimental/Dataflow/Engine/Public/Dataflow/DataflowInstance.h
--- a/Engine/Source/Runtime/Experimental/Dataflow/Engine/Public/Dataflow/DataflowInstance.h
+++ b/Engine/Source/Runtime/Experimental/Dataflow/Engine/Public/Dataflow/DataflowInstance.h
@@ -8,6 +8,8 @@
 
 #include "DataflowInstance.generated.h"
 
+class FAssetRegistryTagsContext;
+
 namespace UE::Dataflow::Private
 {
 	class FVariablesOverridesDetails;
@@ -309,4 +311,10 @@ namespace UE::Dataflow::InstanceUtils
 
 	/** get the list of terminal node names for a specific data flow asset */
 	DATAFLOWENGINE_API TArray<FName> GetTerminalNodeNames(const UDataflow* DataflowAsset);
+
+	// Register tags for an asset using Dataflow
+	DATAFLOWENGINE_API void GetAssetRegistryTags(const UObject* BoundAsset, FAssetRegistryTagsContext& Context);
+
+
+	//////////////////////////////////////////////////////////////////////////////////////////////////////////////////////////////////////
 }
diff --git a/Engine/Source/Runtime/Experimental/GeometryCollectionEngine/Private/GeometryCollection/GeometryCollectionObject.cpp b/Engine/Source/Runtime/Experimental/GeometryCollectionEngine/Private/GeometryCollection/GeometryCollectionObject.cpp
--- a/Engine/Source/Runtime/Experimental/GeometryCollectionEngine/Private/GeometryCollection/GeometryCollectionObject.cpp
+++ b/Engine/Source/Runtime/Experimental/GeometryCollectionEngine/Private/GeometryCollection/GeometryCollectionObject.cpp
@@ -2334,3 +2334,15 @@ bool UGeometryCollection::CanEditChange(const FProperty* InProperty) const
 	return true;
 }
 #endif
+
+void UGeometryCollection::GetAssetRegistryTags(FAssetRegistryTagsContext Context) const
+{
+	// Dataflow specific tags
+	UE::Dataflow::InstanceUtils::GetAssetRegistryTags(this, Context);
+
+	// Geometry collection specific tags
+	const int32 NumTransforms = GeometryCollection? GeometryCollection->Transform.Num(): 0;
+	Context.AddTag(FAssetRegistryTag(TEXT("NumTransforms"), FString::FromInt(NumTransforms), FAssetRegistryTag::TT_Numerical));
+
+	Super::GetAssetRegistryTags(Context);
+}
\ No newline at end of file
diff --git a/Engine/Source/Runtime/Experimental/GeometryCollectionEngine/Public/GeometryCollection/GeometryCollectionObject.h b/Engine/Source/Runtime/Experimental/GeometryCollectionEngine/Public/GeometryCollection/GeometryCollectionObject.h
--- a/Engine/Source/Runtime/Experimental/GeometryCollectionEngine/Public/GeometryCollection/GeometryCollectionObject.h
+++ b/Engine/Source/Runtime/Experimental/GeometryCollectionEngine/Public/GeometryCollection/GeometryCollectionObject.h
@@ -1001,4 +1001,6 @@ private:
 
 	void MigrateDeprecatedRootProxyData();
 	void MigrateDeprecatedDataflowData();
+
+	virtual void GetAssetRegistryTags(FAssetRegistryTagsContext Context) const override;
 };
