Show infinity symbol when IAX is running with no data rate limit.

#rnx

diff --git a/Engine/Source/Runtime/Experimental/IoStore/OnDemand/Private/Statistics.cpp b/Engine/Source/Runtime/Experimental/IoStore/OnDemand/Private/Statistics.cpp
--- a/Engine/Source/Runtime/Experimental/IoStore/OnDemand/Private/Statistics.cpp
+++ b/Engine/Source/Runtime/Experimental/IoStore/OnDemand/Private/Statistics.cpp
@@ -866,7 +866,15 @@ void FOnDemandIoBackendStats::PrintOnScreenStatistics(FCoreDelegates::FSeverityM
 				}
 			}
 
-			const uint32 RateLimit = (GIasHttpRateLimitKiBPerSecond < 0) ? 0 : GIasHttpRateLimitKiBPerSecond;
+			TStringBuilder<32> RateLimit;
+			if (GIasHttpRateLimitKiBPerSecond > 0)
+			{
+				RateLimit << GIasHttpRateLimitKiBPerSecond;
+			}
+			else
+			{
+				RateLimit << TEXT("\u221E"); // Infinity symbol code
+			}
 
 			TStringBuilder<256> Builder;
 			Builder << TEXT("IAX Http:");
