[Chaos]
Replaces existing FTransactionallySafeRWLock with FTransactionallySafeSharedRecursiveMutex to address hang that can occur in EventManager due to recursion
Provides a cvar for returning to old behavior using a TVariant setup that will change between the two types.  Puts container work in lambdas to avoid duplicating code or having to force a lock/immediate unlock when dealing with ResourceLock

diff --git a/Engine/Source/Runtime/Experimental/Chaos/Private/EventManager.cpp b/Engine/Source/Runtime/Experimental/Chaos/Private/EventManager.cpp
--- a/Engine/Source/Runtime/Experimental/Chaos/Private/EventManager.cpp
+++ b/Engine/Source/Runtime/Experimental/Chaos/Private/EventManager.cpp
@@ -9,112 +9,139 @@ namespace Chaos
 		// This is required to be able to track continuous collisions (so that we can detect the "no collision" state)
 		bool bEventManagerDispatchEmptyEvents = true;
 		FAutoConsoleVariableRef CVarEventManagerDispatchEmptyEvent(TEXT("p.Chaos.EventManager.DispatchEmptyEvents"), bEventManagerDispatchEmptyEvents, TEXT("Whether to dispatch events even if the event buffers are empty"));
+
+		bool bUseTransactionallySafeSharedRecursiveMutex = true;
+		FAutoConsoleVariableRef CVarUseUseTransactionallySafeSharedRecursiveMutex(TEXT("Chaos.UseTransactionallySafeSharedMutex"), bUseTransactionallySafeSharedRecursiveMutex, TEXT("Sets the Event Manager to use TransactionallySafeSharedRecursiveMutexes to avoid a pontential recursive hang."));
+	}
+
+
+	FEventManager::FEventManager(const Chaos::EMultiBufferMode& BufferModeIn)
+		: BufferMode(BufferModeIn)
+	{
+		if (Private::bUseTransactionallySafeSharedRecursiveMutex)
+		{
+			ContainerLock.Emplace<UE::FTransactionallySafeSharedRecursiveMutex>();
+			ResourceLock.Emplace<UE::FTransactionallySafeSharedRecursiveMutex>();
+		}
+		else
+		{
+			ContainerLock.Emplace<FTransactionallySafeRWLock>();
+			ResourceLock.Emplace<FTransactionallySafeRWLock>();
+		}
 	}
 
 	void FEventManager::Reset()
 	{
-		ContainerLock.WriteLock();
+		FLockWriteScope ContainerWriteScope(ContainerLock);
 		for (FEventContainerBasePtr Container : EventContainers)
 		{
 			delete Container;
 			Container = nullptr;
 		}
 		EventContainers.Reset();
-		ContainerLock.WriteUnlock();
 	}
 
 	void FEventManager::UnregisterEvent(const EEventType& EventType)
 	{
 		const FEventID EventID = (FEventID)EventType;
-		ContainerLock.WriteLock();
+		FLockWriteScope ContainerWriteScope(ContainerLock);
 		if (EventID < EventContainers.Num())
 		{
 			delete EventContainers[EventID];
 			EventContainers.RemoveAt(EventID);
 		}
-		ContainerLock.WriteUnlock();
 	}
 
 	void FEventManager::UnregisterHandler(const EEventType& EventType, const void* InHandler)
 	{
 		const FEventID EventID = (FEventID)EventType;
-		ContainerLock.ReadLock();
+		FLockReadScope ContainerReadScope(ContainerLock);
 		checkf(EventID < EventContainers.Num(), TEXT("Unregistering event Handler for an event ID that does not exist"));
 		EventContainers[EventID]->UnregisterHandler(InHandler);
-		ContainerLock.ReadUnlock();
 	}
 
 	void FEventManager::FillProducerData(const Chaos::FPBDRigidsSolver* Solver, bool bResetData)
 	{
+		auto ContainerLambda = [this, Solver, bResetData]()
+			{
+				FLockReadScope ContainerReadScope(ContainerLock);
+				for (FEventContainerBasePtr EventContainer : EventContainers)
+				{
+					if (EventContainer)
+					{
+						EventContainer->InjectProducerData(Solver, bResetData);
+					}
+				}
+			};
+
 		if (BufferMode == EMultiBufferMode::Double)
 		{
-			ResourceLock.ReadLock();
-		}
-		ContainerLock.ReadLock();
-		for (FEventContainerBasePtr EventContainer : EventContainers)
-		{
-			if (EventContainer)
-			{
-				EventContainer->InjectProducerData(Solver, bResetData);
-			}
+			FLockReadScope ResourceReadScope(ResourceLock);
+			ContainerLambda();
 		}
-		ContainerLock.ReadUnlock();
-		if (BufferMode == EMultiBufferMode::Double)
+
+		else
 		{
-			ResourceLock.ReadUnlock();
+			ContainerLambda();
 		}
 	}
 
 	void FEventManager::FlipBuffersIfRequired()
 	{
+		auto ContainerLambda = [this]()
+			{
+				FLockReadScope ContainerReadScope(ContainerLock);
+				for (FEventContainerBasePtr EventContainer : EventContainers)
+				{
+					if (EventContainer)
+					{
+						EventContainer->ResetConsumerBuffer();
+						EventContainer->FlipBufferIfRequired();
+					}
+				}
+			};
+
 		if (BufferMode == EMultiBufferMode::Single)
 		{
 			return;
 		}
 		else if (BufferMode == EMultiBufferMode::Double)
 		{
-			ResourceLock.WriteLock();
-		}
-
-		ContainerLock.ReadLock();
-		for (FEventContainerBasePtr EventContainer : EventContainers)
-		{
-			if (EventContainer)
-			{
-				EventContainer->ResetConsumerBuffer();
-				EventContainer->FlipBufferIfRequired();
-			}
+			FLockWriteScope ResourceWriteScope(ResourceLock);
+			ContainerLambda();
 		}
-		ContainerLock.ReadUnlock();
-
-		if (BufferMode == EMultiBufferMode::Double)
+		else
 		{
-			ResourceLock.WriteUnlock();
+			ContainerLambda();
 		}
 	}
 
 	void FEventManager::DispatchEvents()
 	{
+		auto ContainerLambda = [this]()
+			{
+				FLockReadScope ContainerReadScope(ContainerLock);
+				for (FEventContainerBasePtr EventContainer : EventContainers)
+				{
+					if (EventContainer)
+					{
+						EventContainer->DispatchConsumerData();
+					}
+				}
+			};
+
 		if (BufferMode == EMultiBufferMode::Double)
 		{
-			ResourceLock.ReadLock();
+			FLockReadScope ResourceReadScope(ResourceLock);
+			ContainerLambda();
 		}
-
-		ContainerLock.ReadLock();
-		for (FEventContainerBasePtr EventContainer : EventContainers)
+		else
 		{
-			if (EventContainer)
-			{
-				EventContainer->DispatchConsumerData();
-			}
+			ContainerLambda();
 		}
-		ContainerLock.ReadUnlock();
+		
 
-		if (BufferMode == EMultiBufferMode::Double)
-		{
-			ResourceLock.ReadUnlock();
-		} 
-		else if (BufferMode == EMultiBufferMode::Single)
+		if (BufferMode == EMultiBufferMode::Single)
 		{
 			for (FEventContainerBasePtr EventContainer : EventContainers)
 			{
diff --git a/Engine/Source/Runtime/Experimental/Chaos/Public/EventManager.h b/Engine/Source/Runtime/Experimental/Chaos/Public/EventManager.h
--- a/Engine/Source/Runtime/Experimental/Chaos/Public/EventManager.h
+++ b/Engine/Source/Runtime/Experimental/Chaos/Public/EventManager.h
@@ -3,6 +3,9 @@
 #pragma once
 
 #include "CoreMinimal.h"
+#include "Async/TransactionallySafeSharedRecursiveMutex.h"
+#include "Async/SharedLock.h"
+#include "Async/UniqueLock.h"
 #include "Chaos/Framework/MultiBufferResource.h"
 #include "Chaos/Framework/PhysicsProxyBase.h"
 #include "Chaos/PBDRigidsEvolutionFwd.h"
@@ -10,6 +13,8 @@
 #include "Chaos/Defines.h"
 #include "Misc/TransactionallySafeRWLock.h"
 #include "Misc/ScopeRWLock.h"
+#include "Misc/TVariant.h"
+
 
 namespace Chaos
 {
@@ -151,6 +156,7 @@ namespace Chaos
 	template<typename PayloadType>
 	class TEventContainer : public FEventContainerBase
 	{
+
 	public:
 		/**
 		 * Regular constructor
@@ -448,11 +454,46 @@ private:
 	{
 		friend class FPBDRigidsSolver;
 
+		struct FLockWriteScope
+		{
+			FLockWriteScope(const FLockWriteScope&) = delete;
+			FLockWriteScope& operator=(const FLockWriteScope&) = delete;
+			explicit FLockWriteScope(TVariant<FTransactionallySafeRWLock, UE::FTransactionallySafeSharedRecursiveMutex>& Mutex)
+			{
+				if (UE::FTransactionallySafeSharedRecursiveMutex* M = Mutex.TryGet<UE::FTransactionallySafeSharedRecursiveMutex>())
+				{
+					Lock.Emplace<UE::TUniqueLock<UE::FTransactionallySafeSharedRecursiveMutex>>(*M);
+				}
+				else
+				{
+					Lock.Emplace<UE::TWriteScopeLock<FTransactionallySafeRWLock>>(Mutex.Get<FTransactionallySafeRWLock>());
+				}
+			}
+			TVariant<TYPE_OF_NULLPTR, UE::TWriteScopeLock<FTransactionallySafeRWLock>, UE::TUniqueLock<UE::FTransactionallySafeSharedRecursiveMutex>> Lock;
+		};
+
+		struct FLockReadScope
+		{
+			FLockReadScope(const FLockReadScope&) = delete;
+			FLockReadScope& operator=(const FLockReadScope&) = delete;
+			explicit FLockReadScope(TVariant<FTransactionallySafeRWLock, UE::FTransactionallySafeSharedRecursiveMutex>& Mutex)
+			{
+				if (UE::FTransactionallySafeSharedRecursiveMutex* M = Mutex.TryGet<UE::FTransactionallySafeSharedRecursiveMutex>())
+				{
+					Lock.Emplace<UE::TSharedLock<UE::FTransactionallySafeSharedRecursiveMutex>>(*M);
+				}
+				else
+				{
+					Lock.Emplace<UE::TReadScopeLock<FTransactionallySafeRWLock>>(Mutex.Get<FTransactionallySafeRWLock>());
+				}
+			}
+			TVariant<TYPE_OF_NULLPTR, UE::TReadScopeLock<FTransactionallySafeRWLock>, UE::TSharedLock<UE::FTransactionallySafeSharedRecursiveMutex>> Lock;
+		};
+
 	public:
 
-		FEventManager(const Chaos::EMultiBufferMode& BufferModeIn) 
-			: BufferMode(BufferModeIn)
-			{}
+		CHAOS_API FEventManager(const Chaos::EMultiBufferMode& BufferModeIn);
+
 
 		~FEventManager()
 		{
@@ -478,9 +519,8 @@ private:
 		template<typename PayloadType>
 		void RegisterEvent(const EEventType& EventType, TFunction<void(const Chaos::FPBDRigidsSolver* Solver, PayloadType& EventData, bool bResetData)> InFunction)
 		{
-			ContainerLock.WriteLock();
+			FLockWriteScope ContainerWriteScope(ContainerLock);
 			InternalRegisterInjector(FEventID(EventType), new TEventContainer<PayloadType>(BufferMode, InFunction));
-			ContainerLock.WriteUnlock();
 		}
 
 		/**
@@ -489,13 +529,12 @@ private:
 		template<typename PayloadType>
 		void ClearEvents(const EEventType& EventType, TFunction<void(PayloadType & EventData)> InFunction)
 		{
-			ContainerLock.WriteLock();
+			FLockWriteScope ContainerWriteScope(ContainerLock);
 
 			if (TEventContainer<PayloadType>* EventContainer = StaticCast<TEventContainer<PayloadType>*>(EventContainers[FEventID(EventType)]))
 			{
 				EventContainer->DestroyStaleEvents(InFunction);
 			}
-			ContainerLock.WriteUnlock();
 		}
 
 		/**
@@ -510,10 +549,9 @@ private:
 		void RegisterHandler(const EEventType& EventType, HandlerType* Handler, typename TRawEventHandler<PayloadType, HandlerType>::FHandlerFunction HandlerFunction, typename TRawEventHandler<PayloadType, HandlerType>::FInterestedProxyOwnerFunction InterestedProxyOwnerFunction = nullptr)
 		{
 			const FEventID EventID = FEventID(EventType);
-			ContainerLock.ReadLock();
+			FLockReadScope ContainerReadScope(ContainerLock);
 			checkf(EventID < EventContainers.Num(), TEXT("Registering event Handler for an event ID that does not exist"));
 			EventContainers[EventID]->RegisterHandler(new TRawEventHandler<PayloadType, HandlerType>(Handler, HandlerFunction, InterestedProxyOwnerFunction));
-			ContainerLock.ReadUnlock();
 		}
 
 		/**
@@ -541,23 +579,26 @@ private:
 		/** Returns decoded collision index. */
 		static CHAOS_API int32 DecodeCollisionIndex(int32 EncodedCollisionIdx, bool& bSwapOrder);
 
-
 		template<typename PayloadType>
 		void AddEvent(const EEventType& EventType, TFunction<void(PayloadType& EventData)> InFunction)
 		{
+			auto ContainerLambda = [this, EventType, InFunction]()
+				{
+					FLockReadScope ContainerReadScope(ContainerLock);
+					if (TEventContainer<PayloadType>* EventContainer = StaticCast<TEventContainer<PayloadType>*>(EventContainers[FEventID(EventType)]))
+					{
+						EventContainer->AddEvent(InFunction);
+					}
+				};
+
 			if (BufferMode == EMultiBufferMode::Double)
 			{
-				ResourceLock.ReadLock();
+				FLockReadScope ResourceReadScope(ResourceLock);
+				ContainerLambda();
 			}
-			ContainerLock.ReadLock();
-			if (TEventContainer<PayloadType>* EventContainer = StaticCast<TEventContainer<PayloadType>*>(EventContainers[FEventID(EventType)]))
-			{
-				EventContainer->AddEvent(InFunction);
-			}
-			ContainerLock.ReadUnlock();
-			if (BufferMode == EMultiBufferMode::Double)
+			else
 			{
-				ResourceLock.ReadUnlock();
+				ContainerLambda();
 			}
 		}
 
@@ -567,8 +608,11 @@ private:
 
 		Chaos::EMultiBufferMode BufferMode;			// specifies the buffer type to be constructed, single, double, triple
 		TArray<FEventContainerBasePtr> EventContainers;	// Array of event types
-		FTransactionallySafeRWLock ResourceLock;
-		FTransactionallySafeRWLock ContainerLock;
+		
+		TVariant<FTransactionallySafeRWLock, UE::FTransactionallySafeSharedRecursiveMutex> ResourceLock;
+
+		TVariant<FTransactionallySafeRWLock, UE::FTransactionallySafeSharedRecursiveMutex> ContainerLock;
+
 	};
 
 }
