Fix crash on unsubscribing potentially unloaded UStructs in CVD Sessions Manager

diff --git a/Engine/Source/Runtime/Experimental/ChaosSolverEngine/ChaosSolverEngine.Build.cs b/Engine/Source/Runtime/Experimental/ChaosSolverEngine/ChaosSolverEngine.Build.cs
--- a/Engine/Source/Runtime/Experimental/ChaosSolverEngine/ChaosSolverEngine.Build.cs
+++ b/Engine/Source/Runtime/Experimental/ChaosSolverEngine/ChaosSolverEngine.Build.cs
@@ -21,6 +21,13 @@ namespace UnrealBuildTool.Rules
 					"ChaosVDRuntime", // If CVD is compiled out, this only provides access to its structs data types
 				}
 				);
+
+			PrivateDependencyModuleNames.AddRange(
+				new string[]
+				{
+					"Messaging"
+				}
+				);
 			
 
 			SetupModulePhysicsSupport(Target);
diff --git a/Engine/Source/Runtime/Experimental/ChaosSolverEngine/Private/Chaos/ChaosVDEngineEditorBridge.cpp b/Engine/Source/Runtime/Experimental/ChaosSolverEngine/Private/Chaos/ChaosVDEngineEditorBridge.cpp
--- a/Engine/Source/Runtime/Experimental/ChaosSolverEngine/Private/Chaos/ChaosVDEngineEditorBridge.cpp
+++ b/Engine/Source/Runtime/Experimental/ChaosSolverEngine/Private/Chaos/ChaosVDEngineEditorBridge.cpp
@@ -334,10 +334,7 @@ void FChaosVDEngineEditorBridge::TearDown()
 	FTSTicker::RemoveTicker(RecordingStatusUpdateHandle);
 	RecordingStatusUpdateHandle.Reset();
 
-	// Note: This works during engine shutdown because the Module Manager doesn't free the dll on module unload to account for use cases like this
-	// If this appears in a callstack crash it means that assumption changed or was not correct to begin with.
-	// A possible solution is just check if the module is loaded querying the module manager just using the module's name
-	if (FChaosVDRuntimeModule::IsLoaded())
+	if (FModuleManager::Get().IsModuleLoaded(TEXT("ChaosVDRuntime")))
 	{
 
 		FChaosVDRuntimeModule& CVDRuntimeModule = FChaosVDRuntimeModule::Get();
diff --git a/Engine/Source/Runtime/Experimental/ChaosSolverEngine/Private/Chaos/ChaosVDRemoteSessionsManager.cpp b/Engine/Source/Runtime/Experimental/ChaosSolverEngine/Private/Chaos/ChaosVDRemoteSessionsManager.cpp
--- a/Engine/Source/Runtime/Experimental/ChaosSolverEngine/Private/Chaos/ChaosVDRemoteSessionsManager.cpp
+++ b/Engine/Source/Runtime/Experimental/ChaosSolverEngine/Private/Chaos/ChaosVDRemoteSessionsManager.cpp
@@ -282,15 +282,7 @@ void FChaosVDRemoteSessionsManager::ShutdownMessagingSystem()
 {
 	if (MessageEndpoint)
 	{
-		MessageEndpoint->Unsubscribe<FChaosVDSessionPong>();
-		MessageEndpoint->Unsubscribe<FChaosVDRecordingStatusMessage>();
-		MessageEndpoint->Unsubscribe<FChaosVDSessionPing>();
-		MessageEndpoint->Unsubscribe<FChaosVDStartRecordingCommandMessage>();
-		MessageEndpoint->Unsubscribe<FChaosVDStopRecordingCommandMessage>();
-		MessageEndpoint->Unsubscribe<FChaosVDChannelStateChangeCommandMessage>();
-		MessageEndpoint->Unsubscribe<FChaosVDFullSessionInfoRequestMessage>();
-		MessageEndpoint->Unsubscribe<FChaosVDFullSessionInfoResponseMessage>();
-		MessageEndpoint->Unsubscribe<FChaosVDChannelStateChangeResponseMessage>();
+		MessageEndpoint->Unsubscribe();
 	}
 
 	MessageBusPtr.Reset();
diff --git a/Engine/Source/Runtime/Experimental/ChaosVisualDebugger/Private/ChaosVDRuntimeModule.cpp b/Engine/Source/Runtime/Experimental/ChaosVisualDebugger/Private/ChaosVDRuntimeModule.cpp
--- a/Engine/Source/Runtime/Experimental/ChaosVisualDebugger/Private/ChaosVDRuntimeModule.cpp
+++ b/Engine/Source/Runtime/Experimental/ChaosVisualDebugger/Private/ChaosVDRuntimeModule.cpp
@@ -88,11 +88,6 @@ FChaosVDRuntimeModule& FChaosVDRuntimeModule::Get()
 	return FModuleManager::Get().LoadModuleChecked<FChaosVDRuntimeModule>(TEXT("ChaosVDRuntime"));
 }
 
-bool FChaosVDRuntimeModule::IsLoaded()
-{
-	return FModuleManager::Get().IsModuleLoaded(TEXT("ChaosVDRuntime"));
-}
-
 void FChaosVDRuntimeModule::StartupModule()
 {
 	if (FParse::Param(FCommandLine::Get(), TEXT("StartCVDRecording")))
diff --git a/Engine/Source/Runtime/Experimental/ChaosVisualDebugger/Public/ChaosVDRuntimeModule.h b/Engine/Source/Runtime/Experimental/ChaosVisualDebugger/Public/ChaosVDRuntimeModule.h
--- a/Engine/Source/Runtime/Experimental/ChaosVisualDebugger/Public/ChaosVDRuntimeModule.h
+++ b/Engine/Source/Runtime/Experimental/ChaosVisualDebugger/Public/ChaosVDRuntimeModule.h
@@ -44,7 +44,6 @@ class FChaosVDRuntimeModule : public IModuleInterface
 public:
 
 	static UE_API FChaosVDRuntimeModule& Get();
-	static UE_API bool IsLoaded();
 	
 	UE_API virtual void StartupModule() override;
 	UE_API virtual void ShutdownModule() override;
