Fix FHost::SetBufferSize to properly assign the buffer size when 'EDirection::Send' is used.

#rnx

#Synced-CL 47893047

diff --git a/Engine/Source/Runtime/Experimental/IoStore/HttpClient/Private/ConnectionPool.inl b/Engine/Source/Runtime/Experimental/IoStore/HttpClient/Private/ConnectionPool.inl
--- a/Engine/Source/Runtime/Experimental/IoStore/HttpClient/Private/ConnectionPool.inl
+++ b/Engine/Source/Runtime/Experimental/IoStore/HttpClient/Private/ConnectionPool.inl
@@ -72,7 +72,7 @@ FHost::FHost(const FParams& Params)
 ////////////////////////////////////////////////////////////////////////////////
 void FHost::SetBufferSize(EDirection Dir, int32 Size)
 {
-	(Dir == EDirection::Send) ? SendBufKb : RecvBufKb = uint16(Size >> 10);
+	(Dir == EDirection::Send ? SendBufKb : RecvBufKb) = uint16(Size >> 10);
 }
 
 ////////////////////////////////////////////////////////////////////////////////
