Remove old HTTP dispatcher in I/O store on-demand module.

- moved all Cvars from old backend to HTTP I/O dispatcher
- renamed FOnDemandEndpointConfig to FOnDemandIoStoreConfig since it is mainly used to parse the uber TOC path
- will renable lane tracing in a later change

#rnx

#Synced-CL 48026973

diff --git a/Engine/Source/Runtime/Experimental/IoStore/OnDemand/Private/DistributionEndpoints.cpp b/Engine/Source/Runtime/Experimental/IoStore/OnDemand/Private/DistributionEndpoints.cpp
deleted file mode 100644
diff --git a/Engine/Source/Runtime/Experimental/IoStore/OnDemand/Private/DistributionEndpoints.h b/Engine/Source/Runtime/Experimental/IoStore/OnDemand/Private/DistributionEndpoints.h
deleted file mode 100644
diff --git a/Engine/Source/Runtime/Experimental/IoStore/OnDemand/Private/OnDemandConfig.cpp b/Engine/Source/Runtime/Experimental/IoStore/OnDemand/Private/OnDemandConfig.cpp
--- a/Engine/Source/Runtime/Experimental/IoStore/OnDemand/Private/OnDemandConfig.cpp
+++ b/Engine/Source/Runtime/Experimental/IoStore/OnDemand/Private/OnDemandConfig.cpp
@@ -17,7 +17,6 @@
 #include "Misc/PathViews.h"
 #include "Misc/Paths.h"
 #include "OnDemandInstallCache.h"
-#include "OnDemandIoDispatcherBackend.h"
 #include "String/LexFromString.h"
 
 #if WITH_IOSTORE_ONDEMAND_TESTS
@@ -89,7 +88,7 @@ static bool ApplyEncryptionKeyFromString(const FString& GuidKeyPair)
 }
 
 ////////////////////////////////////////////////////////////////////////////////
-static bool TryParseConfigContent(const FString& ConfigContent, const FString& ConfigFileName, FOnDemandEndpointConfig& OutEndpoint)
+static bool TryParseConfigContent(const FString& ConfigContent, const FString& ConfigFileName, FOnDemandIoStoreConfig& OutConfig)
 {
 	if (ConfigContent.IsEmpty())
 	{
@@ -99,26 +98,10 @@ static bool TryParseConfigContent(const FString& ConfigContent, const FString& C
 	FConfigFile ConfigFile;
 	ConfigFile.ProcessInputFileContents(ConfigContent, ConfigFileName);
 
-	ConfigFile.GetString(TEXT("Endpoint"), TEXT("DistributionUrl"), OutEndpoint.DistributionUrl);
-	if (!OutEndpoint.DistributionUrl.IsEmpty())
-	{
-		ConfigFile.GetString(TEXT("Endpoint"), TEXT("FallbackUrl"), OutEndpoint.FallbackUrl);
-
-		if (!GDistributedEndpointFallbackUrl.IsEmpty())
-		{
-			OutEndpoint.FallbackUrl = GDistributedEndpointFallbackUrl;
-		}
-	}
-	
-	ConfigFile.GetArray(TEXT("Endpoint"), TEXT("ServiceUrl"), OutEndpoint.ServiceUrls);
-	ConfigFile.GetString(TEXT("Endpoint"), TEXT("TocPath"), OutEndpoint.TocPath);
-
-	if (OutEndpoint.DistributionUrl.EndsWith(TEXT("/")))
-	{
-		OutEndpoint.DistributionUrl = OutEndpoint.DistributionUrl.Left(OutEndpoint.DistributionUrl.Len() - 1);
-	}
+	ConfigFile.GetArray(TEXT("Endpoint"), TEXT("ServiceUrl"), OutConfig.ServiceUrls);
+	ConfigFile.GetString(TEXT("Endpoint"), TEXT("TocPath"), OutConfig.TocPath);
 
-	for (FString& ServiceUrl : OutEndpoint.ServiceUrls)
+	for (FString& ServiceUrl : OutConfig.ServiceUrls)
 	{
 		if (ServiceUrl.EndsWith(TEXT("/")))
 		{
@@ -126,9 +109,9 @@ static bool TryParseConfigContent(const FString& ConfigContent, const FString& C
 		}
 	}
 
-	if (OutEndpoint.TocPath.StartsWith(TEXT("/")))
+	if (OutConfig.TocPath.StartsWith(TEXT("/")))
 	{
-		OutEndpoint.TocPath.RightChopInline(1);
+		OutConfig.TocPath.RightChopInline(1);
 	}
 
 	FString ContentKey;
@@ -137,7 +120,7 @@ static bool TryParseConfigContent(const FString& ConfigContent, const FString& C
 		ApplyEncryptionKeyFromString(ContentKey);
 	}
 
-	return OutEndpoint.IsValid();
+	return OutConfig.IsValid();
 }
 
 ////////////////////////////////////////////////////////////////////////////////
@@ -239,7 +222,7 @@ FIasCacheConfig GetStreamingCacheConfig(const TCHAR* CommandLine)
 }
 
 ////////////////////////////////////////////////////////////////////////////////
-TIoStatusOr<FOnDemandEndpointConfig> TryParseEndpointConfig(const TCHAR* CommandLine)
+TIoStatusOr<FOnDemandIoStoreConfig> TryParseConfig(const TCHAR* CommandLine)
 {
 	{
 		FString EncryptionKey;
@@ -268,7 +251,7 @@ TIoStatusOr<FOnDemandEndpointConfig> TryParseEndpointConfig(const TCHAR* Command
 	}
 #endif
 
-	FOnDemandEndpointConfig OutConfig;
+	FOnDemandIoStoreConfig OutConfig;
 #if !UE_BUILD_SHIPPING
 	FString UrlParam;
 	if (FParse::Value(CommandLine, TEXT("Ias.TocUrl="), UrlParam))
@@ -331,7 +314,6 @@ TIoStatusOr<FOnDemandEndpointConfig> TryParseEndpointConfig(const TCHAR* Command
 	FString HostsOverride;
 	if (FParse::Value(CommandLine, TEXT("Iax.DefaultHostGroup="), HostsOverride))
 	{
-		OutConfig.DistributionUrl.Empty();
 		OutConfig.ServiceUrls.Empty();
 		HostsOverride.ParseIntoArray(OutConfig.ServiceUrls, TEXT(","), true);
 		for (FString& Host : OutConfig.ServiceUrls)
diff --git a/Engine/Source/Runtime/Experimental/IoStore/OnDemand/Private/OnDemandConfig.h b/Engine/Source/Runtime/Experimental/IoStore/OnDemand/Private/OnDemandConfig.h
--- a/Engine/Source/Runtime/Experimental/IoStore/OnDemand/Private/OnDemandConfig.h
+++ b/Engine/Source/Runtime/Experimental/IoStore/OnDemand/Private/OnDemandConfig.h
@@ -6,11 +6,40 @@
 #include "Containers/StringView.h"
 #include "IO/IoStatus.h"
 
-namespace UE::IoStore { struct FOnDemandEndpointConfig; }
-namespace UE::IoStore { struct FIasCacheConfig; }
-namespace UE::IoStore { struct FOnDemandInstallCacheConfig; }
+namespace UE::IoStore
+{
+
+struct FIasCacheConfig;
+struct FOnDemandInstallCacheConfig;
 
-namespace UE::IoStore::Config
+struct FOnDemandIoStoreConfig
+{
+	TArray<FString> ServiceUrls;
+	FString			TocPath;
+	FString			TocFilePath;
+
+	bool IsValid() const
+	{
+		return !TocPath.IsEmpty() || !TocFilePath.IsEmpty();
+	}
+
+	bool RequiresTls() const
+	{
+		auto IsHttps = [](FStringView Url) -> bool { return Url.StartsWith(TEXTVIEW("https"), ESearchCase::IgnoreCase); };
+
+		for (const FString& Url : ServiceUrls)
+		{
+			if (IsHttps(Url))
+			{
+				return true;
+			}
+		}
+
+		return false;
+	}
+};
+
+namespace Config
 {
 
 /*
@@ -46,10 +75,12 @@ int64 ParseSizeParam(FStringView Value);
  */
 int64 ParseSizeParam(const TCHAR* CommandLine, const TCHAR* Param);
 
-TIoStatusOr<FOnDemandEndpointConfig> TryParseEndpointConfig(const TCHAR* CommandLine);
+TIoStatusOr<FOnDemandIoStoreConfig> TryParseConfig(const TCHAR* CommandLine);
 
 FIasCacheConfig GetStreamingCacheConfig(const TCHAR* CommandLine);
 
 TIoStatusOr<FOnDemandInstallCacheConfig> TryParseInstallCacheConfig(const TCHAR* CommandLine);
 
 } // namespace UE::IoStore::Config
+
+} // namespace UE::IoStore
diff --git a/Engine/Source/Runtime/Experimental/IoStore/OnDemand/Private/OnDemandContentInstaller.cpp b/Engine/Source/Runtime/Experimental/IoStore/OnDemand/Private/OnDemandContentInstaller.cpp
--- a/Engine/Source/Runtime/Experimental/IoStore/OnDemand/Private/OnDemandContentInstaller.cpp
+++ b/Engine/Source/Runtime/Experimental/IoStore/OnDemand/Private/OnDemandContentInstaller.cpp
@@ -13,12 +13,10 @@
 #include "IO/IoStoreOnDemand.h"
 #include "IO/OnDemandError.h"
 #include "Misc/Timespan.h"
-#include "OnDemandHttpThread.h"
 #include "OnDemandInstallCache.h"
 #include "OnDemandIoStore.h"
 #include "OnDemandPackageStoreBackend.h"
 #include "Statistics.h"
-#include "Templates/Overload.h"
 #include <atomic>
 
 TRACE_DECLARE_INT_COUNTER(OnDemandContentInstaller_HttpRequestsIssued, TEXT("OnDemandContentInstaller/HttpRequestsIssued"));
@@ -256,9 +254,8 @@ TOptional<UE::UnifiedError::FError> FOnDemandContentInstaller::FRequest::Consume
 }
 
 ////////////////////////////////////////////////////////////////////////////////
-FOnDemandContentInstaller::FOnDemandContentInstaller(FOnDemandIoStore& InIoStore, FOnDemandHttpThread* InHttpClient)
+FOnDemandContentInstaller::FOnDemandContentInstaller(FOnDemandIoStore& InIoStore)
 	: IoStore(InIoStore)
-	, HttpClient(InHttpClient)
 	, InstallerPipe(TEXT("IoStoreOnDemandInstallerPipe"))
 {
 }
@@ -441,22 +438,7 @@ void FOnDemandContentInstaller::UpdateInstallRequestPriority(FSharedInternalInst
 							}
 						}
 
-						Visit(UE::Overload(
-							[MaxPriority](FIoHttpRequest& Handle)
-							{
-								Handle.UpdatePriorty(MaxPriority);
-							},
-							[this, MaxPriority](void* Handle)
-							{
-								check(HttpClient != nullptr);
-								if (Handle != nullptr)
-								{
-									HttpClient->ReprioritizeRequest(Handle, MaxPriority);
-								}
-							}), 
-							ChunkDownload.HttpHandle);
-
-						// Should be a max of one active request for any chunk
+						ChunkDownload.HttpHandle.UpdatePriorty(MaxPriority);
 						break;
 					}
 				}
@@ -915,14 +897,7 @@ void FOnDemandContentInstaller::ExecuteInstallRequest(FRequest& Request)
 
 	NotifyInstallProgress(Request);
 
-	const bool bUseHttpIoDispatcher = FHttpIoDispatcher::IsInitialized();
-	check(bUseHttpIoDispatcher || HttpClient != nullptr);
-
-	TOptional<FIoHttpBatch> Batch;
-	if (bUseHttpIoDispatcher)
-	{
-		Batch.Emplace(FHttpIoDispatcher::NewBatch());
-	}
+	FIoHttpBatch Batch = FHttpIoDispatcher::NewBatch();
 
 	constexpr const int32 HttpRetryCount = 2;
 	for (FChunkHttpRequestHandle& HttpRequest : InstallRequest.HttpRequestHandles)
@@ -964,72 +939,28 @@ void FOnDemandContentInstaller::ExecuteInstallRequest(FRequest& Request)
 				*TStringBuilder<64>(InPlace, ChunkHash), PendingChunkRequest->RequestId);
 
 			TRACE_COUNTER_INCREMENT(OnDemandContentInstaller_HttpRequestsIssued);
-
-			if (bUseHttpIoDispatcher)
-			{
-				PendingChunkRequest->HttpHandle.Emplace<FIoHttpRequest>(Batch->Get(
-					ChunkInfo.HostGroupName(),
-					ChunkInfo.RelativeUrl(),
-					FIoHttpHeaders(),
-					HttpOptions,
-					ChunkHash,
-					[this, ChunkHash, RequestId = PendingChunkRequest->RequestId](FIoHttpResponse&& HttpResponse)
-					{
-						InstallerPipe.Launch(
-							TEXT("ProcessIoStoreOnDemandDownloadedChunk"),
-							[this, ChunkHash, RequestId, HttpResponse = MoveTemp(HttpResponse)]() mutable
-							{
-								FIoBuffer Chunk = HttpResponse.GetBody();
-								ProcessDownloadedChunk(ChunkHash, RequestId, HttpResponse.GetErrorCode(), HttpResponse.GetStatusCode(), MoveTemp(Chunk));
-							},
-							UE::Tasks::ETaskPriority::BackgroundLow);
-					})
-				);
-			}
-			else
-			{
-				PendingChunkRequest->HttpHandle.Emplace<void*>(HttpClient->IssueRequest(
-					MoveTemp(ChunkInfo),
-					FIoOffsetAndLength(),
-					Request.Priority,
-					[this, ChunkHash, RequestId = PendingChunkRequest->RequestId](uint32 HttpStatusCode, FStringView ErrorReason, FIoBuffer&& Chunk)
-					{
-#if !UE_BUILD_SHIPPING
-						if (IoStore::CVars::IoStoreSimulateHttpError >= 0)
-						{
-							HttpStatusCode = (uint32)IoStore::CVars::IoStoreSimulateHttpError;
-						}
-#endif
-						EIoErrorCode CompletionStatus = (HttpStatusCode > 199 && HttpStatusCode < 300) ? EIoErrorCode::Ok : EIoErrorCode::ReadError;
-						if (ErrorReason.Contains(TEXTVIEW("cancelled"), ESearchCase::IgnoreCase) ||
-							ErrorReason.Contains(TEXTVIEW("canceled"),  ESearchCase::IgnoreCase)) // Yes we spell it both ways (╯°□°）╯︵ ┻━┻
-						{
-							CompletionStatus = EIoErrorCode::Cancelled;
-						}
-						else
+			PendingChunkRequest->HttpHandle = Batch.Get(
+				ChunkInfo.HostGroupName(),
+				ChunkInfo.RelativeUrl(),
+				FIoHttpHeaders(),
+				HttpOptions,
+				ChunkHash,
+				[this, ChunkHash, RequestId = PendingChunkRequest->RequestId](FIoHttpResponse&& HttpResponse)
+				{
+					InstallerPipe.Launch(
+						TEXT("ProcessIoStoreOnDemandDownloadedChunk"),
+						[this, ChunkHash, RequestId, HttpResponse = MoveTemp(HttpResponse)]() mutable
 						{
-							UE_CLOG(ErrorReason.IsEmpty() == false, LogIoStoreOnDemand, Warning, TEXT("Failed to download chunk: %.*s"), ErrorReason.Len(), ErrorReason.GetData());
-						}
-
-						InstallerPipe.Launch(
-							TEXT("ProcessIoStoreOnDemandDownloadedChunk"),
-							[this, ChunkHash, RequestId, CompletionStatus, HttpStatusCode, Chunk = MoveTemp(Chunk)]() mutable
-							{
-								ProcessDownloadedChunk(ChunkHash, RequestId, CompletionStatus, HttpStatusCode, MoveTemp(Chunk));
-							},
-							UE::Tasks::ETaskPriority::BackgroundLow);
-					},
-					EHttpRequestType::Installed)
-				);
-			}
+							FIoBuffer Chunk = HttpResponse.GetBody();
+							ProcessDownloadedChunk(ChunkHash, RequestId, HttpResponse.GetErrorCode(), HttpResponse.GetStatusCode(), MoveTemp(Chunk));
+						},
+						UE::Tasks::ETaskPriority::BackgroundLow);
+				});
 		}
 
 	}
 
-	if (Batch)
-	{
-		Batch->Issue();
-	}
+	Batch.Issue();
 	
 	InstallRequest.bHttpRequestsIssued = true;
 }
@@ -1237,7 +1168,7 @@ void FOnDemandContentInstaller::OnProcessDownloadedChunkNotifyRequest(
 	bool bSetErrorOnRequest = false;
 	if (Request.IsOk())
 	{
-		if (!ensureMsgf(!bChunkCancelled, TEXT("Recieved cancelled http callback for active Request %s ! This should never happen."), *InstallRequest.Args.DebugName))
+		if (!ensureMsgf(!bChunkCancelled, TEXT("Received cancelled http callback for active Request %s ! This should never happen."), *InstallRequest.Args.DebugName))
 		{
 			bSetErrorOnRequest = true;
 			Request.Result = ChunkDownloadResult;
@@ -1353,28 +1284,9 @@ void FOnDemandContentInstaller::TryCancelHttpRequestsForInstallRequest(FRequest:
 
 			if (bCancelHttpRequest)
 			{
-				// No requests are waiting on this chunk, cancel it
-				Visit(UE::Overload(
-					[&NumCancelled](FIoHttpRequest& Handle)
-					{
-						if (Handle.IsValid())
-						{
-							Handle.Cancel();
-							++NumCancelled;
-						}
-					},
-					[this, &NumCancelled](void* Handle)
-					{
-						check(HttpClient != nullptr);
-						if (Handle != nullptr)
-						{
-							HttpClient->CancelRequest(Handle);
-							++NumCancelled;
-						}
-					}),
-					ChunkDownload->HttpHandle);
-
+				ChunkDownload->HttpHandle.Cancel();
 				ChunkDownload->bChunkCanceled = true;
+				++NumCancelled;
 				UE_LOG(LogIoStoreOnDemand, Verbose, TEXT("Cancelled chunk download: hash - %s, Id - %u"), *TStringBuilder<64>(InPlace, PendingChunkHash), ChunkDownload->RequestId);
 
 				// ProcessDownloadedChunk will remove the chunk request
diff --git a/Engine/Source/Runtime/Experimental/IoStore/OnDemand/Private/OnDemandContentInstaller.h b/Engine/Source/Runtime/Experimental/IoStore/OnDemand/Private/OnDemandContentInstaller.h
--- a/Engine/Source/Runtime/Experimental/IoStore/OnDemand/Private/OnDemandContentInstaller.h
+++ b/Engine/Source/Runtime/Experimental/IoStore/OnDemand/Private/OnDemandContentInstaller.h
@@ -281,12 +281,11 @@ class FOnDemandContentInstaller
 
 	struct FChunkDownloadRequest
 	{
-		TVariant<void*, FIoHttpRequest> HttpHandle;
+		FIoHttpRequest							HttpHandle;
 		TIntrusiveList<FChunkHttpRequestHandle> ChunkRequestHttpHandles;
-
 		// Should only be used to compare two requests for the same iochunk, otherwise RequestId is not unique
-		uint32				RequestId = 0;
-		bool				bChunkCanceled = false;
+		uint32									RequestId = 0;
+		bool									bChunkCanceled = false;
 	};
 
 	static bool RequestSortPredicate(const FRequest& LHS, const FRequest& RHS)
@@ -307,7 +306,7 @@ class FOnDemandContentInstaller
 	using FRequestAllocator	= TSingleThreadedSlabAllocator<FRequest, 32>;
 
 public:
-									FOnDemandContentInstaller(FOnDemandIoStore& IoStore, FOnDemandHttpThread* HttpClient);
+									FOnDemandContentInstaller(FOnDemandIoStore& IoStore);
 									~FOnDemandContentInstaller();
 
 	FSharedInternalInstallRequest	EnqueueInstallRequest(
@@ -351,7 +350,6 @@ private:
 	void							Shutdown();
 
 	FOnDemandIoStore&				IoStore;
-	FOnDemandHttpThread*			HttpClient;
 	UE::Tasks::FPipe				InstallerPipe;
 
 	FMutex							Mutex;
diff --git a/Engine/Source/Runtime/Experimental/IoStore/OnDemand/Private/OnDemandHttpIoDispacher.cpp b/Engine/Source/Runtime/Experimental/IoStore/OnDemand/Private/OnDemandHttpIoDispacher.cpp
--- a/Engine/Source/Runtime/Experimental/IoStore/OnDemand/Private/OnDemandHttpIoDispacher.cpp
+++ b/Engine/Source/Runtime/Experimental/IoStore/OnDemand/Private/OnDemandHttpIoDispacher.cpp
@@ -23,21 +23,82 @@
 #include "Misc/Fork.h"
 #include "Misc/SingleThreadRunnable.h"
 #include "OnDemandHttpClient.h"
-#include "OnDemandHttpThread.h" // EHttpRequestType
 #include "Statistics.h"
 
 #include <atomic>
 
 // When enabled the cvar 'iax.InvalidUrlChance' will be enabled and allow us to simulate invalid urls
 #define UE_ALLOW_INVALID_URL_DEBUGGING !UE_BUILD_SHIPPING
-
+#define UE_ALLOW_DISABLE_CANCELLING 1 && !UE_BUILD_SHIPPING
 #define UE_ALLOW_CACHE_POISONING !UE_BUILD_SHIPPING
+#define UE_ALLOW_HTTP_PAUSE (1 && !NO_CVARS && !UE_BUILD_SHIPPING)
 
-// TODO - Maybe do LexToString/LexFromString (dupe code in UnrealEngine.cpp ThreadPriorityToString)
-extern const TCHAR* LexToString(EThreadPriority Priority);
+////////////////////////////////////////////////////////////////////////////////
+const TCHAR* LexToString(EThreadPriority Priority)
+{
+	switch (Priority)
+	{
+		case EThreadPriority::TPri_Normal:
+			return TEXT("TPri_Normal");
+		case EThreadPriority::TPri_AboveNormal:
+			return TEXT("TPri_AboveNormal");
+		case EThreadPriority::TPri_BelowNormal:
+			return TEXT("TPri_BelowNormal");
+		case EThreadPriority::TPri_Highest:
+			return TEXT("TPri_Highest");
+		case EThreadPriority::TPri_Lowest:
+			return TEXT("TPri_Lowest");
+		case EThreadPriority::TPri_SlightlyBelowNormal:
+			return TEXT("TPri_SlightlyBelowNormal");
+		case EThreadPriority::TPri_TimeCritical:
+			return TEXT("TPri_TimeCritical");
+		case EThreadPriority::TPri_Num:
+		default:
+			return TEXT("TPri_Undefined");
+			break;
+	};
+}
 
+////////////////////////////////////////////////////////////////////////////////
+#if UE_ALLOW_HTTP_PAUSE
+void FConsoleCommandUnregister::operator()(IConsoleCommand* ConsoleCommandPtr) const
+{
+	IConsoleManager::Get().UnregisterConsoleObject(ConsoleCommandPtr);
+};
+#endif // UE_ALLOW_HTTP_PAUSE
+
+////////////////////////////////////////////////////////////////////////////////
 namespace UE::IoStore
 {
+
+/** Allows requests to be filtered by type */
+enum class EHttpRequestTypeFilter : uint8
+{
+	None  = 0,
+
+	Streaming = 1 << 0,	/** IAS */
+	Installed = 1 << 1,	/** IAD */
+
+	All = Streaming | Installed
+};
+ENUM_CLASS_FLAGS(EHttpRequestTypeFilter);
+
+////////////////////////////////////////////////////////////////////////////////
+DECLARE_MULTICAST_DELEGATE(FOnRecreateHttpClient);
+FOnRecreateHttpClient OnRecreateHttpClient;
+
+///////////////////////////////////////////////////////////////////////////////
+void OnHttpClientCVarChanged(IConsoleVariable* CVar)
+{
+	if (ensure(CVar != nullptr))
+	{
+		const FString CVarName = IConsoleManager::Get().FindConsoleObjectName(CVar);
+		UE_LOGFMT(LogIas, Log, "Existing http client config invalidated as cvar '{CvarName}' has changed", CVarName);
+
+		OnRecreateHttpClient.Broadcast();
+	}
+}
+
 ////////////////////////////////////////////////////////////////////////////////
 int32 GHttpIoDispatcherMaxRangeKiB = 4096;
 static FAutoConsoleVariableRef CVar_HttpIoDispatcherMaxRangeKiB(
@@ -57,28 +118,194 @@ static FAutoConsoleVariableRef CVar_HttpIoDispatcherShutdownTimeoutSeconds(
 	ECVF_Default
 );
 
-////////////////////////////////////////////////////////////////////////////////
-DECLARE_MULTICAST_DELEGATE(FOnRecreateHttpClient);
-extern FOnRecreateHttpClient OnRecreateHttpClient;
+/** Note that GIasHttpPrimaryEndpoint has no effect after initial start up */
+int32 GIasHttpPrimaryEndpoint = 0;
+static FAutoConsoleVariableRef CVar_IasHttpPrimaryEndpoint(
+	TEXT("ias.HttpPrimaryEndpoint"),
+	GIasHttpPrimaryEndpoint,
+	TEXT("Primary endpoint to use returned from the distribution endpoint")
+);
+
+int32 GIasHttpTimeOutMs = 10 * 1000;
+static FAutoConsoleVariableRef CVar_IasHttpTimeOutMs(
+	TEXT("ias.HttpTimeOutMs"),
+	GIasHttpTimeOutMs,
+	TEXT("Time out value for HTTP requests in milliseconds")
+);
+
+bool GIasHttpEnabled = true;
+static FAutoConsoleVariableRef CVar_IasHttpEnabled(
+	TEXT("ias.HttpEnabled"),
+	GIasHttpEnabled,
+	TEXT("Enables individual asset streaming via HTTP")
+);
+
+bool GIasHttpOptionalBulkDataEnabled = true;
+static FAutoConsoleVariableRef CVar_IasHttpOptionalBulkDataEnabled(
+	TEXT("ias.HttpOptionalBulkDataEnabled"),
+	GIasHttpOptionalBulkDataEnabled,
+	TEXT("Enables optional bulk data via HTTP")
+);
+
+bool GIasReportAnalyticsEnabled = true;
+static FAutoConsoleVariableRef CVar_IoReportAnalytics(
+	TEXT("ias.ReportAnalytics"),
+	GIasReportAnalyticsEnabled,
+	TEXT("Enables reporting statics to the analytics system")
+);
+
+static int32 GDistributedEndpointRetryWaitTime = 15;
+static FAutoConsoleVariableRef CVar_DistributedEndpointRetryWaitTime(
+	TEXT("ias.DistributedEndpointRetryWaitTime"),
+	GDistributedEndpointRetryWaitTime,
+	TEXT("How long to wait (in seconds) after failing to resolve a distributed endpoint before retrying")
+);
+
+static int32 GDistributedEndpointAttemptCount = 5;
+static FAutoConsoleVariableRef CVar_DistributedEndpointAttemptCount(
+	TEXT("ias.DistributedEndpointAttemptCount"),
+	GDistributedEndpointAttemptCount,
+	TEXT("Number of times we should try to resolve a distributed endpoint befor eusing the fallback url (if there is one)")
+);
+
+bool GIasEnableRequestHandleClear = true;
+static FAutoConsoleVariableRef CVar_EnableRequestHandleClear(
+	TEXT("ias.EnableRequestHandleClear"),
+	GIasEnableRequestHandleClear,
+	TEXT("When enabled FChunkRequest::HttpRequest will be cleared when the request callback is triggered")
+);
+
+#if UE_ALLOW_DISABLE_CANCELLING
+bool GIasCancelationEnabled = true;
+static FAutoConsoleVariableRef CVar_IasCancelationEnabled(
+	TEXT("ias.EnableCancelation"),
+	GIasCancelationEnabled,
+	TEXT("Allows existing IO requests to be canceled")
+);
+#endif //UE_ALLOW_DISABLE_CANCELLING
 
-extern int32 GIasHttpTimeOutMs;
-extern int32 GOnDemandBackendThreadPriorityIndex; 
-extern bool GIaxHttpEnableInflightCancellation;
-extern float GIaxInvalidUrlChance;
 #if UE_ALLOW_CACHE_POISONING
-extern bool GIasPoisonCache;
+bool GIasPoisonCache = false;
+static FAutoConsoleVariableRef CVar_IasPoisonCache(
+	TEXT("ias.PoisonCache"),
+	GIasPoisonCache,
+	TEXT("Fills all data materialized from the cache with 0x4d")
+);
 #endif // UE_ALLOW_CACHE_POISONING
-int32 CalculateRequestQueueLengthMax();
 
-/** Matches the priority set up found in Engine/Source/Runtime/Online/HTTP/Private/HttpThread.cpp */
-static constexpr int32 GThreadPriorities[5] =
-{
-	EThreadPriority::TPri_Lowest,
-	EThreadPriority::TPri_BelowNormal,
-	EThreadPriority::TPri_SlightlyBelowNormal,
-	EThreadPriority::TPri_Normal,
-	EThreadPriority::TPri_AboveNormal
-};
+int32 GIasHttpHealthCheckWaitTime = 3000;
+static FAutoConsoleVariableRef CVar_IasHttpHealthCheckWaitTime(
+	TEXT("ias.HttpHealthCheckWaitTime"),
+	GIasHttpHealthCheckWaitTime,
+	TEXT("Number of milliseconds to wait before reconnecting to avaiable endpoint(s)")
+);
+
+int32 GOnDemandBackendThreadPriorityIndex = 4; // EThreadPriority::TPri_AboveNormal
+FAutoConsoleVariableRef CVarOnDemandBackendThreadPriority(
+	TEXT("ias.onDemandBackendThreadPriority"),
+	GOnDemandBackendThreadPriorityIndex,
+	TEXT("Thread priority of the on demand backend thread: 0=Lowest, 1=BelowNormal, 2=SlightlyBelowNormal, 3=Normal, 4=AboveNormal\n")
+	TEXT("Note that this is switchable at runtime"),
+	ECVF_Default);
+
+int32 GIaxHttpVersion = 1;
+static FAutoConsoleVariableRef CVar_IaxHttpVersion(
+	TEXT("iax.HttpVersion"),
+	GIaxHttpVersion,
+	TEXT("Protocol version to use, either '1' or '2'")
+);
+
+int32 GIasHttpConnectionCount = 4;
+static FAutoConsoleVariableRef CVar_IasHttpConnectionCount(
+	TEXT("ias.HttpConnectionCount"),
+	GIasHttpConnectionCount,
+	TEXT("Number of open HTTP connections to the on demand endpoint(s)."),
+	FConsoleVariableDelegate::CreateStatic(&OnHttpClientCVarChanged)
+);
+
+int32 GIaxHttpMaxInflight = 1;
+static FAutoConsoleVariableRef CVar_IaxHttpMaxInflight(
+	TEXT("iax.HttpMaxInflight"),
+	GIaxHttpMaxInflight,
+	TEXT("Number of requests to issue at once per connection")
+);
+
+int32 GIasHttpRecvBufKiB = -1;
+static FAutoConsoleVariableRef CVar_GIasHttpRecvBufKiB(
+	TEXT("ias.HttpRecvBufKiB"),
+	GIasHttpRecvBufKiB,
+	TEXT("Recv buffer size"),
+	FConsoleVariableDelegate::CreateStatic(&OnHttpClientCVarChanged)
+);
+
+int32 GIasHttpSendBufKiB = -1;
+static FAutoConsoleVariableRef CVar_GIasHttpSendBufKiB(
+	TEXT("ias.HttpSendBufKiB"),
+	GIasHttpSendBufKiB,
+	TEXT("Send buffer size"),
+	FConsoleVariableDelegate::CreateStatic(&OnHttpClientCVarChanged)
+);
+
+int32 GIasHttpRetryCount = -1;
+static FAutoConsoleVariableRef CVar_IasHttpRetryCount(
+	TEXT("ias.HttpRetryCount"),
+	GIasHttpRetryCount,
+	TEXT("Number of times that a request should be retried before being considered failed. A negative value will use the default behaviour, which is one retry per host url provided."),
+	FConsoleVariableDelegate::CreateStatic(&OnHttpClientCVarChanged)
+);
+
+int32 GIasHttpFailTimeOutMs = 4 * 1000;
+static FAutoConsoleVariableRef CVar_IasHttpFailTimeOutMs(
+	TEXT("ias.HttpFailTimeOutMs"),
+	GIasHttpFailTimeOutMs,
+	TEXT("Fail infinite network waits that take longer than this (in ms, a value of zero will use the default)"),
+	FConsoleVariableDelegate::CreateStatic(&OnHttpClientCVarChanged)
+);
+
+bool GIasHttpAllowChunkedXfer = true;
+static FAutoConsoleVariableRef CVar_IasHttpAllowChunkedXfer(
+	TEXT("ias.HttpAllowChunkedXfer"),
+	GIasHttpAllowChunkedXfer,
+	TEXT("Enable/disable IAS' support for chunked transfer encoding"),
+	FConsoleVariableDelegate::CreateStatic(&OnHttpClientCVarChanged)
+);
+
+int32 GIasHttpConcurrentRequests = 8;
+static FAutoConsoleVariableRef CVar_IasHttpConcurrentRequests(
+	TEXT("ias.HttpConcurrentRequests"),
+	GIasHttpConcurrentRequests,
+	TEXT("Number of concurrent requests in the http client.")
+);
+
+int32 GIasHttpRateLimitKiBPerSecond = 0;
+static FAutoConsoleVariableRef CVar_GIasHttpRateLimitKiBPerSecond(
+	TEXT("ias.HttpRateLimitKiBPerSecond"),
+	GIasHttpRateLimitKiBPerSecond,
+	TEXT("Http throttle limit in KiBPerSecond")
+);
+
+int32 GIasHttpPollTimeoutMs = 17;
+static FAutoConsoleVariableRef CVar_GIasHttpPollTimeoutMs(
+	TEXT("ias.HttpPollTimeoutMs"),
+	GIasHttpPollTimeoutMs,
+	TEXT("Http tick poll timeout in milliseconds")
+);
+
+bool GIaxHttpEnableInflightCancellation = true;
+static FAutoConsoleVariableRef CVar_IaxHttpEnableInflightCancellation(
+	TEXT("iax.HttpEnableInflightCancellation"),
+	GIaxHttpEnableInflightCancellation,
+	TEXT("When enabled the system will attempt to cancel inflight http requests")
+);
+
+#if UE_ALLOW_INVALID_URL_DEBUGGING
+float GIaxInvalidUrlChance = 0.0;
+static FAutoConsoleVariableRef CVar_GIaxInvalidUrlChance(
+	TEXT("iax.InvalidUrlChance"),
+	GIaxInvalidUrlChance,
+	TEXT("Chance that a url for a GET request will be invalid (0-100%)")
+);
+#endif // UE_ALLOW_INVALID_URL_DEBUGGING
 
 bool GHttpCacheEnabled = true;
 static FAutoConsoleVariableRef CVar_HttpCacheEnabled(
@@ -87,10 +314,31 @@ static FAutoConsoleVariableRef CVar_HttpCacheEnabled(
 	TEXT("")
 );
 
-namespace HttpIoDispatcher
+////////////////////////////////////////////////////////////////////////////////
+int32 CalculateRequestQueueLengthMax()
 {
+	// Keeping the loop loaded with work to do is important to cover round-trip
+	// latency and the effect of CDN cache misses. Reducing what is in the queue
+	// DOES NOT have any effect on throughput throttling.
+	int32 ConnNum = (GIaxHttpVersion == 2) ? 1 : GIasHttpConnectionCount;
+	int32 InflightNum = (GIaxHttpVersion == 2) ? GIaxHttpMaxInflight : 1;
+	int32 Value = FMath::Max((ConnNum * InflightNum) + 2, 1);
+	return FMath::Min(Value, int32(HTTP::FEventLoop::MaxActiveTickets));
+}
+
 ////////////////////////////////////////////////////////////////////////////////
+static constexpr int32 GThreadPriorities[5] =
+{
+	EThreadPriority::TPri_Lowest,
+	EThreadPriority::TPri_BelowNormal,
+	EThreadPriority::TPri_SlightlyBelowNormal,
+	EThreadPriority::TPri_Normal,
+	EThreadPriority::TPri_AboveNormal
+};
 
+////////////////////////////////////////////////////////////////////////////////
+namespace HttpIoDispatcher
+{
 /**
  * Tracks the lifespan of a http request for as long as we want to keep it active. When the token is
  * destroyed we will try to cancel the request if it is still in progress.
diff --git a/Engine/Source/Runtime/Experimental/IoStore/OnDemand/Private/OnDemandHttpIoDispatcher.h b/Engine/Source/Runtime/Experimental/IoStore/OnDemand/Private/OnDemandHttpIoDispatcher.h
--- a/Engine/Source/Runtime/Experimental/IoStore/OnDemand/Private/OnDemandHttpIoDispatcher.h
+++ b/Engine/Source/Runtime/Experimental/IoStore/OnDemand/Private/OnDemandHttpIoDispatcher.h
@@ -7,10 +7,18 @@
 
 #define UE_API IOSTOREONDEMANDCORE_API
 
-
 namespace UE::IoStore
 {
 
+/** Describes the system making the request */
+enum class EHttpRequestType : uint8
+{
+	Streaming = 0,	/** IAS */
+	Installed,		/** IAD */
+
+	NUM_SOURCES
+};
+
 class IOnDemandHttpIoDispatcher
 	: public IHttpIoDispatcher
 {
diff --git a/Engine/Source/Runtime/Experimental/IoStore/OnDemand/Private/OnDemandHttpIoDispatcherBackend.cpp b/Engine/Source/Runtime/Experimental/IoStore/OnDemand/Private/OnDemandHttpIoDispatcherBackend.cpp
--- a/Engine/Source/Runtime/Experimental/IoStore/OnDemand/Private/OnDemandHttpIoDispatcherBackend.cpp
+++ b/Engine/Source/Runtime/Experimental/IoStore/OnDemand/Private/OnDemandHttpIoDispatcherBackend.cpp
@@ -11,15 +11,25 @@
 #include "IO/IoStoreOnDemand.h"
 #include "OnDemandBackendStatus.h"
 #include "OnDemandIoStore.h"
-#include "OnDemandHttpThread.h" // EHttpRequestType
+#include "OnDemandHttpIoDispatcher.h" // EHttpRequestType
 #include "Statistics.h"
 
 namespace UE::IoStore
 {
 ////////////////////////////////////////////////////////////////////////////////
-
-extern int32 GIasHttpRangeRequestMinSizeKiB;
-extern bool GIasEnableWriteOnlyDecoding;
+int32 GIasHttpRangeRequestMinSizeKiB = 128;
+static FAutoConsoleVariableRef CVar_IasHttpRangeRequestMinSizeKiB(
+	TEXT("ias.HttpRangeRequestMinSizeKiB"),
+	GIasHttpRangeRequestMinSizeKiB,
+	TEXT("Minimum chunk size for partial chunk request(s)")
+);
+
+bool GIasEnableWriteOnlyDecoding = false;
+static FAutoConsoleVariableRef CVar_WriteOnlyDecoding(
+	TEXT("ias.EnableWriteOnlyDecoding"),
+	GIasEnableWriteOnlyDecoding,
+	TEXT("Enables the use of 'WriteOnly' flag when decoding to buffers with the 'HardwareTargetBuffer' flag")
+);
 
 ////////////////////////////////////////////////////////////////////////////////
 class FHttpIoDispatcherBackend
diff --git a/Engine/Source/Runtime/Experimental/IoStore/OnDemand/Private/OnDemandHttpThread.cpp b/Engine/Source/Runtime/Experimental/IoStore/OnDemand/Private/OnDemandHttpThread.cpp
deleted file mode 100644
diff --git a/Engine/Source/Runtime/Experimental/IoStore/OnDemand/Private/OnDemandHttpThread.h b/Engine/Source/Runtime/Experimental/IoStore/OnDemand/Private/OnDemandHttpThread.h
deleted file mode 100644
diff --git a/Engine/Source/Runtime/Experimental/IoStore/OnDemand/Private/OnDemandIoDispatcherBackend.cpp b/Engine/Source/Runtime/Experimental/IoStore/OnDemand/Private/OnDemandIoDispatcherBackend.cpp
deleted file mode 100644
diff --git a/Engine/Source/Runtime/Experimental/IoStore/OnDemand/Private/OnDemandIoDispatcherBackend.h b/Engine/Source/Runtime/Experimental/IoStore/OnDemand/Private/OnDemandIoDispatcherBackend.h
deleted file mode 100644
diff --git a/Engine/Source/Runtime/Experimental/IoStore/OnDemand/Private/OnDemandIoStore.cpp b/Engine/Source/Runtime/Experimental/IoStore/OnDemand/Private/OnDemandIoStore.cpp
--- a/Engine/Source/Runtime/Experimental/IoStore/OnDemand/Private/OnDemandIoStore.cpp
+++ b/Engine/Source/Runtime/Experimental/IoStore/OnDemand/Private/OnDemandIoStore.cpp
@@ -39,10 +39,8 @@
 #include "OnDemandHttpClient.h"
 #include "OnDemandHttpIoDispatcher.h"
 #include "OnDemandHttpIoDispatcherBackend.h"
-#include "OnDemandHttpThread.h"
 #include "OnDemandInstallCache.h"
 #include "OnDemandContentInstaller.h"
-#include "OnDemandIoDispatcherBackend.h"
 #include "OnDemandPackageStoreBackend.h"
 #include "ProfilingDebugging/AssetMetadataTrace.h"
 #include "Serialization/MemoryReader.h"
@@ -52,15 +50,6 @@
 #include "String/LexFromString.h"
 #endif
 
-///////////////////////////////////////////////////////////////////////////////
-bool GHttpIoDispatcherEnabled = true;
-static FAutoConsoleVariableRef CVar_HttpIoDispatcherEnabled(
-	TEXT("iax.HttpIoDispatcherEnabled"),
-	GHttpIoDispatcherEnabled,
-	TEXT("Whether the HTTP I/O dispatcher should be enabled."),
-	ECVF_SaveForNextBoot
-);
-
 ///////////////////////////////////////////////////////////////////////////////
 namespace UE::IoStore
 {
@@ -223,8 +212,6 @@ FOnDemandIoStore::~FOnDemandIoStore()
 	}
 
 	Installer.Reset();
-	HttpClient.Reset();
-
 	if (FHttpIoDispatcher::IsInitialized())
 	{
 		FHttpIoDispatcher::OnHostGroupRegistered().RemoveAll(this);
@@ -343,11 +330,6 @@ FIoStatus FOnDemandIoStore::Initialize()
 		UE_LOG(LogIoStoreOnDemand, Log, TEXT("Install cache disabled"));
 	}
 
-	const bool bHttpIoDispatcherEnabled = (GHttpIoDispatcherEnabled || FParse::Param(FCommandLine::Get(), TEXT("HttpIoDispatcher")))
-		&& !FParse::Param(FCommandLine::Get(), TEXT("NoHttpIoDispatcher"));
-	FOnDemandIoBackendStats::Get()->SetHttpIoDispatcherEnabled(bHttpIoDispatcherEnabled);
-
-	if (bHttpIoDispatcherEnabled)
 	{
 		TUniquePtr<IIasCache> HttpCache;
 		FIasCacheConfig CacheConfig = Config::GetStreamingCacheConfig(FCommandLine::Get());
@@ -389,12 +371,8 @@ FIoStatus FOnDemandIoStore::Initialize()
 		FIoDispatcher::Get().Mount(HttpIoBackend.ToSharedRef(), HttpBackendPriority);
 		FHttpIoDispatcher::OnHostGroupRegistered().AddRaw(this, &FOnDemandIoStore::OnHostGroupRegistered);
 	}
-	else
-	{
-		HttpClient = MakeUnique<FOnDemandHttpThread>();
-	}
 
-	Installer = MakeUnique<FOnDemandContentInstaller>(*this, HttpClient.Get());
+	Installer = MakeUnique<FOnDemandContentInstaller>(*this);
 
 	return InstallCacheStatus;
 }
@@ -403,8 +381,8 @@ FIoStatus FOnDemandIoStore::InitializePostHotfix()
 {
 	const TCHAR* CommandLine = FCommandLine::Get();
 
-	TIoStatusOr<FOnDemandEndpointConfig> EndpointConfigStatus = Config::TryParseEndpointConfig(CommandLine);
-	if (EndpointConfigStatus.IsOk() == false)
+	TIoStatusOr<FOnDemandIoStoreConfig> ConfigStatus = Config::TryParseConfig(CommandLine);
+	if (ConfigStatus.IsOk() == false)
 	{
 		// Temporary solution: Delete the ias cache directory if the game hasn't been configured to use IAS
 		if (FPaths::HasProjectPersistentDownloadDir())
@@ -420,16 +398,16 @@ FIoStatus FOnDemandIoStore::InitializePostHotfix()
 			}
 		}
 
-		return EndpointConfigStatus.Status();
+		return ConfigStatus.Status();
 	}
 
-	FOnDemandEndpointConfig EndpointConfig = EndpointConfigStatus.ConsumeValueOrDie();
+	FOnDemandIoStoreConfig Config = ConfigStatus.ConsumeValueOrDie();
 
 	{
 		bool bFirstAttempt = false;
 		if (FIoStatus CertStatus = LoadDefaultHttpCertificates(bFirstAttempt); !CertStatus.IsOk())
 		{
-			if (EndpointConfig.RequiresTls())
+			if (Config.RequiresTls())
 			{
 				FIoStatus Status = FIoStatusBuilder(EIoErrorCode::ReadError)
 					<< TEXT("Endpoint requires TLS but certificates failed to load due to '")
@@ -440,14 +418,6 @@ FIoStatus FOnDemandIoStore::InitializePostHotfix()
 		}
 	}
 
-	if (FHttpIoDispatcher::IsInitialized() == false)
-	{
-		if (FIoStatus Status = InitializeStreamingBackend(EndpointConfig); Status.IsOk() == false)
-		{
-			return Status;
-		}
-	}
-
 	bool bUsePerContainerTocsConfigValue = false;
 	if (GConfig)
 	{
@@ -464,12 +434,12 @@ FIoStatus FOnDemandIoStore::InitializePostHotfix()
 #endif
 
 	TOptional<FOnDemandMountArgs> MountArgs;
-	if (EndpointConfig.TocFilePath.IsEmpty() == false)
+	if (Config.TocFilePath.IsEmpty() == false)
 	{
 		FOnDemandHostGroup HostGroup;
-		if (EndpointConfig.ServiceUrls.IsEmpty() == false)
+		if (Config.ServiceUrls.IsEmpty() == false)
 		{
-			TIoStatusOr<FOnDemandHostGroup> HostGroupStatus = FOnDemandHostGroup::Create(EndpointConfig.ServiceUrls);
+			TIoStatusOr<FOnDemandHostGroup> HostGroupStatus = FOnDemandHostGroup::Create(Config.ServiceUrls);
 			if (HostGroupStatus.IsOk())
 			{
 				HostGroup = HostGroupStatus.ConsumeValueOrDie();
@@ -480,24 +450,24 @@ FIoStatus FOnDemandIoStore::InitializePostHotfix()
 		{
 			MountArgs.Emplace(FOnDemandMountArgs
 			{
-				.MountId = EndpointConfig.TocFilePath,
-				.FilePath = EndpointConfig.TocFilePath,
+				.MountId = Config.TocFilePath,
+				.FilePath = Config.TocFilePath,
 				.HostGroup = HostGroup,
 				.Options = EOnDemandMountOptions::StreamOnDemand
 			});
 		}
 	}
-	else if (!EndpointConfig.ServiceUrls.IsEmpty() && !EndpointConfig.TocPath.IsEmpty())
+	else if (!Config.ServiceUrls.IsEmpty() && !Config.TocPath.IsEmpty())
 	{
-		TIoStatusOr<FOnDemandHostGroup> HostGroup = FOnDemandHostGroup::Create(EndpointConfig.ServiceUrls);
+		TIoStatusOr<FOnDemandHostGroup> HostGroup = FOnDemandHostGroup::Create(Config.ServiceUrls);
 		if (HostGroup.IsOk() == false)
 		{
 			return HostGroup.Status();
 		}
 		MountArgs.Emplace(FOnDemandMountArgs
 		{
-			.MountId = EndpointConfig.ServiceUrls[0] / EndpointConfig.TocPath,
-			.TocRelativeUrl = EndpointConfig.TocPath,
+			.MountId = Config.ServiceUrls[0] / Config.TocPath,
+			.TocRelativeUrl = Config.TocPath,
 			.HostGroup = HostGroup.ConsumeValueOrDie(),
 			.Options = EOnDemandMountOptions::StreamOnDemand
 		});
@@ -507,14 +477,14 @@ FIoStatus FOnDemandIoStore::InitializePostHotfix()
 	TOptional<FOnDemandInstallArgs> InstallArgs;
 	if (FParse::Param(FCommandLine::Get(), TEXT("Iad.InstallAllContent")))
 	{
-		TIoStatusOr<FOnDemandHostGroup> HostGroup = FOnDemandHostGroup::Create(EndpointConfig.ServiceUrls);
+		TIoStatusOr<FOnDemandHostGroup> HostGroup = FOnDemandHostGroup::Create(Config.ServiceUrls);
 		check(HostGroup.IsOk());
 
 		if (MountArgs)
 		{
 			MountArgs.GetValue().HostGroup = HostGroup.ConsumeValueOrDie(); 
 			MountArgs.GetValue().Options = EOnDemandMountOptions::InstallOnDemand;
-			MountArgs.GetValue().TocRelativeUrl = EndpointConfig.TocPath;
+			MountArgs.GetValue().TocRelativeUrl = Config.TocPath;
 
 			static FOnDemandContentHandle ContentHandle = FOnDemandContentHandle::Create(TEXT("AllContent"));
 			InstallArgs.Emplace();
@@ -548,52 +518,6 @@ FIoStatus FOnDemandIoStore::InitializePostHotfix()
 	return FIoStatus::Ok;
 }
 
-FIoStatus FOnDemandIoStore::InitializeStreamingBackend(const FOnDemandEndpointConfig& EndpointConfig)
-{
-	TUniquePtr<IIasCache> StreamingCache;
-	FIasCacheConfig CacheConfig = Config::GetStreamingCacheConfig(FCommandLine::Get());
-
-	if (CacheConfig.DiskQuota > 0)
-	{
-		if (FPaths::HasProjectPersistentDownloadDir())
-		{
-			FString CacheDir = FPaths::ProjectPersistentDownloadDir();
-			StreamingCache = MakeIasCache(*CacheDir, CacheConfig, DiskCacheGovernor);
-
-			UE_CLOG(!StreamingCache.IsValid(), LogIas, Warning, TEXT("File cache disabled - streaming only (init-fail)"));
-		}
-		else
-		{
-			UE_LOG(LogIas, Warning, TEXT("File cache disabled - streaming only (project has no persistent download dir enabled for this platform)"));
-		}
-	}
-	else
-	{
-		UE_LOG(LogIas, Log, TEXT("File cache disabled - streaming only (zero-quota)"));
-	}
-
-	check(HttpClient.IsValid());
-	StreamingBackend = MakeOnDemandIoDispatcherBackend(EndpointConfig, *this, *HttpClient, MoveTemp(StreamingCache));
-	if (StreamingBackend.IsValid() == false)
-	{
-		return FIoStatus(EIoErrorCode::InvalidParameter, TEXT("Failed to create streaming backend"));
-	}
-
-	StreamingBackend->SetBulkOptionalEnabled(EnumHasAnyFlags(StreamingOptions, EOnDemandStreamingOptions::OptionalBulkDataDisabled) == false);
-
-	int32 BackendPriority = -10;
-#if !UE_BUILD_SHIPPING
-	if (FParse::Param(FCommandLine::Get(), TEXT("Ias")))
-	{
-		UE_LOG(LogIas, Log, TEXT("Setting streaming backend priority higher than file system backend"));
-		BackendPriority = 10;
-	}
-#endif
-	FIoDispatcher::Get().Mount(StreamingBackend.ToSharedRef(), BackendPriority);
-
-	return FIoStatus::Ok;
-}
-
 FOnDemandRegisterHostGroupResult FOnDemandIoStore::RegisterHostGroup(FOnDemandRegisterHostGroupArgs&& Args)
 {
 	using namespace UE::UnifiedError::Core;
@@ -637,28 +561,7 @@ FOnDemandRegisterHostGroupResult FOnDemandIoStore::RegisterHostGroup(FOnDemandRe
 	TArray<FAnsiString> HostNames	= SanitizeHostNames(Args.HostNames);
 	FAnsiString AnsiTestUrl			= FAnsiString(StringCast<ANSICHAR>(*Args.TestUrl));
 
-	FIoStatus Status = FIoStatus::Invalid;
-	FIASHostGroup RegisteredHostGroup;
-
-	if (FHttpIoDispatcher::IsInitialized())
-	{
-		Status				= FHttpIoDispatcher::RegisterHostGroup(Args.HostGroupName, HostNames, AnsiTestUrl); 
-		RegisteredHostGroup = FHostGroupManager::Get().Find(Args.HostGroupName); // The HTTP I/O dispatcher currenlty uses the host group manager
-	}
-	else
-	{
-		if (TIoStatusOr<FIASHostGroup> RegisterStatus = FHostGroupManager::Get().Register(Args.HostGroupName, AnsiTestUrl); RegisterStatus.IsOk())
-		{
-			RegisteredHostGroup = RegisterStatus.ConsumeValueOrDie();
-			Status				= RegisteredHostGroup.Resolve(HostNames);
-		}
-		else
-		{
-			Status = RegisterStatus.Status();
-		}
-	}
-
-	if (Status.IsOk() == false)
+	if (FIoStatus Status = FHttpIoDispatcher::RegisterHostGroup(Args.HostGroupName, HostNames, AnsiTestUrl); !Status.IsOk())
 	{
 		return FOnDemandRegisterHostGroupResult
 		{
@@ -666,6 +569,7 @@ FOnDemandRegisterHostGroupResult FOnDemandIoStore::RegisterHostGroup(FOnDemandRe
 		};
 	}
 
+	FIASHostGroup RegisteredHostGroup = FHostGroupManager::Get().Find(Args.HostGroupName); // The HTTP I/O dispatcher currenlty uses the host group manager
 	return FOnDemandRegisterHostGroupResult
 	{
 		.HostGroup = RegisteredHostGroup.GetUnderlyingHostGroup()
@@ -1793,106 +1697,23 @@ FIoStatus FOnDemandIoStore::SetupHostGroup(const FSharedOnDemandContainer& Conta
 		return OutHostNames;
 	};
 
-	if (FHttpIoDispatcher::IsInitialized())
-	{
-		UE_CLOG(Container->HostGroupName != MountArgs.HostGroupName, LogIoStoreOnDemand, Log,
-			TEXT("Overriding container host group '%s' with '%s' from mount option"),
-			*Container->HostGroupName.ToString(), *MountArgs.HostGroupName.ToString());
-
-		Container->HostGroupName = MountArgs.HostGroupName;
-		if (FHttpIoDispatcher::IsHostGroupRegistered(Container->HostGroupName))
-		{
-			return FIoStatus::Ok;
-		}
-
-		if (MountArgs.HostGroup.IsEmpty() == false)
-		{
-			TArray<FAnsiString> HostNames = SanitizeHostNames(MountArgs.HostGroup.Hosts());
-			return FHttpIoDispatcher::RegisterHostGroup(Container->HostGroupName, HostNames, Container->GetTestUrl());
-		}
-
-		return FIoStatus(EIoErrorCode::PendingHostGroup); 
-	}
-
-	// This path will be removed once the HTTP I/O dispatcher is on by default
-	// StreamOnDemand (IAS) uses a different default hostgroup in this code path so if the MountArgs are requesting 'Default' we need
-	// to route that to the IAS default instead.
-	if (EnumHasAnyFlags(Container->Flags, EOnDemandContainerFlags::StreamOnDemand) && MountArgs.HostGroupName == TEXT("Default"))
-	{
-		//TODO: Remove this special case for streaming containers
-		Container->HostGroupName	= FName("DefaultOnDemand");
-		Container->HostGroup		= FHostGroupManager::Get().Find(Container->HostGroupName);
-
-		return FIoStatus::Ok;
-	}
-
 	UE_CLOG(Container->HostGroupName != MountArgs.HostGroupName, LogIoStoreOnDemand, Log,
 		TEXT("Overriding container host group '%s' with '%s' from mount option"),
 		*Container->HostGroupName.ToString(), *MountArgs.HostGroupName.ToString());
 
-	Container->HostGroupName	= MountArgs.HostGroupName;
-	Container->HostGroup		= FHostGroupManager::Get().Find(Container->HostGroupName);
-
-	if (Container->HostGroup.IsResolved())
+	Container->HostGroupName = MountArgs.HostGroupName;
+	if (FHttpIoDispatcher::IsHostGroupRegistered(Container->HostGroupName))
 	{
 		return FIoStatus::Ok;
 	}
 
-	if (MountArgs.HostGroup.IsEmpty())
-	{
-		return FIoStatus(EIoErrorCode::PendingHostGroup);
-	}
-	
-	TIoStatusOr<FIASHostGroup> NewHostGroup = FHostGroupManager::Get().Register(Container->HostGroupName, Container->GetTestUrl());
-	if (NewHostGroup.IsOk() == false)
-	{
-		return NewHostGroup.Status();
-	}
-
-	Container->HostGroup = NewHostGroup.ConsumeValueOrDie();
-
-	TArray<FAnsiString> HostNames = SanitizeHostNames(MountArgs.HostGroup.Hosts());
-#if !UE_BUILD_SHIPPING
-	if (Container->HostGroupName == FOnDemandHostGroup::DefaultName)
-	{
-		FString HostValue;
-		if (FParse::Value(FCommandLine::Get(), TEXT("Iax.DefaultHostGroup="), HostValue))
-		{
-			HostNames.Reset();
-			TArray<FString> HostArray;
-			HostValue.ParseIntoArray(HostArray, TEXT(","), true);
-			for (FString& Host : HostArray)
-			{
-				Host.TrimStartAndEndInline();
-				HostNames.Add(FAnsiString(*Host));
-			}
-		}
-	}
-#endif
-	if (FIoStatus Status = Container->HostGroup.Resolve(HostNames); Status.IsOk() == false)
-	{
-		return Status;
-	}
-
+	if (MountArgs.HostGroup.IsEmpty() == false)
 	{
-		const FOnDemandHostGroup& HostGroup = Container->HostGroup.GetUnderlyingHostGroup();
-		TConstArrayView<FAnsiString> Hosts	= HostGroup.Hosts();
-		TStringBuilder<128> Sb;
-
-		UE_LOG(LogIoStoreOnDemand, Log, TEXT("Created new host group '%s'"), *Container->HostGroupName.ToString());
-		for (int32 HostIndex = 0; HostIndex < Hosts.Num(); ++HostIndex) 
-		{
-			Sb.Reset();
-			Sb << Hosts[HostIndex];
-			if (HostIndex == HostGroup.PrimaryHostIndex())
-			{
-				Sb << TEXT(" (primary)");
-			}
-			UE_LOG(LogIoStoreOnDemand, Log, TEXT("\t%s"), Sb.ToString()); 
-		}
+		TArray<FAnsiString> HostNames = SanitizeHostNames(MountArgs.HostGroup.Hosts());
+		return FHttpIoDispatcher::RegisterHostGroup(Container->HostGroupName, HostNames, Container->GetTestUrl());
 	}
 
-	return FIoStatus::Ok;
+	return FIoStatus(EIoErrorCode::PendingHostGroup); 
 }
 
 TArray<FSharedOnDemandContainer> FOnDemandIoStore::GetContainers(EOnDemandContainerFlags ContainerFlags) const
@@ -2230,10 +2051,7 @@ FOnDemandCacheUsage FOnDemandIoStore::GetCacheUsage(const FOnDemandGetCacheUsage
 		}
 	}
 
-	if (StreamingBackend.IsValid())
-	{
-		CacheUsage.StreamingCache = StreamingBackend->GetCacheUsage();
-	}
+	FOnDemandIoBackendStats::Get()->GetIasCacheStats(CacheUsage.StreamingCache.TotalSize, CacheUsage.StreamingCache.MaxSize);
 
 	return CacheUsage;
 }
@@ -2302,17 +2120,12 @@ void FOnDemandIoStore::DumpMountedContainersToLog() const
 
 bool FOnDemandIoStore::IsOnDemandStreamingEnabled() const
 {
-	return StreamingBackend != nullptr;
+	return HttpIoBackend.IsValid();
 }
 
 void FOnDemandIoStore::SetStreamingOptions(EOnDemandStreamingOptions Options) 
 {
 	StreamingOptions = Options;
-	if (StreamingBackend.IsValid())
-	{
-		StreamingBackend->SetBulkOptionalEnabled(EnumHasAnyFlags(Options, EOnDemandStreamingOptions::OptionalBulkDataDisabled) == false);
-	}
-
 	if (HttpIoBackend.IsValid())
 	{
 		HttpIoBackend->SetOptionalBulkDataEnabled(EnumHasAnyFlags(Options, EOnDemandStreamingOptions::OptionalBulkDataDisabled) == false);
@@ -2321,14 +2134,7 @@ void FOnDemandIoStore::SetStreamingOptions(EOnDemandStreamingOptions Options)
 
 void FOnDemandIoStore::GetHttpStats(FOnDemandHttpStats& Out) const
 {
-	if (FHttpIoDispatcher::IsInitialized())
-	{
-		FHttpIoDispatcher::GetHttpStats(Out);
-	}
-	else
-	{
-		HttpClient->GetHttpStats(Out);
-	}
+	FHttpIoDispatcher::GetHttpStats(Out);
 }
 
 void FOnDemandIoStore::ReportAnalytics(TArray<FAnalyticsEventAttribute>& OutAnalyticsArray) const
@@ -2338,11 +2144,6 @@ void FOnDemandIoStore::ReportAnalytics(TArray<FAnalyticsEventAttribute>& OutAnal
 		Installer->ReportAnalytics(OutAnalyticsArray);
 	}
 
-	if (StreamingBackend.IsValid())
-	{
-		return StreamingBackend->ReportAnalytics(OutAnalyticsArray);
-	}
-
 	if (HttpIoBackend.IsValid())
 	{
 		HttpIoBackend->ReportAnalytics(OutAnalyticsArray);
@@ -2351,11 +2152,6 @@ void FOnDemandIoStore::ReportAnalytics(TArray<FAnalyticsEventAttribute>& OutAnal
 
 TUniquePtr<IAnalyticsRecording> FOnDemandIoStore::StartAnalyticsRecording() const
 {
-	if (StreamingBackend.IsValid())
-	{
-		return StreamingBackend->StartAnalyticsRecording();
-	}
-
 	return TUniquePtr<IAnalyticsRecording>();
 }
 
diff --git a/Engine/Source/Runtime/Experimental/IoStore/OnDemand/Private/OnDemandIoStore.h b/Engine/Source/Runtime/Experimental/IoStore/OnDemand/Private/OnDemandIoStore.h
--- a/Engine/Source/Runtime/Experimental/IoStore/OnDemand/Private/OnDemandIoStore.h
+++ b/Engine/Source/Runtime/Experimental/IoStore/OnDemand/Private/OnDemandIoStore.h
@@ -30,11 +30,8 @@ namespace UE::IoStore
 
 class FOnDemandDebugCommands;
 class FOnDemandContentInstaller;
-class FOnDemandHttpThread;
-struct FOnDemandEndpointConfig;
 enum class ELogOnDemandCacheUsage : uint8;
 
-using FSharedStreamingBackend		= TSharedPtr<class IOnDemandIoDispatcherBackend>;
 using FSharedPackageStoreBackend	= TSharedPtr<class IOnDemandPackageStoreBackend>;
 using FSharedInstallCache			= TSharedPtr<class IOnDemandInstallCache>;
 using FSharedHttpIoBackend			= TSharedPtr<class IOnDemandHttpIoDispatcherBackend>;
@@ -360,8 +357,6 @@ public:
 										FOnDemandIoStore();
 										~FOnDemandIoStore();
 
-	FIoStatus							InitializeStreamingBackend(const FOnDemandEndpointConfig& EndpointConfig);
-
 	FOnDemandChunkInfo					GetStreamingChunkInfo(const FIoChunkId& ChunkId);
 	FOnDemandChunkInfo					GetInstalledChunkInfo(const FIoChunkId& ChunkId, EIoErrorCode& OutErrorCode);
 
@@ -436,11 +431,9 @@ private:
 	TUniquePtr<FOnDemandDebugCommands>		DebugCommands;
 #endif // !UE_BUILD_SHIPPING
 
-	TUniquePtr<FOnDemandHttpThread>			HttpClient;
 	TUniquePtr<FOnDemandContentInstaller>	Installer;
 	FSharedInstallCache						InstallCache;
 	FSharedPackageStoreBackend				PackageStoreBackend;
-	FSharedStreamingBackend					StreamingBackend;
 	FSharedHttpIoBackend					HttpIoBackend;
 	FDiskCacheGovernor						DiskCacheGovernor;
 
diff --git a/Engine/Source/Runtime/Experimental/IoStore/OnDemand/Private/Statistics.cpp b/Engine/Source/Runtime/Experimental/IoStore/OnDemand/Private/Statistics.cpp
--- a/Engine/Source/Runtime/Experimental/IoStore/OnDemand/Private/Statistics.cpp
+++ b/Engine/Source/Runtime/Experimental/IoStore/OnDemand/Private/Statistics.cpp
@@ -15,7 +15,7 @@
 #include "Internationalization/Internationalization.h"
 #include "Math/UnrealMathUtility.h"
 #include "Misc/CoreDelegates.h"
-#include "OnDemandHttpThread.h"
+#include "OnDemandHttpIoDispatcher.h"
 #include "Templates/Requires.h"
 
 LLM_DEFINE_TAG(Ias);
@@ -438,8 +438,6 @@ struct FHttpRecentHistoryStatistics
 
 FHttpRecentHistoryStatistics GHttpHistory[(uint8)EHttpRequestType::NUM_SOURCES];
 
-// Debug flag indicating which HTTP backend is running
-static bool				GHttpIoDispatcherEnabled = false;
 // Experimental Http Stats
 static uint32			GHttpCdnCacheHit = 0;
 static uint32			GHttpCdnCacheMiss = 0;
@@ -770,11 +768,9 @@ void FOnDemandIoBackendStats::PrintOnScreenStatistics(FCoreDelegates::FSeverityM
 		TStringBuilder<256> BackendStatusText;
 		BackendStatusText << TEXT("IAS Backend Status: ");
 		BackendStatus.ToString(BackendStatusText);
-		BackendStatusText.Appendf(TEXT(" | HTTP I/O dispatcher: %s"), GHttpIoDispatcherEnabled ? TEXT("Enabled") : TEXT("Disabled"));
 
-		if (GHttpIoDispatcherEnabled)
 		{
-			BackendStatusText.Append(" | MaxRange: ");
+			BackendStatusText.Append(" | HttpMaxRange: ");
 
 			const int32 MaxRange = GHttpIoDispatcherMaxRangeKiB;
 			if (MaxRange < 0)
@@ -1020,11 +1016,7 @@ void FOnDemandIoBackendStats::ReportEndPointAnalytics(TArray<FAnalyticsEventAttr
 			TtfbMs = Stats.GetTimeToFirstByteMs();
 		}
 
-		AppendAnalyticsEventAttributeArray(OutAnalyticsArray
-			, TEXT("IaxTtfbMs"), TtfbMs
-			, TEXT("HttpIoDispatcherEnabled"), GHttpIoDispatcherEnabled
-		);
-
+		AppendAnalyticsEventAttributeArray(OutAnalyticsArray, TEXT("IaxTtfbMs"), TtfbMs);
 		ReportHttpStats(TEXT("Ias"), EHttpRequestType::Streaming);
 		ReportHttpStats(TEXT("Iad"), EHttpRequestType::Installed);
 	}
@@ -1436,11 +1428,6 @@ void FOnDemandIoBackendStats::OnApplicationResume()
 	++GAppResumeCount;
 }
 
-void FOnDemandIoBackendStats::SetHttpIoDispatcherEnabled(bool bEnabled)
-{
-	GHttpIoDispatcherEnabled = bEnabled;
-}
-
 ////////////////////////////////////////////////////////////////////////////////
 void FOnDemandContentInstallerStats::OnRequestEnqueued()
 {
diff --git a/Engine/Source/Runtime/Experimental/IoStore/OnDemand/Private/Statistics.h b/Engine/Source/Runtime/Experimental/IoStore/OnDemand/Private/Statistics.h
--- a/Engine/Source/Runtime/Experimental/IoStore/OnDemand/Private/Statistics.h
+++ b/Engine/Source/Runtime/Experimental/IoStore/OnDemand/Private/Statistics.h
@@ -112,7 +112,6 @@ public:
 	void OnHttpRequestStarted() IAS_STATISTICS_IMPL()
 	void OnHttpRequestCompleted() IAS_STATISTICS_IMPL()
 
-	void SetHttpIoDispatcherEnabled(bool bEnabled) IAS_STATISTICS_IMPL()
 private:
 
 #if IAS_WITH_STATISTICS
