Chaos - fix friction limit in split impulse mode (WIP)

#rnx

#Synced-CL 48180896

diff --git a/Engine/Source/Runtime/Experimental/Chaos/Private/Chaos/Collision/PBDCollisionContainerSolver.cpp b/Engine/Source/Runtime/Experimental/Chaos/Private/Chaos/Collision/PBDCollisionContainerSolver.cpp
--- a/Engine/Source/Runtime/Experimental/Chaos/Private/Chaos/Collision/PBDCollisionContainerSolver.cpp
+++ b/Engine/Source/Runtime/Experimental/Chaos/Private/Chaos/Collision/PBDCollisionContainerSolver.cpp
@@ -144,8 +144,8 @@ namespace Chaos
 		// The contact point error we are trying to correct in this solver
 		const FVec3f WorldContactDelta = FVec3f(WorldContact0 - WorldContact1);
 		FRealSingle WorldContactDeltaNormal = FVec3f::DotProduct(WorldContactDelta, WorldContactNormal);
-		const FRealSingle WorldContactDeltaTangentU = FVec3f::DotProduct(WorldContactDelta + WorldFrictionDelta, WorldContactTangentU);
-		const FRealSingle WorldContactDeltaTangentV = FVec3f::DotProduct(WorldContactDelta + WorldFrictionDelta, WorldContactTangentV);
+		FRealSingle WorldContactDeltaTangentU = FVec3f::DotProduct(WorldContactDelta + WorldFrictionDelta, WorldContactTangentU);
+		FRealSingle WorldContactDeltaTangentV = FVec3f::DotProduct(WorldContactDelta + WorldFrictionDelta, WorldContactTangentV);
 
 		// The target contact velocity, taking restitution into account
 		FRealSingle WorldContactTargetVelocityNormal = FRealSingle(0);
@@ -157,6 +157,12 @@ namespace Chaos
 			}
 		}
 
+		if (ManifoldPoint.Flags.bInitialContact)
+		{
+			WorldContactDeltaTangentU = 0;
+			WorldContactDeltaTangentV = 0;
+		}
+
 		// Overlap remaining from the previous frame, estimated from current contact phi and velocity
 		FRealSingle WorldContactResidualPhi = FMath::Min(WorldContactDeltaNormal, FRealSingle(0)) - FMath::Min(ContactVelocityNormal * Dt, FRealSingle(0));
 
@@ -309,22 +315,10 @@ namespace Chaos
 		if (SolverSettings.NumPositionFrictionIterations > 0)
 		{
 			PositionStaticFriction = StaticFriction;
-
-			// We have an option to apply dynamic friction in the velocity solver phase for spheres and capsules. 
-			// @todo(chaos): UE5.4: Remove quadratic special case when PBD static friction is enabled
-			const bool bUsePBDDynamicFriction = !Constraint->HasQuadraticShape() || (SolverSettings.NumVelocityFrictionIterations == 0);
-			if (bUsePBDDynamicFriction)
-			{
-				PositionDynamicFriction = DynamicFriction;
-			}
-			else
-			{
-				VelocityDynamicFriction = DynamicFriction;
-			}
+			PositionDynamicFriction = DynamicFriction;
 		}
-		else if (SolverSettings.NumVelocityFrictionIterations > 0)
+		if (SolverSettings.NumVelocityFrictionIterations > 0)
 		{
-			// We have an option to run without static friction and velocity-based dynamic friction (for RBAN)
 			VelocityDynamicFriction = DynamicFriction;
 		}
 
diff --git a/Engine/Source/Runtime/Experimental/Chaos/Public/Chaos/Collision/PBDCollisionSolver.h b/Engine/Source/Runtime/Experimental/Chaos/Public/Chaos/Collision/PBDCollisionSolver.h
--- a/Engine/Source/Runtime/Experimental/Chaos/Public/Chaos/Collision/PBDCollisionSolver.h
+++ b/Engine/Source/Runtime/Experimental/Chaos/Public/Chaos/Collision/PBDCollisionSolver.h
@@ -596,6 +596,11 @@ namespace Chaos
 		{
 			FSolverReal PushOutNormal = FSolverReal(0);
 
+			// How much of this position correction should be momentumless?
+			// In SplitImpulse mode, only residual overlap corrections apply momentum.
+			FSolverReal SplitAlpha = ManifoldPoint.SplitImpulseAlpha;
+			FSolverReal SplitAlphaInv = FSolverReal(1) - ManifoldPoint.SplitImpulseAlpha;
+
 			CalculatePositionCorrectionNormal(
 				Stiffness,
 				ContactDeltaNormal,
@@ -603,7 +608,8 @@ namespace Chaos
 				ManifoldPoint.NetPushOutNormal,
 				PushOutNormal);						// Out
 
-			ManifoldPoint.NetPushOutNormal += PushOutNormal;
+
+			ManifoldPoint.NetPushOutNormal += SplitAlphaInv * PushOutNormal;
 
 			// Update the particle state based on the pushout
 			if (IsDynamic(0))
@@ -611,16 +617,16 @@ namespace Chaos
 				const FSolverVec3 DX0 = (State.InvMs[0] * PushOutNormal) * ManifoldPoint.ContactNormal;
 				const FSolverVec3 DR0 = ManifoldPoint.ContactNormalAngular0 * PushOutNormal;
 				FConstraintSolverBody& Body0 = SolverBody0();
-				if (ManifoldPoint.SplitImpulseAlpha > 0)
+				if (SplitAlpha > 0)
 				{
-					Body0.ApplyPositionCorrectionDelta(ManifoldPoint.SplitImpulseAlpha * DX0);
-					Body0.ApplyRotationCorrectionDelta(ManifoldPoint.SplitImpulseAlpha * DR0);
+					Body0.ApplyPositionCorrectionDelta(SplitAlpha * DX0);
+					Body0.ApplyRotationCorrectionDelta(SplitAlpha * DR0);
 				}
 
-				if (ManifoldPoint.SplitImpulseAlpha < 1)
+				if (SplitAlphaInv > 0)
 				{
-					Body0.ApplyPositionDelta((1.0f - ManifoldPoint.SplitImpulseAlpha) * DX0);
-					Body0.ApplyRotationDelta((1.0f - ManifoldPoint.SplitImpulseAlpha) * DR0);
+					Body0.ApplyPositionDelta(SplitAlphaInv * DX0);
+					Body0.ApplyRotationDelta(SplitAlphaInv * DR0);
 				}
 			}
 			if (IsDynamic(1))
@@ -628,16 +634,16 @@ namespace Chaos
 				const FSolverVec3 DX1 = (State.InvMs[1] * -PushOutNormal) * ManifoldPoint.ContactNormal;
 				const FSolverVec3 DR1 = ManifoldPoint.ContactNormalAngular1 * -PushOutNormal;
 				FConstraintSolverBody& Body1 = SolverBody1();
-				if (ManifoldPoint.SplitImpulseAlpha > 0)
+				if (SplitAlpha > 0)
 				{
-					Body1.ApplyPositionCorrectionDelta(ManifoldPoint.SplitImpulseAlpha * DX1);
-					Body1.ApplyRotationCorrectionDelta(ManifoldPoint.SplitImpulseAlpha * DR1);
+					Body1.ApplyPositionCorrectionDelta(SplitAlpha * DX1);
+					Body1.ApplyRotationCorrectionDelta(SplitAlpha * DR1);
 				}
 
-				if (ManifoldPoint.SplitImpulseAlpha < 1)
+				if (SplitAlphaInv > 0)
 				{
-					Body1.ApplyPositionDelta((1.0f - ManifoldPoint.SplitImpulseAlpha) * DX1);
-					Body1.ApplyRotationDelta((1.0f - ManifoldPoint.SplitImpulseAlpha) * DR1);
+					Body1.ApplyPositionDelta(SplitAlphaInv * DX1);
+					Body1.ApplyRotationDelta(SplitAlphaInv * DR1);
 				}
 			}
 		}
@@ -702,8 +708,7 @@ namespace Chaos
 					const FSolverReal ImpulseTangentScale = MaxImpulseTangent * FMath::InvSqrt(NetImpulseTangentSq);
 					ImpulseTangentU = (NetImpulseTangentU + ImpulseTangentU) * ImpulseTangentScale - NetImpulseTangentU;
 					ImpulseTangentV = (NetImpulseTangentV + ImpulseTangentV) * ImpulseTangentScale - NetImpulseTangentV;
-					// TODO: Should we update the StaticFrictionRatio if we clamp here? This is used to move the
-					// static friction anchors to the edge of the friction cone
+					ManifoldPoint.StaticFrictionRatio = ImpulseTangentScale;
 				}
 
 				ManifoldPoint.NetImpulseTangentU += ImpulseTangentU;
@@ -1012,6 +1017,7 @@ namespace Chaos
 			)
 		{
 			check(State.ManifoldPoints != nullptr);
+			check(PointIndex >= 0);
 			check(PointIndex < State.NumManifoldPoints);
 			check(InSplitImpulseAlpha >= 0.0f);
 			check(InSplitImpulseAlpha <= 1.0f);
@@ -1164,11 +1170,14 @@ namespace Chaos
 					continue;
 				}
 
-				// Solve for normal velocity only for points where we have a position correction
-				// (unless the position solve step is disabled, in which case we assume we want to run the velocity solve)
+				// Solver for normal velocity when...
+				// - we applied a positional impulse
+				// - we skipped the position solve, but not the velocity solve
+				// - we are in split impulse mode and have an overlap
 				const bool bShouldSolveNormalVelocity = 
 					(SolverManifoldPoint.NetPushOutNormal > FSolverReal(0)) 
-					|| !!SolverManifoldPoint.Flags.bDisablePositionSolve;
+					|| !!SolverManifoldPoint.Flags.bDisablePositionSolve
+					|| ((SolverManifoldPoint.SplitImpulseAlpha > 0) && (SolverManifoldPoint.ContactDeltaNormal < 0.0f));
 
 				// Solver for tangential velocity if we have non-zero friction and a position correction
 				// (unless the position solve step is disabled, in which case we assume we want to run the velocity solve)
