Chaos - convex-mesh collisions remap boundary edge collisions to face collisions (was already being done for boundary vertex collisions, but was missing for edges).

#rnx
#Synced-CL 48180896

diff --git a/Engine/Source/Runtime/Experimental/Chaos/Private/Chaos/Collision/MeshContactGenerator.cpp b/Engine/Source/Runtime/Experimental/Chaos/Private/Chaos/Collision/MeshContactGenerator.cpp
--- a/Engine/Source/Runtime/Experimental/Chaos/Private/Chaos/Collision/MeshContactGenerator.cpp
+++ b/Engine/Source/Runtime/Experimental/Chaos/Private/Chaos/Collision/MeshContactGenerator.cpp
@@ -16,6 +16,12 @@ namespace Chaos::CVars
 
 	bool bMeshContactGeneratorFixContactNormalFixEnabled = true;
 	FAutoConsoleVariableRef CVarChaosMeshContactGeneratorFixContactNormalFixEnabled(TEXT("p.Chaos.MeshContactGenerator.FixContactNormal.FixEnabled"), bMeshContactGeneratorFixContactNormalFixEnabled, TEXT("Until new code path is well tested"));
+
+	bool bMeshContactGeneratorTreatBoundaryVertexAsFace = true;
+	FAutoConsoleVariableRef CVarChaosMeshContactGeneratorTreatBoundaryVertexAsFace(TEXT("p.Chaos.MeshContactGenerator.TreatBoundaryVertexAsFace"), bMeshContactGeneratorTreatBoundaryVertexAsFace, TEXT("Treat boundary vertex collisions as face collisions"));
+
+	bool bMeshContactGeneratorTreatBoundaryEdgeAsFace = true;
+	FAutoConsoleVariableRef CVarChaosMeshContactGeneratorTreatBoundaryEdgeAsFace(TEXT("p.Chaos.MeshContactGenerator.TreatBoundaryEdgeAsFace"), bMeshContactGeneratorTreatBoundaryEdgeAsFace, TEXT("Treat boundary edge collisions as face collisions"));
 }
 
 namespace Chaos::Private
@@ -289,10 +295,10 @@ namespace Chaos::Private
 		// on the triangles that share that edge or vertex.
 		if (ContactPointData.IsEdge())
 		{
-			// We have an edge collision. The contact normal must lie between the normals of the two faces using the edge
 			const int32 OtherLocalTriangleIndex = GetOtherTriangleIndexForEdge(LocalTriangleIndex, ContactPointData.GetEdgeID());
 			if (OtherLocalTriangleIndex != INDEX_NONE)
 			{
+				// This is a shared edge, so the contact normal must lie between the normals of the two faces using the edge.
 				const FTriangleExt& OtherTriangle = Triangles[OtherLocalTriangleIndex];
 				const FVec3& OtherTriangleNormal = OtherTriangle.GetNormal();
 
@@ -330,6 +336,14 @@ namespace Chaos::Private
 					ContactPointData.SetTriangleIndex(CorrectedTriangleIndex);
 				}
 			}
+			else if (CVars::bMeshContactGeneratorTreatBoundaryEdgeAsFace)
+			{
+				// This is a boundary edge, so remap the contact to a face contact.
+				// @todo(chaos): ideally we would have information about edge sharing between adjancent meshes, in which case
+				// we could perform normal correction in all cases.
+				ContactPoint.ShapeContactNormal = TriangleNormal;
+				ContactPoint.ShapeContactPoints[0] = ContactPoint.ShapeContactPoints[1] + ContactPoint.Phi * TriangleNormal;
+			}
 		}
 		else if (ContactPointData.IsVertex())
 		{
@@ -353,18 +367,21 @@ namespace Chaos::Private
 				{
 					const FTriangleExt& OtherTriangle = Triangles[OtherLocalTriangleIndex];
 
-					// We don't collide with boundary vertices. The vertex is a boundary vertex if any of the edges including
-					// the vertex are boundary edges (i.e., not shared between two triangles)
-					int32 OtherVertexIndexB, OtherVertexIndexC;
-					if (OtherTriangle.GetOtherVertexIDs(VertexIndexA, OtherVertexIndexB, OtherVertexIndexC))
+					// If we have a boundary vertex, remap the contact to a face context. The vertex is a boundary 
+					// vertex if any of the edges including the vertex are boundary edges (i.e., not shared between two triangles)
+					if (CVars::bMeshContactGeneratorTreatBoundaryVertexAsFace)
 					{
-						const FContactEdgeID EdgeAB = FContactEdgeID(VertexIndexA, OtherVertexIndexB);
-						const FContactEdgeID EdgeCA = FContactEdgeID(OtherVertexIndexC, VertexIndexA);
-						if (!IsSharedEdge(EdgeAB) || !IsSharedEdge(EdgeCA))
+						int32 OtherVertexIndexB, OtherVertexIndexC;
+						if (OtherTriangle.GetOtherVertexIDs(VertexIndexA, OtherVertexIndexB, OtherVertexIndexC))
 						{
-							ContactPoint.ShapeContactNormal = TriangleNormal;
-							ContactPoint.ShapeContactPoints[0] = ContactPoint.ShapeContactPoints[1] + ContactPoint.Phi * TriangleNormal;
-							break;
+							const FContactEdgeID EdgeAB = FContactEdgeID(VertexIndexA, OtherVertexIndexB);
+							const FContactEdgeID EdgeCA = FContactEdgeID(OtherVertexIndexC, VertexIndexA);
+							if (!IsSharedEdge(EdgeAB) || !IsSharedEdge(EdgeCA))
+							{
+								ContactPoint.ShapeContactNormal = TriangleNormal;
+								ContactPoint.ShapeContactPoints[0] = ContactPoint.ShapeContactPoints[1] + ContactPoint.Phi * TriangleNormal;
+								break;
+							}
 						}
 					}
 
