IAX defaults to H2.

#rnx

#Synced-CL 48026973

diff --git a/Engine/Source/Runtime/Experimental/IoStore/OnDemand/Private/OnDemandHttpIoDispacher.cpp b/Engine/Source/Runtime/Experimental/IoStore/OnDemand/Private/OnDemandHttpIoDispacher.cpp
--- a/Engine/Source/Runtime/Experimental/IoStore/OnDemand/Private/OnDemandHttpIoDispacher.cpp
+++ b/Engine/Source/Runtime/Experimental/IoStore/OnDemand/Private/OnDemandHttpIoDispacher.cpp
@@ -208,11 +208,12 @@ FAutoConsoleVariableRef CVarOnDemandBackendThreadPriority(
 	TEXT("Note that this is switchable at runtime"),
 	ECVF_Default);
 
-int32 GIaxHttpVersion = 1;
+int32 GIaxHttpVersion = 2;
 static FAutoConsoleVariableRef CVar_IaxHttpVersion(
 	TEXT("iax.HttpVersion"),
 	GIaxHttpVersion,
-	TEXT("Protocol version to use, either '1' or '2'")
+	TEXT("Protocol version to prefer, either '1' or '2'. The latter requires https."),
+	FConsoleVariableDelegate::CreateStatic(&OnHttpClientCVarChanged)
 );
 
 int32 GIasHttpConnectionCount = 4;
