Fix wrong merging of IAX max number of inlight request(s).

#rnx
#Synced-CL 48026973

diff --git a/Engine/Source/Runtime/Experimental/IoStore/OnDemand/Private/OnDemandHttpIoDispacher.cpp b/Engine/Source/Runtime/Experimental/IoStore/OnDemand/Private/OnDemandHttpIoDispacher.cpp
--- a/Engine/Source/Runtime/Experimental/IoStore/OnDemand/Private/OnDemandHttpIoDispacher.cpp
+++ b/Engine/Source/Runtime/Experimental/IoStore/OnDemand/Private/OnDemandHttpIoDispacher.cpp
@@ -224,11 +224,15 @@ static FAutoConsoleVariableRef CVar_IasHttpConnectionCount(
 	FConsoleVariableDelegate::CreateStatic(&OnHttpClientCVarChanged)
 );
 
-int32 GIaxHttpMaxInflight = 1;
+/* Note that the value of this cvar does not affect throttling. Setting it too low
+ * will prevent the CDN peer from reordering cache hits ahead of misses. Latency
+ * will also start to become apparent and negatively effect user experience. */
+int32 GIaxHttpMaxInflight = 8;
 static FAutoConsoleVariableRef CVar_IaxHttpMaxInflight(
 	TEXT("iax.HttpMaxInflight"),
 	GIaxHttpMaxInflight,
-	TEXT("Number of requests to issue at once per connection")
+	TEXT("Number of requests to issue at once per connection"),
+	FConsoleVariableDelegate::CreateStatic(&OnHttpClientCVarChanged)
 );
 
 int32 GIasHttpRecvBufKiB = -1;
