Make ISMPoolComponent only tick when necessary
#Synced-CL 48128650

diff --git a/Engine/Source/Runtime/Experimental/ISMPool/Private/ISMPool/ISMPoolComponent.cpp b/Engine/Source/Runtime/Experimental/ISMPool/Private/ISMPool/ISMPoolComponent.cpp
--- a/Engine/Source/Runtime/Experimental/ISMPool/Private/ISMPool/ISMPoolComponent.cpp
+++ b/Engine/Source/Runtime/Experimental/ISMPool/Private/ISMPool/ISMPoolComponent.cpp
@@ -10,6 +10,21 @@
 
 #include UE_INLINE_GENERATED_CPP_BY_NAME(ISMPoolComponent)
 
+static void ISMPoolComponent_EnableTicksForPoolManagedComponents(IConsoleVariable*)
+{
+	//exclude CDOs & garbage
+	constexpr EObjectFlags ExclusionFlags = RF_ClassDefaultObject;
+	constexpr EInternalObjectFlags InternalExclusionFlags = EInternalObjectFlags::Garbage;
+	
+	constexpr bool bIncludeDerivedClasses = true;
+
+	for (TObjectIterator<UISMPoolComponent> It(ExclusionFlags, bIncludeDerivedClasses, InternalExclusionFlags); It; ++It)
+	{
+		//non-pool managed components don't need their tick re-enabling
+		It->EnableComponentTickIfPoolManagementEnabled();
+	}
+}
+
 // Don't release ISM components when they empty, but keep them (and their scene proxy) alive.
 // This can remove the high cost associated with repeated registration, scene proxy creation and mesh draw command creation.
 // But it can also have a high memory overhead since the ISMs retain hard references to their static meshes.
@@ -17,7 +32,8 @@ static bool GComponentKeepAlive = false;
 FAutoConsoleVariableRef CVarISMPoolComponentKeepAlive(
 	TEXT("r.ISMPool.ComponentKeepAlive"),
 	GComponentKeepAlive,
-	TEXT("Keep ISM components alive when all their instances are removed."));
+	TEXT("Keep ISM components alive when all their instances are removed."),
+	FConsoleVariableDelegate::CreateStatic(ISMPoolComponent_EnableTicksForPoolManagedComponents));
 
 // Use a FreeList to enable recycling of ISM components.
 // ISM components aren't unregistered, but their scene proxy is destroyed.
@@ -29,7 +45,8 @@ static bool GComponentRecycle = true;
 FAutoConsoleVariableRef CVarISMPoolComponentRecycle(
 	TEXT("r.ISMPool.ComponentRecycle"),
 	GComponentRecycle,
-	TEXT("Recycle ISM components to a free list for reuse when all their instances are removed."));
+	TEXT("Recycle ISM components to a free list for reuse when all their instances are removed."),
+	FConsoleVariableDelegate::CreateStatic(ISMPoolComponent_EnableTicksForPoolManagedComponents));
 
 static bool GISMPoolClearComponentMeshOnRecycle = true;
 FAutoConsoleVariableRef CVarISMPoolClearComponentMeshOnRecycle(
@@ -43,7 +60,8 @@ static int32 GComponentFreeListTargetSize = 50;
 FAutoConsoleVariableRef CVarISMPoolComponentFreeListTargetSize(
 	TEXT("r.ISMPool.ComponentFreeListTargetSize"),
 	GComponentFreeListTargetSize,
-	TEXT("Target size for number of ISM components in the recycling free list."));
+	TEXT("Target size for number of ISM components in the recycling free list."),
+	FConsoleVariableDelegate::CreateStatic(ISMPoolComponent_EnableTicksForPoolManagedComponents));
 
 // Keep copies of all custom instance data for restoration on readding an instance.
 static bool GShadowCopyCustomData = false;
@@ -588,6 +606,11 @@ void FISMPool::UpdateAbsoluteTransforms(const FTransform& BaseTransform, EUpdate
 	}
 }
 
+bool FISMPool::RequiresTick() const
+{
+	return !PrellocationQueue.IsEmpty() || (FreeListISM.Num() > GComponentFreeListTargetSize);
+}
+
 void FISMPool::Tick(UISMPoolComponent* OwningComponent)
 {
 	// Recache component lifecycle state from cvar.
@@ -650,6 +673,7 @@ void UISMPoolComponent::TickComponent(float DeltaTime, enum ELevelTick TickType,
 {
 	Super::TickComponent(DeltaTime, TickType, ThisTickFunction);
 	Pool.Tick(this);
+	SetComponentTickEnabled(Pool.RequiresTick());
 }
 
 UISMPoolComponent::FMeshGroupId  UISMPoolComponent::CreateMeshGroup(bool bAllowPerInstanceRemoval)
@@ -709,10 +733,12 @@ bool UISMPoolComponent::BatchUpdateInstanceCustomData(FMeshGroupId MeshGroupId,
 void UISMPoolComponent::PreallocateMeshInstance(const FISMPoolStaticMeshInstance& MeshInstance)
 {
 	Pool.RequestPreallocateMeshInstance(MeshInstance);
+	SetComponentTickEnabled(Pool.RequiresTick());
 }
 
 void UISMPoolComponent::SetTickablePoolManagement(bool bEnablePoolManagement)
 {
+	bIsPoolManagementEnabled = bEnablePoolManagement;
 	if (!bEnablePoolManagement)
 	{
 		// Disable the keep alive and recycle pool management systems.
@@ -724,6 +750,14 @@ void UISMPoolComponent::SetTickablePoolManagement(bool bEnablePoolManagement)
 	PrimaryComponentTick.SetTickFunctionEnable(bEnablePoolManagement);
 }
 
+void UISMPoolComponent::EnableComponentTickIfPoolManagementEnabled()
+{
+	if (bIsPoolManagementEnabled)
+	{
+		SetComponentTickEnabled(true);
+	}
+}
+
 void UISMPoolComponent::SetOverrideTransformUpdates(bool bOverrideUpdates)
 {
 	Pool.bDisableBoundsAndTransformUpdate = bOverrideUpdates;
diff --git a/Engine/Source/Runtime/Experimental/ISMPool/Public/ISMPool/ISMPoolComponent.h b/Engine/Source/Runtime/Experimental/ISMPool/Public/ISMPool/ISMPoolComponent.h
--- a/Engine/Source/Runtime/Experimental/ISMPool/Public/ISMPool/ISMPoolComponent.h
+++ b/Engine/Source/Runtime/Experimental/ISMPool/Public/ISMPool/ISMPoolComponent.h
@@ -308,6 +308,8 @@ struct FISMPool
 
 	void UpdateAbsoluteTransforms(const FTransform& BaseTransform, EUpdateTransformFlags UpdateTransformFlags, ETeleportType Teleport);
 
+	bool RequiresTick() const;
+
 	/** Array of ISM objects. */
 	TArray<FISMPoolISM> ISMs;
 	/** Mapping from mesh description to ISMs array slot. */
@@ -382,10 +384,13 @@ public:
 
 	ISMPOOL_API void UpdateAbsoluteTransforms(const FTransform& BaseTransform, EUpdateTransformFlags UpdateTransformFlags, ETeleportType Teleport);
 
+	void EnableComponentTickIfPoolManagementEnabled();
+
 private:
 	uint32 NextMeshGroupId = 0;
 	TMap<FMeshGroupId, FISMPoolMeshGroup> MeshGroups;
 	FISMPool Pool;
+	bool bIsPoolManagementEnabled = true;
 
 	// Expose internals for debug draw support.
 	friend class UISMPoolDebugDrawComponent;
