Fix a very rare bug where it was possible for a cancellation request to be made to FMultiEndpointHttpClient after it was destroyed causing a deadlock.

#rnx

- Previously we were relying on it not mattering if an already completed request tries to 'cancel' as FMultiEndpointHttpClient will just early out, however if this cancel call comes after the FMultiEndpointHttpClient has been destroyed it's mutex is likely to be trash and will permantly lock so we need to be more careful.
-- The scenarios where this can happen are very unlikely.
- Add a new method 'OnRequestCompleted' to FCancellationToken allowing us to signal that the http request for the token has completed. Once we know that the request has completed we no longer need to try to cancel it when the token goes out of scope.
- Add a new method 'CompleteCancellationToken' to FHttpDispatcher that once a http request has completed will find the cancellation token for that request and signal that it is completed.
-- Note that all the FHttpRequest in a request chain share the same token.
- Move the code that applies the cancellation token to the request chain to it's own method 'ApplyCancellationToken' to make the code easier to read.
- FMultiEndpointHttpClient::ProcessRetryAttempts can no longer hold a single lock on TicketLookupMutex to prevent deadlocking if the request being retried has been canceled and so invoked it's callback.
-- Now we only hold the lock when checking or removing from the TicketLookupMap.
- Fix spelling mistake, 'cancelation' -> 'cancellation'

#Synced-CL 47893047

diff --git a/Engine/Source/Runtime/Experimental/IoStore/OnDemand/Private/OnDemandHttpClient.cpp b/Engine/Source/Runtime/Experimental/IoStore/OnDemand/Private/OnDemandHttpClient.cpp
--- a/Engine/Source/Runtime/Experimental/IoStore/OnDemand/Private/OnDemandHttpClient.cpp
+++ b/Engine/Source/Runtime/Experimental/IoStore/OnDemand/Private/OnDemandHttpClient.cpp
@@ -624,45 +624,52 @@ void FMultiEndpointHttpClient::ProcessRetryAttempts(uint32 TicketCount)
 	}
 
 	TRACE_CPUPROFILER_EVENT_SCOPE(FMultiEndpointHttpClient::ProcessRetryAttempts);
-	UE::TUniqueLock _(TicketLookupMutex);
+
+	auto IsTicketCanceled = [this](FHttpTicketId TicketId) -> bool
+	{
+		UE::TUniqueLock _(TicketLookupMutex);
+
+		FTicketInfo* TicketInfo = TicketLookupMap.Find(TicketId);
+		if (!ensureMsgf(TicketInfo != nullptr, TEXT("Could not find a valid FTicketInfo for FHttpTicketId %d"), TicketId))
+		{
+			// I think we can survive at this point by assuming that the request just needs to be retried.
+			return false;
+		}
+
+		checkf(TicketInfo->HttpTicket == 0, TEXT("FHttpTicketId %d is still in progress!"), TicketId);
+		return TicketInfo->bCancelRequested;
+	};
 
 	const int32 RequestCount = FMath::Min(Retries.Num(), int32(HTTP::FEventLoop::MaxActiveTickets - TicketCount));
 	for (int32 Idx = 0; Idx < RequestCount; Idx++)
 	{
 		FRequest& Request = Retries[Idx];
 
-		if (FTicketInfo* TicketInfo = TicketLookupMap.Find(Request.TicketId))
+		if (IsTicketCanceled(Request.TicketId) == false)
 		{
-			check(TicketInfo->HttpTicket == 0);
-
-			if (!TicketInfo->bCancelRequested)
+			IssueRequest(MoveTemp(Request));
+		}
+		else
+		{
+			FMultiEndpointHttpClientResponse Response
 			{
-				IssueRequest(MoveTemp(Request));
-			}
-			else
+				.Reason = TEXT("Canceled"),
+				.DurationMilliseconds = uint64(FPlatformTime::ToMilliseconds64(FPlatformTime::Cycles64() - Request.StartTime)),
+				.StatusCode = Request.StatusCode,
+				.RetryCount = Request.RetryCount - 1, // Reduce retry count by 1 as this retry attempt didn't really do anything
+				.HostIndex = Request.Host,
+				.bCanceled = true,
+				.CDNCacheStatus = Request.CDNCacheStatus
+			};
+
+			// The request is completed so we can remove it from the lookup map
 			{
-				FMultiEndpointHttpClientResponse Response
-				{
-					.Reason = TEXT("Canceled"),
-					.DurationMilliseconds = uint64(FPlatformTime::ToMilliseconds64(FPlatformTime::Cycles64() - Request.StartTime)),
-					.StatusCode = Request.StatusCode,
-					.RetryCount = Request.RetryCount -1, // Reduce retry count by 1 as this retry attempt didn't really do anything
-					.HostIndex = Request.Host,
-					.bCanceled = true,
-					.CDNCacheStatus = Request.CDNCacheStatus
-				};
-
-				// The request is completed so we can remove it from the lookup map
+				UE::TUniqueLock _(TicketLookupMutex);
 				TicketLookupMap.Remove(Request.TicketId);
-
-				// TODO: Consider moving this outside of the TicketLookupMutex lock?
-				FOnHttpResponse OnResponse = MoveTemp(Request.OnResponse);
-				OnResponse(MoveTemp(Response));	
 			}
-		}
-		else
-		{
-			checkNoEntry(); // If the request is ongoing then we should always have a valid entry
+
+			FOnHttpResponse OnResponse = MoveTemp(Request.OnResponse);
+			OnResponse(MoveTemp(Response));
 		}
 	}
 
diff --git a/Engine/Source/Runtime/Experimental/IoStore/OnDemand/Private/OnDemandHttpIoDispacher.cpp b/Engine/Source/Runtime/Experimental/IoStore/OnDemand/Private/OnDemandHttpIoDispacher.cpp
--- a/Engine/Source/Runtime/Experimental/IoStore/OnDemand/Private/OnDemandHttpIoDispacher.cpp
+++ b/Engine/Source/Runtime/Experimental/IoStore/OnDemand/Private/OnDemandHttpIoDispacher.cpp
@@ -348,20 +348,26 @@ namespace HttpIoDispatcher
  * Tracks the lifespan of a http request for as long as we want to keep it active. When the token is
  * destroyed we will try to cancel the request if it is still in progress.
  */
-class FCancelationToken
+class FCancellationToken
 {
 public:
-	FCancelationToken(FMultiEndpointHttpClient& InClient, FMultiEndpointHttpClient::FHttpTicketId InTicketId)
+	FCancellationToken(FMultiEndpointHttpClient& InClient, FMultiEndpointHttpClient::FHttpTicketId InTicketId)
 		: Client(InClient)
 		, TicketId(InTicketId)
 	{
+	}
 
+	~FCancellationToken()
+	{
+		if (TicketId != 0)
+		{
+			Client.CancelRequest(TicketId);
+		}
 	}
 
-	~FCancelationToken()
+	void OnRequestCompleted()
 	{
-		// Note that if the request has already completed then this will do nothing.
-		Client.CancelRequest(TicketId);
+		TicketId = 0;
 	}
 
 private:
@@ -456,7 +462,7 @@ struct FHttpRequest final
 	std::atomic<EIoErrorCode>	CompletionStatus{EIoErrorCode::Unknown};
 	uint32						SeqNo = ++NextSeqNo;
 
-	TSharedPtr<FCancelationToken> CancelationToken;
+	TSharedPtr<FCancellationToken> CancellationToken;
 };
 using FHttpRequestList = TIntrusiveList<FHttpRequest>;
 
@@ -1035,6 +1041,9 @@ private:
 										CompleteRequest(Request, TArray<FAnsiString>(), Body, TotalRange, StatusCode, bCached);
 									}
 
+	void							ApplyCancellationToken(FHttpRequest* RequestChain, FMultiEndpointHttpClient::FHttpTicketId TicketId);
+	void							CompleteCancellationToken(FHttpRequest* RequestChain);
+
 #if UE_ALLOW_HTTP_PAUSE
 	void OnPauseCommand(bool bShouldPause, const TArray<FString>& Args, FOutputDevice& Ar);
 #endif //UE_ALLOW_HTTP_PAUSE
@@ -1046,8 +1055,8 @@ private:
 	FHttpQueue								HttpQueue;
 	FHostGroupRegistered					HostGroupRegisterd;
 	FEventRef								WakeUp;
-	/** Keeps assignment of the CancelationToken safe between ::CancelRequest and ::ProcessHttpRequests */
-	FMutex									CancelationMutex;
+	/** Keeps assignment of the CancellationToken safe between ::CancelRequest and ::ProcessHttpRequests */
+	FMutex									CancellationMutex;
 	FDelegateHandle							OnRecreateHttpClientHandle;
 	std::atomic_bool						bRecreateHttpClient = false;
 	std::atomic_bool						bStopRequested = false;
@@ -1384,7 +1393,7 @@ void FHttpDispatcher::CancelRequest(FIoHttpRequestHandle Handle)
 	}
 	Request->Flags.Add(EHttpRequestFlags::CancelRequested);
 
-	UE_LOG(LogHttpIoDispatcher, Verbose, TEXT("Cancelling request, SeqNo=%u"), Request->SeqNo);
+	//UE_LOG(LogHttpIoDispatcher, Verbose, TEXT("Cancelling request, SeqNo=%u"), Request->SeqNo);
 
 	if (EnumHasAnyFlags(CurrentFlags, EHttpRequestFlags::CacheInflight))
 	{
@@ -1402,8 +1411,8 @@ void FHttpDispatcher::CancelRequest(FIoHttpRequestHandle Handle)
 
 	if (GIaxHttpEnableInflightCancellation)
 	{
-		TUniqueLock _(CancelationMutex);
-		Request->CancelationToken.Reset();
+		TUniqueLock _(CancellationMutex);
+		Request->CancellationToken.Reset();
 	}
 }
 
@@ -1695,6 +1704,8 @@ bool FHttpDispatcher::ProcessHttpRequests()
 					FOnDemandIoBackendStats::Get()->OnHttpRequestCompleted();
 					NumConcurrentRequests--;
 
+					CompleteCancellationToken(HttpRequest);
+
 					UE::Tasks::Launch(
 						UE_SOURCE_LOCATION,
 						[this, HttpRequest, HttpResponse = MoveTemp(HttpResponse), TotalRange, HostGroup = MoveTemp(HostGroup)]() mutable
@@ -1703,21 +1714,7 @@ bool FHttpDispatcher::ProcessHttpRequests()
 						}, UE::Tasks::ETaskPriority::BackgroundLow);
 				});
 
-			{
-				TUniqueLock _(CancelationMutex);
-			
-				TSharedPtr<FCancelationToken> Token = MakeShared<FCancelationToken>(*HttpClient.Get(), TicketId);
-				FHttpRequest* RequestList = HttpRequest;
-				while (RequestList)
-				{
-					if (!RequestList->Flags.HasAny(EHttpRequestFlags::CancelRequested))
-					{
-						RequestList->CancelationToken = Token;
-					}
-
-					RequestList = RequestList->Next;
-				}
-			}
+			ApplyCancellationToken(HttpRequest, TicketId);
 
 			if (NumConcurrentRequests >= MaxConcurrentRequests)
 			{
@@ -1917,14 +1914,48 @@ void FHttpDispatcher::OnPauseCommand(bool bShouldPause, const TArray<FString>& A
 
 	if (!bShouldPause)
 	{
-		// Since we are unpausing we might have moved previously paused requests back to the active queue
-		// so we need the thread to start runnning again.
+		// Since we are un-pausing we might have moved previously paused requests back to the active queue
+		// so we need the thread to start running again.
 		WakeUp->Trigger();
 	}
 }
 
 #endif // UE_ALLOW_HTTP_PAUSE
 
+void FHttpDispatcher::ApplyCancellationToken(FHttpRequest* RequestChain, FMultiEndpointHttpClient::FHttpTicketId TicketId)
+{
+	TUniqueLock _(CancellationMutex);
+
+	TSharedPtr<FCancellationToken> Token = MakeShared<FCancellationToken>(*HttpClient.Get(), TicketId);
+
+	while (RequestChain != nullptr)
+	{
+		if (!RequestChain->Flags.HasAny(EHttpRequestFlags::CancelRequested))
+		{
+			RequestChain->CancellationToken = Token;
+		}
+
+		RequestChain = RequestChain->Next;
+	}
+}
+
+void FHttpDispatcher::CompleteCancellationToken(FHttpRequest* RequestChain)
+{
+	TUniqueLock _(CancellationMutex);
+
+	while (RequestChain != nullptr)
+	{
+		if (RequestChain->CancellationToken)
+		{
+			RequestChain->CancellationToken->OnRequestCompleted();
+			// Every request in the chain shares the same token (unless null) so we can early out at this point
+			return;
+		}
+
+		RequestChain = RequestChain->Next;
+	}
+}
+
 } // namespace UE::IoStore::HttpIoDispatcher
 
 ///////////////////////////////////////////////////////////////////////////////
