Add support for completing I/O store HTTP request before completion callback.
-this change adds a cvar ias.FastCancelEnabled=true that when enabled will complete I/O dispatcher HTTP request before
the completion callback has been triggered when canceling requests

#tests Manual testing
#rnx

#Synced-CL 48271216

diff --git a/Engine/Source/Runtime/Experimental/IoStore/OnDemand/Private/OnDemandHttpIoDispatcherBackend.cpp b/Engine/Source/Runtime/Experimental/IoStore/OnDemand/Private/OnDemandHttpIoDispatcherBackend.cpp
--- a/Engine/Source/Runtime/Experimental/IoStore/OnDemand/Private/OnDemandHttpIoDispatcherBackend.cpp
+++ b/Engine/Source/Runtime/Experimental/IoStore/OnDemand/Private/OnDemandHttpIoDispatcherBackend.cpp
@@ -14,6 +14,8 @@
 #include "OnDemandHttpIoDispatcher.h" // EHttpRequestType
 #include "Statistics.h"
 
+#include <atomic>
+
 namespace UE::IoStore
 {
 ////////////////////////////////////////////////////////////////////////////////
@@ -31,6 +33,12 @@ static FAutoConsoleVariableRef CVar_WriteOnlyDecoding(
 	TEXT("Enables the use of 'WriteOnly' flag when decoding to buffers with the 'HardwareTargetBuffer' flag")
 );
 
+bool GIasFastCancelEnabled = true;
+static FAutoConsoleVariableRef CVar_IasFastCancelEnabled(
+	TEXT("ias.FastCancelEnabled"),
+	GIasFastCancelEnabled,
+	TEXT("Enables canceling IAS requests before the completion callback has been triggered.")
+);
 ////////////////////////////////////////////////////////////////////////////////
 class FHttpIoDispatcherBackend
 	: public IOnDemandHttpIoDispatcherBackend
@@ -40,13 +48,13 @@ class FHttpIoDispatcherBackend
 	struct FResolvedRequest
 	{
 		FResolvedRequest(FIoRequestImpl& InRequest, FOnDemandChunkInfo&& InChunkInfo, const FIoOffsetAndLength& InChunkRange)
-			: DispatcherRequest(InRequest)
+			: DispatcherRequest(&InRequest)
 			, ChunkInfo(MoveTemp(InChunkInfo))
 			, ChunkRange(InChunkRange)
 			, StartTimeCycles(FPlatformTime::Cycles64())
 		{
-			check(DispatcherRequest.BackendData == nullptr);
-			DispatcherRequest.BackendData = this;
+			check(DispatcherRequest->BackendData == nullptr);
+			DispatcherRequest->BackendData = this;
 		}
 
 		static FResolvedRequest* Get(FIoRequestImpl& DispatcherRequest)
@@ -56,24 +64,31 @@ class FHttpIoDispatcherBackend
 
 		static FResolvedRequest* Create(FIoRequestImpl& DispatcherRequest, FOnDemandChunkInfo&& ChunkInfo, const FIoOffsetAndLength& ChunkRange)
 		{
+			TUniqueLock Lock(AllocatorMutex);
 			return Allocator.Construct(DispatcherRequest, MoveTemp(ChunkInfo), ChunkRange);
 		}
 
-		static void Destroy(FResolvedRequest* ResolvedRequest)
+		static void Release(FResolvedRequest* ResolvedRequest)
 		{
-			check(ResolvedRequest->DispatcherRequest.BackendData == ResolvedRequest);
-			ResolvedRequest->DispatcherRequest.BackendData = nullptr;
-			Allocator.Destroy(ResolvedRequest);
+			if (1 == ResolvedRequest->RefCount.fetch_sub(1, std::memory_order_acq_rel))
+			{
+				ResolvedRequest->~FResolvedRequest();
+				TUniqueLock Lock(AllocatorMutex);
+				Allocator.Free(ResolvedRequest);
+			}
 		}
 
 		using FAllocator		= TSingleThreadedSlabAllocator<FResolvedRequest, 64>;
 		static FAllocator		Allocator;
+		static FMutex			AllocatorMutex;
 
-		FIoRequestImpl&			DispatcherRequest;
+		FIoRequestImpl*			DispatcherRequest;
 		FOnDemandChunkInfo		ChunkInfo;
 		FIoOffsetAndLength		ChunkRange;
 		FIoHttpRequest			HttpRequest;
 		uint64					StartTimeCycles;
+		std::atomic_int32_t		RefCount{2};
+		std::atomic_bool		bCompleted{false};
 	};
 
 public:
@@ -107,6 +122,7 @@ private:
 };
 
 FHttpIoDispatcherBackend::FResolvedRequest::FAllocator FHttpIoDispatcherBackend::FResolvedRequest::Allocator;
+FMutex FHttpIoDispatcherBackend::FResolvedRequest::AllocatorMutex;
 
 ////////////////////////////////////////////////////////////////////////////////
 FHttpIoDispatcherBackend::FHttpIoDispatcherBackend(FOnDemandIoStore& InIoStore)
@@ -232,8 +248,8 @@ FIoRequestImpl* FHttpIoDispatcherBackend::GetCompletedIoRequests()
 			continue;
 		}
 
-		FResolvedRequest::Destroy(reinterpret_cast<FResolvedRequest*>(R.BackendData));
-		check(R.BackendData == nullptr);
+		FResolvedRequest::Release(reinterpret_cast<FResolvedRequest*>(R.BackendData));
+		R.BackendData = nullptr;
 	}
 
 	return Out.GetHead();
@@ -244,6 +260,21 @@ void FHttpIoDispatcherBackend::CancelIoRequest(FIoRequestImpl* Request)
 	if (FResolvedRequest* Resolved = FResolvedRequest::Get(*Request))
 	{
 		Resolved->HttpRequest.Cancel();
+
+		// Try and complete the request before the callback to cancel faster
+		if (GIasFastCancelEnabled)
+		{
+			bool bExpected = false;
+			if (Resolved->bCompleted.compare_exchange_strong(bExpected, true))
+			{
+				Resolved->DispatcherRequest = nullptr;
+				{
+					TUniqueLock Lock(CompletedRequestMutex);
+					CompletedRequests.AddTail(Request);
+				}
+				BackendContext->WakeUpDispatcherThreadDelegate.Execute();
+			}
+		}
 	}
 }
 
@@ -279,9 +310,20 @@ TIoStatusOr<uint64> FHttpIoDispatcherBackend::GetSizeForChunk(const FIoChunkId&
 
 void FHttpIoDispatcherBackend::CompleteRequest(FResolvedRequest& ResolvedRequest, FIoHttpResponse&& HttpResponse)
 {
+	// Check whether the I/O dispatcher request has been canceled and added to the completed request(s)
+	if (GIasFastCancelEnabled)
+	{
+		bool bExpected = false;
+		if (ResolvedRequest.bCompleted.compare_exchange_strong(bExpected, true) == false)
+		{
+			return FResolvedRequest::Release(&ResolvedRequest);
+		}
+	}
+
+	check(ResolvedRequest.DispatcherRequest != nullptr);
 	FIoBuffer Chunk						= HttpResponse.GetBody();
 	FOnDemandChunkInfo& ChunkInfo		= ResolvedRequest.ChunkInfo;
-	FIoRequestImpl& DispatcherRequest	= ResolvedRequest.DispatcherRequest;
+	FIoRequestImpl& DispatcherRequest	= *ResolvedRequest.DispatcherRequest;
 
 	const bool bCancelled				= DispatcherRequest.IsCancelled();
 	bool bDecoded						= false;
@@ -348,6 +390,7 @@ void FHttpIoDispatcherBackend::CompleteRequest(FResolvedRequest& ResolvedRequest
 	}
 	
 	LogRequest(ResolvedRequest, DispatcherRequest, bDecoded, HttpResponse.IsCached(), bCancelled, DurationMs);
+	FResolvedRequest::Release(&ResolvedRequest);
 
 	{
 		TUniqueLock Lock(CompletedRequestMutex);
