Fixing chaos scale tri mesh sweeps having issues with zero scale on an axis.
The old code brought the outputs back and forth through a few spaces which would lose data if an axis' scale was zero. This change removes the redundant transforms which preserves the data.
#Synced-CL 48026629

diff --git a/Engine/Source/Programs/HeadlessChaos/Private/HeadlessChaosTestSweep.cpp b/Engine/Source/Programs/HeadlessChaos/Private/HeadlessChaosTestSweep.cpp
--- a/Engine/Source/Programs/HeadlessChaos/Private/HeadlessChaosTestSweep.cpp
+++ b/Engine/Source/Programs/HeadlessChaos/Private/HeadlessChaosTestSweep.cpp
@@ -895,6 +895,15 @@ namespace ChaosTest
 			int32 FaceId;
 		};
 
+		struct FSweepCCDResultData
+		{
+			bool bResult;
+			FReal TOI;
+			FReal Phi;
+			FVec3 Position, Normal, FaceNormal;
+			int32 FaceId;
+		};
+
 		FSweepResultData SweepQuerySphere(const FImplicitObject& TestObject, const FRigidTransform3& TestObjectTransform, const FVec3& SweepStart, const FVec3& SweepDirection, const FReal SweepLength = 100)
 		{
 			const Chaos::FSphere Sphere(FVec3::ZeroVector, SphereRadius);
@@ -1237,6 +1246,15 @@ namespace ChaosTest
 			return Result;
 		}
 
+		static FMeshInitData CreateTriInitData(const FVec3& P0, const FVec3& P1, const FVec3& P2)
+		{
+			FMeshInitData Result;
+			Result.TrimeshParticles = FTriangleMeshImplicitObject::ParticlesType({P0, P1, P2});
+			Result.Indices = { {0, 1, 2}, };
+			Result.AutoBuildMaterials();
+			return Result;
+		}
+
 		FSweepResultData SweepSphere(const TImplicitObjectScaled<FTriangleMeshImplicitObject>& TriangleMesh, const FVec3& SweepStart, const FVec3& SweepDirection, const FReal SweepLength = 100)
 		{
 			const Chaos::FSphere Sphere(FVec3::ZeroVector, SphereRadius);
@@ -1251,13 +1269,27 @@ namespace ChaosTest
 		FSweepResultData SweepSphere(FMeshInitData& InitData, const FVec3& SweepStart, const FVec3& SweepDirection, const FReal SweepLength = 100)
 		{
 			FTriangleMeshImplicitObject TriangleMesh(MoveTemp(InitData.TrimeshParticles), MoveTemp(InitData.Indices), MoveTemp(InitData.Materials), nullptr, nullptr, true);
-			FTriangleMeshBackfaceSweepFixture::SweepSphere(TriangleMesh, SweepStart, SweepDirection, SweepLength);
 			// Unfortunately, it is not possible to currently call FTriangleMeshImplicitObject::SweepGeom directly. 
 			// There's code inside of TImplicitObjectScaled that negates some inputs values (the direction) that is necessary for the tests to work (follows the editor path).
 			TImplicitObjectScaled<FTriangleMeshImplicitObject> ScaledMesh(&TriangleMesh, TriMeshScale);
 			return SweepSphere(ScaledMesh, SweepStart, SweepDirection, SweepLength);
 		}
 
+		FSweepCCDResultData SweepSphereCCD(FMeshInitData& InitData, const FVec3& SweepStart, const FVec3& SweepDirection, const FReal SweepLength = 100)
+		{
+			const Chaos::FSphere Sphere(FVec3::ZeroVector, SphereRadius);
+			const FRigidTransform3 StartTM(SweepStart, TRotation<FReal, 3>::Identity);
+			const FVec3 NormalizedSweepDirection = SweepDirection.GetSafeNormal();
+			const FRigidTransform3 BToATM = StartTM.GetRelativeTransformNoScale(FRigidTransform3());
+
+			FTriangleMeshImplicitObject TriangleMesh(MoveTemp(InitData.TrimeshParticles), MoveTemp(InitData.Indices), MoveTemp(InitData.Materials), nullptr, nullptr, true);
+			TImplicitObjectScaled<FTriangleMeshImplicitObject> ScaledMesh(&TriangleMesh, TriMeshScale);
+
+			FSweepCCDResultData Result;
+			Result.bResult = ScaledMesh.LowLevelSweepGeomCCD(Sphere, BToATM, NormalizedSweepDirection, SweepLength, 0, 0, Result.TOI, Result.Phi, Result.Position, Result.Normal, Result.FaceId, Result.FaceNormal);
+			return Result;
+		}
+
 		FSweepResultData SweepSphereVsSimpleTri(const FVec3& SweepStart, const FVec3& SweepDirection, const FReal SweepLength = 100)
 		{
 			FMeshInitData InitData = CreateSimpleTriInitData();
@@ -1441,6 +1473,297 @@ namespace ChaosTest
 		EXPECT_EQ(Result.bResult, true);
 	}
 
+	// Do a bunch of tests with a slanted triangle and different non-uniform scales
+	TEST_F(FScaledTriMeshSweepFixture, TriMeshWithUniformScale_Sweep_ResultsAreExpected)
+	{
+		SphereRadius = 1;
+		TriMeshScale = FVec3(1, 1, 1);
+		FMeshInitData InitData = CreateTriInitData(FVec3(5, 0, 0), FVec3(0, 5, 0), FVec3(0, 0, 5));
+		FVec3 SweepStart = FVec3(2.5, 10, 2.5);
+		FVec3 SweepDir = FVec3(0, -1, 0);
+		FSweepResultData Result = SweepSphere(InitData, SweepStart, SweepDir);
+
+		EXPECT_EQ(Result.bResult, true);
+		EXPECT_NEAR(Result.TOI, 8.2679491043090820, KINDA_SMALL_NUMBER);
+		EXPECT_VECTOR_NEAR(Result.Position, FVec3(1.9226497411727905, 1.1547006368637085, 1.9226497411727905), KINDA_SMALL_NUMBER);
+		EXPECT_VECTOR_NEAR(Result.Normal, FVec3(1, 1, 1).GetUnsafeNormal(), KINDA_SMALL_NUMBER);
+	}
+
+	TEST_F(FScaledTriMeshSweepFixture, TriMeshWithDoubleXScale_Sweep_ResultsAreExpected)
+	{
+		SphereRadius = 1;
+		TriMeshScale = FVec3(2, 1, 1);
+		FMeshInitData InitData = CreateTriInitData(FVec3(5, 0, 0), FVec3(0, 5, 0), FVec3(0, 0, 5));
+		FVec3 SweepStart = FVec3(2.5, 10, 2.5);
+		FVec3 SweepDir = FVec3(0, -1, 0);
+		FSweepResultData Result = SweepSphere(InitData, SweepStart, SweepDir);
+
+		EXPECT_EQ(Result.bResult, true);
+		EXPECT_NEAR(Result.TOI, 7.2500000000000000, KINDA_SMALL_NUMBER);
+		EXPECT_VECTOR_NEAR(Result.Position, FVec3(2.1666667461395264, 2.0833334922790527, 1.8333333730697632), KINDA_SMALL_NUMBER);
+		EXPECT_VECTOR_NEAR(Result.Normal, FVec3(.5, 1, 1).GetUnsafeNormal(), KINDA_SMALL_NUMBER);
+	}
+
+	TEST_F(FScaledTriMeshSweepFixture, TriMeshWithHalfXScale_Sweep_ResultsAreExpected)
+	{
+		SphereRadius = 1;
+		TriMeshScale = FVec3(.5, 1, 1);
+		FMeshInitData InitData = CreateTriInitData(FVec3(10, 0, 0), FVec3(0, 10, 0), FVec3(0, 0, 10));
+		FVec3 SweepStart = FVec3(2.5, 10, 2.5);
+		FVec3 SweepDir = FVec3(0, -1, 0);
+		FSweepResultData Result = SweepSphere(InitData, SweepStart, SweepDir);
+
+		EXPECT_EQ(Result.bResult, true);
+		EXPECT_NEAR(Result.TOI, 5.0505099296569824, KINDA_SMALL_NUMBER);
+		EXPECT_VECTOR_NEAR(Result.Position, FVec3(1.6835033893585205, 4.5412416458129883, 2.0917518138885498), KINDA_SMALL_NUMBER);
+		EXPECT_VECTOR_NEAR(Result.Normal, FVec3(2, 1, 1).GetUnsafeNormal(), KINDA_SMALL_NUMBER);
+	}
+
+	TEST_F(FScaledTriMeshSweepFixture, TriMeshWithZeroXScale_Sweep_ResultsAreExpected)
+	{
+		SphereRadius = 1;
+		TriMeshScale = FVec3(0, 1, 1);
+		FMeshInitData InitData = CreateTriInitData(FVec3(5, 0, 0), FVec3(0, 5, 0), FVec3(0, 0, 5));
+		FVec3 SweepStart = FVec3(10, 2.5, 2.5);
+		FVec3 SweepDir = FVec3(-1, 0, 0);
+		FSweepResultData Result = SweepSphere(InitData, SweepStart, SweepDir);
+
+		EXPECT_EQ(Result.bResult, true);
+		EXPECT_NEAR(Result.TOI, 9, KINDA_SMALL_NUMBER);
+		EXPECT_VECTOR_NEAR(Result.Position, FVec3(0, 2.5, 2.5), KINDA_SMALL_NUMBER);
+		EXPECT_VECTOR_NEAR(Result.Normal, FVec3(1, 0, 0).GetUnsafeNormal(), KINDA_SMALL_NUMBER);
+	}
+
+	TEST_F(FScaledTriMeshSweepFixture, TriMeshWithZeroYScale_Sweep_ResultsAreExpected)
+	{
+		SphereRadius = 1;
+		TriMeshScale = FVec3(1, 0, 1);
+		FMeshInitData InitData = CreateTriInitData(FVec3(5, 0, 0), FVec3(0, 5, 0), FVec3(0, 0, 5));
+		FVec3 SweepStart = FVec3(2.5, 10, 2.5);
+		FVec3 SweepDir = FVec3(0, -1, 0);
+		FSweepResultData Result = SweepSphere(InitData, SweepStart, SweepDir);
+
+		EXPECT_EQ(Result.bResult, true);
+		EXPECT_NEAR(Result.TOI, 9, KINDA_SMALL_NUMBER);
+		EXPECT_VECTOR_NEAR(Result.Position, FVec3(2.5, 0, 2.5), KINDA_SMALL_NUMBER);
+		EXPECT_VECTOR_NEAR(Result.Normal, FVec3(0, 1, 0).GetUnsafeNormal(), KINDA_SMALL_NUMBER);
+	}
+
+	TEST_F(FScaledTriMeshSweepFixture, TriMeshWithZeroZScale_Sweep_ResultsAreExpected)
+	{
+		SphereRadius = 1;
+		TriMeshScale = FVec3(1, 1, 0);
+		FMeshInitData InitData = CreateTriInitData(FVec3(5, 0, 0), FVec3(0, 5, 0), FVec3(0, 0, 5));
+		FVec3 SweepStart = FVec3(2.5, 2.5, 10);
+		FVec3 SweepDir = FVec3(0, 0, -1);
+		FSweepResultData Result = SweepSphere(InitData, SweepStart, SweepDir);
+
+		EXPECT_EQ(Result.bResult, true);
+		EXPECT_NEAR(Result.TOI, 9, KINDA_SMALL_NUMBER);
+		EXPECT_VECTOR_NEAR(Result.Position, FVec3(2.5, 2.5, 0), KINDA_SMALL_NUMBER);
+		EXPECT_VECTOR_NEAR(Result.Normal, FVec3(0, 0, 1).GetUnsafeNormal(), KINDA_SMALL_NUMBER);
+	}
+
+	TEST_F(FScaledTriMeshSweepFixture, TriMeshWithYNormalAndZeroYScale_Sweep_ResultsAreExpected)
+	{
+		// This is a special case where TriMeshScale * HitNormal = (1, 0, 1) * (0, 1, 0) = (0, 0, 0).
+		// This was a previous bug that would return (0, 0, 0) instead of the expected (0, 1, 0).
+		SphereRadius = 1;
+		TriMeshScale = FVec3(1, 0, 1);
+		FMeshInitData InitData = CreateTriInitData(FVec3(5, 0, 0), FVec3(0, 0, 0), FVec3(0, 0, 5));
+		FVec3 SweepStart = FVec3(2.5, 10, 2.5);
+		FVec3 SweepDir = FVec3(0, -1, 0);
+		FSweepResultData Result = SweepSphere(InitData, SweepStart, SweepDir);
+
+		EXPECT_EQ(Result.bResult, true);
+		EXPECT_NEAR(Result.TOI, 9, KINDA_SMALL_NUMBER);
+		EXPECT_VECTOR_NEAR(Result.Position, FVec3(2.5, 0, 2.5), KINDA_SMALL_NUMBER);
+		EXPECT_VECTOR_NEAR(Result.Normal, FVec3(0, 1, 0).GetUnsafeNormal(), KINDA_SMALL_NUMBER);
+	}
+
+	TEST_F(FScaledTriMeshSweepFixture, TriMeshWithLargeScaleThatBarelyHits_Sweep_ResultsAreExpected)
+	{
+		// This test case helps to ensure that the TOI is correctly handled internally. The trimesh does some back and forth scaling of time and this was added to catch a breaking case.
+		FReal SweepLength = 10;
+		SphereRadius = 1;
+		TriMeshScale = FVec3(50, 50, 50);
+		FMeshInitData InitData = CreateTriInitData(FVec3(0, 1, 0), FVec3(0, 1, 5), FVec3(5, 1, 0));
+		FVec3 SweepStart = FVec3(2.5, 60.9, 2.5);
+		FVec3 SweepDir = FVec3(0, -1, 0);
+		FSweepResultData Result = SweepSphere(InitData, SweepStart, SweepDir, SweepLength);
+
+		EXPECT_EQ(Result.bResult, true);
+		EXPECT_NEAR(Result.TOI, 9.9, KINDA_SMALL_NUMBER);
+		EXPECT_VECTOR_NEAR(Result.Position, FVec3(2.5, 50, 2.5), KINDA_SMALL_NUMBER);
+		EXPECT_VECTOR_NEAR(Result.Normal, FVec3(0, 1, 0).GetUnsafeNormal(), KINDA_SMALL_NUMBER);
+	}
+	TEST_F(FScaledTriMeshSweepFixture, TriMeshWithLargeScaleThatBarelyMisses_Sweep_ResultsAreExpected)
+	{
+		// This test case helps to ensure that the TOI is correctly handled internally. The trimesh does some back and forth scaling of time and this was added to catch a breaking case.
+		FReal SweepLength = 10;
+		SphereRadius = 1;
+		TriMeshScale = FVec3(50, 50, 50);
+		FMeshInitData InitData = CreateTriInitData(FVec3(0, 1, 0), FVec3(0, 1, 5), FVec3(5, 1, 0));
+		FVec3 SweepStart = FVec3(2.5, 61.1, 2.5);
+		FVec3 SweepDir = FVec3(0, -1, 0);
+		FSweepResultData Result = SweepSphere(InitData, SweepStart, SweepDir, SweepLength);
+
+		EXPECT_EQ(Result.bResult, false);
+	}
+
+	TEST_F(FScaledTriMeshSweepFixture, TriMeshWithUniformScale_SweepCCD_ResultsAreExpected)
+	{
+		const FReal SweepLength = 100;
+		SphereRadius = 1;
+		TriMeshScale = FVec3(1, 1, 1);
+		FMeshInitData InitData = CreateTriInitData(FVec3(5, 0, 0), FVec3(0, 5, 0), FVec3(0, 0, 5));
+		FVec3 SweepStart = FVec3(2.5, 10, 2.5);
+		FVec3 SweepDir = FVec3(0, -1, 0);
+		FSweepCCDResultData Result = SweepSphereCCD(InitData, SweepStart, SweepDir, SweepLength);
+
+		EXPECT_EQ(Result.bResult, true);
+		EXPECT_NEAR(Result.TOI, 8.2679491043090820 / SweepLength, KINDA_SMALL_NUMBER);
+		EXPECT_NEAR(Result.Phi, 0, KINDA_SMALL_NUMBER);
+		EXPECT_VECTOR_NEAR(Result.Position, FVec3(1.9226497411727905, 1.1547006368637085, 1.9226497411727905), KINDA_SMALL_NUMBER);
+		EXPECT_VECTOR_NEAR(Result.Normal, FVec3(1, 1, 1).GetUnsafeNormal(), KINDA_SMALL_NUMBER);
+	}
+
+	TEST_F(FScaledTriMeshSweepFixture, TriMeshWithDoubleXScale_SweepCCD_ResultsAreExpected)
+	{
+		const FReal SweepLength = 100;
+		SphereRadius = 1;
+		TriMeshScale = FVec3(2, 1, 1);
+		FMeshInitData InitData = CreateTriInitData(FVec3(5, 0, 0), FVec3(0, 5, 0), FVec3(0, 0, 5));
+		FVec3 SweepStart = FVec3(2.5, 10, 2.5);
+		FVec3 SweepDir = FVec3(0, -1, 0);
+		FSweepCCDResultData Result = SweepSphereCCD(InitData, SweepStart, SweepDir, SweepLength);
+
+		EXPECT_EQ(Result.bResult, true);
+		EXPECT_NEAR(Result.TOI, 7.2500000000000000 / SweepLength, KINDA_SMALL_NUMBER);
+		EXPECT_NEAR(Result.Phi, 0, KINDA_SMALL_NUMBER);
+		EXPECT_VECTOR_NEAR(Result.Position, FVec3(2.1666667461395264, 2.0833334922790527, 1.8333333730697632), KINDA_SMALL_NUMBER);
+		EXPECT_VECTOR_NEAR(Result.Normal, FVec3(.5, 1, 1).GetUnsafeNormal(), KINDA_SMALL_NUMBER);
+	}
+
+	TEST_F(FScaledTriMeshSweepFixture, TriMeshWithHalfXScale_SweepCCD_ResultsAreExpected)
+	{
+		const FReal SweepLength = 100;
+		SphereRadius = 1;
+		TriMeshScale = FVec3(.5, 1, 1);
+		FMeshInitData InitData = CreateTriInitData(FVec3(10, 0, 0), FVec3(0, 10, 0), FVec3(0, 0, 10));
+		FVec3 SweepStart = FVec3(2.5, 10, 2.5);
+		FVec3 SweepDir = FVec3(0, -1, 0);
+		FSweepCCDResultData Result = SweepSphereCCD(InitData, SweepStart, SweepDir, SweepLength);
+
+		EXPECT_EQ(Result.bResult, true);
+		EXPECT_NEAR(Result.TOI, 5.0505099296569824 / SweepLength, KINDA_SMALL_NUMBER);
+		EXPECT_NEAR(Result.Phi, 0, KINDA_SMALL_NUMBER);
+		EXPECT_VECTOR_NEAR(Result.Position, FVec3(1.6835033893585205, 4.5412416458129883, 2.0917518138885498), KINDA_SMALL_NUMBER);
+		EXPECT_VECTOR_NEAR(Result.Normal, FVec3(2, 1, 1).GetUnsafeNormal(), KINDA_SMALL_NUMBER);
+	}
+
+	TEST_F(FScaledTriMeshSweepFixture, TriMeshWithZeroXScale_SweepCCD_ResultsAreExpected)
+	{
+		const FReal SweepLength = 100;
+		SphereRadius = 1;
+		TriMeshScale = FVec3(0, 1, 1);
+		FMeshInitData InitData = CreateTriInitData(FVec3(5, 0, 0), FVec3(0, 5, 0), FVec3(0, 0, 5));
+		FVec3 SweepStart = FVec3(10, 2.5, 2.5);
+		FVec3 SweepDir = FVec3(-1, 0, 0);
+		FSweepCCDResultData Result = SweepSphereCCD(InitData, SweepStart, SweepDir, SweepLength);
+
+		EXPECT_EQ(Result.bResult, true);
+		EXPECT_NEAR(Result.TOI, 9.0 / SweepLength, KINDA_SMALL_NUMBER);
+		EXPECT_NEAR(Result.Phi, 0, KINDA_SMALL_NUMBER);
+		EXPECT_VECTOR_NEAR(Result.Position, FVec3(0, 2.5, 2.5), KINDA_SMALL_NUMBER);
+		EXPECT_VECTOR_NEAR(Result.Normal, FVec3(1, 0, 0).GetUnsafeNormal(), KINDA_SMALL_NUMBER);
+	}
+
+	TEST_F(FScaledTriMeshSweepFixture, TriMeshWithZeroYScale_SweepCCD_ResultsAreExpected)
+	{
+		const FReal SweepLength = 100;
+		SphereRadius = 1;
+		TriMeshScale = FVec3(1, 0, 1);
+		FMeshInitData InitData = CreateTriInitData(FVec3(5, 0, 0), FVec3(0, 5, 0), FVec3(0, 0, 5));
+		FVec3 SweepStart = FVec3(2.5, 10, 2.5);
+		FVec3 SweepDir = FVec3(0, -1, 0);
+		FSweepCCDResultData Result = SweepSphereCCD(InitData, SweepStart, SweepDir, SweepLength);
+
+		EXPECT_EQ(Result.bResult, true);
+		EXPECT_NEAR(Result.TOI, 9.0 / SweepLength, KINDA_SMALL_NUMBER);
+		EXPECT_NEAR(Result.Phi, 0, KINDA_SMALL_NUMBER);
+		EXPECT_VECTOR_NEAR(Result.Position, FVec3(2.5, 0, 2.5), KINDA_SMALL_NUMBER);
+		EXPECT_VECTOR_NEAR(Result.Normal, FVec3(0, 1, 0).GetUnsafeNormal(), KINDA_SMALL_NUMBER);
+	}
+
+	TEST_F(FScaledTriMeshSweepFixture, TriMeshWithZeroZScale_SweepCCD_ResultsAreExpected)
+	{
+		const FReal SweepLength = 100;
+		SphereRadius = 1;
+		TriMeshScale = FVec3(1, 1, 0);
+		FMeshInitData InitData = CreateTriInitData(FVec3(5, 0, 0), FVec3(0, 5, 0), FVec3(0, 0, 5));
+		FVec3 SweepStart = FVec3(2.5, 2.5, 10);
+		FVec3 SweepDir = FVec3(0, 0, -1);
+		FSweepCCDResultData Result = SweepSphereCCD(InitData, SweepStart, SweepDir, SweepLength);
+
+		EXPECT_EQ(Result.bResult, true);
+		EXPECT_NEAR(Result.TOI, 9.0 / SweepLength, KINDA_SMALL_NUMBER);
+		EXPECT_NEAR(Result.Phi, 0, KINDA_SMALL_NUMBER);
+		EXPECT_VECTOR_NEAR(Result.Position, FVec3(2.5, 2.5, 0), KINDA_SMALL_NUMBER);
+		EXPECT_VECTOR_NEAR(Result.Normal, FVec3(0, 0, 1).GetUnsafeNormal(), KINDA_SMALL_NUMBER);
+	}
+
+	TEST_F(FScaledTriMeshSweepFixture, TriMeshWithYNormalAndZeroYScale_SweepCCD_ResultsAreExpected)
+	{
+		// This is a special case where TriMeshScale * HitNormal = (1, 0, 1) * (0, 1, 0) = (0, 0, 0).
+		// This was a previous bug that would return (0, 0, 0) instead of the expected (0, 1, 0).
+		const FReal SweepLength = 100;
+		SphereRadius = 1;
+		TriMeshScale = FVec3(1, 0, 1);
+		FMeshInitData InitData = CreateTriInitData(FVec3(5, 0, 0), FVec3(0, 0, 0), FVec3(0, 0, 5));
+		FVec3 SweepStart = FVec3(2.5, 10, 2.5);
+		FVec3 SweepDir = FVec3(0, -1, 0);
+		FSweepCCDResultData Result = SweepSphereCCD(InitData, SweepStart, SweepDir, SweepLength);
+
+		EXPECT_EQ(Result.bResult, true);
+		EXPECT_NEAR(Result.TOI, 9.0 / SweepLength, KINDA_SMALL_NUMBER);
+		EXPECT_NEAR(Result.Phi, 0, KINDA_SMALL_NUMBER);
+		EXPECT_VECTOR_NEAR(Result.Position, FVec3(2.5, 0, 2.5), KINDA_SMALL_NUMBER);
+		EXPECT_VECTOR_NEAR(Result.Normal, FVec3(0, 1, 0).GetUnsafeNormal(), KINDA_SMALL_NUMBER);
+	}
+
+	TEST_F(FScaledTriMeshSweepFixture, TriMeshWithLargeScaleThatBarelyHits_SweepCCD_ResultsAreExpected)
+	{
+		// This test case helps to ensure that the TOI is correctly handled internally. The trimesh does some back and forth scaling of time and this was added to catch a breaking case.
+		FReal SweepLength = 10;
+		SphereRadius = 1;
+		TriMeshScale = FVec3(50, 50, 50);
+		FMeshInitData InitData = CreateTriInitData(FVec3(0, 1, 0), FVec3(0, 1, 5), FVec3(5, 1, 0));
+		FVec3 SweepStart = FVec3(2.5, 60.9, 2.5);
+		FVec3 SweepDir = FVec3(0, -1, 0);
+		FSweepCCDResultData Result = SweepSphereCCD(InitData, SweepStart, SweepDir, SweepLength);
+
+		EXPECT_EQ(Result.bResult, true);
+		EXPECT_NEAR(Result.TOI, 9.9 / SweepLength, KINDA_SMALL_NUMBER);
+		EXPECT_NEAR(Result.Phi, 0, KINDA_SMALL_NUMBER);
+		EXPECT_VECTOR_NEAR(Result.Position, FVec3(2.5, 50, 2.5), KINDA_SMALL_NUMBER);
+		EXPECT_VECTOR_NEAR(Result.Normal, FVec3(0, 1, 0).GetUnsafeNormal(), KINDA_SMALL_NUMBER);
+	}
+
+	TEST_F(FScaledTriMeshSweepFixture, TriMeshWithLargeScaleThatBarelyMisses_SweepCCD_ResultsAreExpected)
+	{
+		// This test case helps to ensure that the TOI is correctly handled internally. The trimesh does some back and forth scaling of time and this was added to catch a breaking case.
+		FReal SweepLength = 10;
+		SphereRadius = 1;
+		TriMeshScale = FVec3(50, 50, 50);
+		FMeshInitData InitData = CreateTriInitData(FVec3(0, 1, 0), FVec3(0, 1, 5), FVec3(5, 1, 0));
+		FVec3 SweepStart = FVec3(2.5, 61.1, 2.5);
+		FVec3 SweepDir = FVec3(0, -1, 0);
+		FSweepCCDResultData Result = SweepSphereCCD(InitData, SweepStart, SweepDir, SweepLength);
+
+		EXPECT_EQ(Result.bResult, false);
+	}
+
 	TEST_F(FSphereSweepFixture, TransformedUnionContainingNonUniformScaledObject_SweepSphere_ResultsAreExpected)
 	{
 		// This test is to catch a bug where a Transformed union with a scaled object would 
diff --git a/Engine/Source/Runtime/Experimental/Chaos/Private/Chaos/TriangleMeshImplicitObject.cpp b/Engine/Source/Runtime/Experimental/Chaos/Private/Chaos/TriangleMeshImplicitObject.cpp
--- a/Engine/Source/Runtime/Experimental/Chaos/Private/Chaos/TriangleMeshImplicitObject.cpp
+++ b/Engine/Source/Runtime/Experimental/Chaos/Private/Chaos/TriangleMeshImplicitObject.cpp
@@ -69,16 +69,26 @@ TImplicitObjectScaled<QueryGeomType> ScaleGeomIntoWorldHelper(const TImplicitObj
 	return MakeScaledHelper(QueryGeom, TriMeshScale);
 }
 
-void TransformSweepOutputsHelper(FVec3 TriMeshScale, const FVec3& HitNormal, const FVec3& HitPosition, const FRealSingle LengthScale,
-	const FRealSingle Time, FVec3& OutNormal, FVec3& OutPosition, FRealSingle& OutTime)
+void TransformSweepOutputsHelper(FVec3 TriMeshScale, const FVec3& HitNormal, const FVec3& HitPosition, const FReal LengthScale,
+	const FReal Time, const bool bOutputIncludesScale, FVec3& OutNormal, FVec3& OutPosition, FReal& OutTime)
 {
-	if (ensure(TriMeshScale != FVec3(0.0f)))
+	// To compute correct position/normal, we tested in the scaled space.
+	// If we want to return the unscaled result we have to undo this scaling, otherwise we output as is.
+	if (!bOutputIncludesScale)
 	{
-		FVec3 InvTriMeshScale = 1.0f / TriMeshScale;
-
-		OutTime = Time / LengthScale;
-		OutNormal = (TriMeshScale * HitNormal).GetSafeNormal();
-		OutPosition = InvTriMeshScale * HitPosition;
+		if (ensure(TriMeshScale != FVec3(0.0f)))
+		{
+			FVec3 InvTriMeshScale = 1.0f / TriMeshScale;
+			OutTime = Time / LengthScale;
+			OutNormal = (TriMeshScale * HitNormal).GetSafeNormal();
+			OutPosition = InvTriMeshScale * HitPosition;
+		}
+	}
+	else
+	{
+		OutTime = Time;
+		OutNormal = HitNormal.GetSafeNormal();
+		OutPosition = HitPosition;
 	}
 }
 
@@ -1278,7 +1288,7 @@ struct FTriangleMeshSweepVisitor
 		return true;
 	}
 
-	bool VisitSweep(const TSpatialVisitorData<int32>& VisitData, FRealSingle& CurDataLength)
+	bool VisitSweep(const TSpatialVisitorData<int32>& VisitData, FRealSingle& CurDataUnscaledLength)
 	{
 		const int32 TriIdx = VisitData.Payload;
 
@@ -1344,51 +1354,51 @@ struct FTriangleMeshSweepVisitor
 				}
 			}
 		}
-		FRealSingle Time;
+		FRealSingle ScaledTime;
+		const FRealSingle ScaledCurLength = LengthScale * CurDataUnscaledLength;
 		bool bIsHitting = false;
+
 		if (UNLIKELY(Tri.IsTooBigForSinglePrecision()))
 		{
 			VectorRegister4Double OutPositionDouble, OutNormalDouble;
-			bIsHitting = GJKRaycast2ImplSimd<VectorRegister4Double>(Tri, QueryGeom, VectorRegister4Double(RotationSimd), VectorRegister4Double(TranslationSimd), VectorRegister4Double(RayDirSimd), LengthScale * CurDataLength, Time, OutPositionDouble, OutNormalDouble, bComputeMTD, GlobalVectorConstants::Double1000);
+			bIsHitting = GJKRaycast2ImplSimd<VectorRegister4Double>(Tri, QueryGeom, VectorRegister4Double(RotationSimd), VectorRegister4Double(TranslationSimd), VectorRegister4Double(RayDirSimd), ScaledCurLength, ScaledTime, OutPositionDouble, OutNormalDouble, bComputeMTD, GlobalVectorConstants::Double1000);
 			OutPositionSimd = MakeVectorRegisterFloatFromDouble(OutPositionDouble);
 			OutNormalSimd = MakeVectorRegisterFloatFromDouble(OutNormalDouble);
 		}
 		else
 		{
-			bIsHitting = GJKRaycast2ImplSimd(Tri, QueryGeom, RotationSimd, TranslationSimd, RayDirSimd, LengthScale * CurDataLength, Time, OutPositionSimd, OutNormalSimd, bComputeMTD, GlobalVectorConstants::Float1000);
+			bIsHitting = GJKRaycast2ImplSimd(Tri, QueryGeom, RotationSimd, TranslationSimd, RayDirSimd, ScaledCurLength, ScaledTime, OutPositionSimd, OutNormalSimd, bComputeMTD, GlobalVectorConstants::Float1000);
 		}
 		if(bIsHitting)
 		{
 			// Parallel sweeps should still be culled (if culls back-face is enabled) if there is no initial intersection. 
 			// We still want to keep a sweep if there is initial intersection.
-			if (Time > 0 && bIsParallelSweepAndCullsBackFace)
+			if (ScaledTime > 0 && bIsParallelSweepAndCullsBackFace)
 			{
 				return true;
 			}
 
-			// Time is world scale, OutTime is local scale.
-			if(Time < LengthScale * OutTime)
+			if(ScaledTime < OutTime)
 			{
+				OutTime = ScaledTime;
 
-				FVec3 HitPosition;
-				FVec3 HitNormal;
 				alignas(16) FRealSingle OutFloat[4];
 				VectorStoreAligned(OutNormalSimd, OutFloat);
-				HitNormal.X = OutFloat[0];
-				HitNormal.Y = OutFloat[1];
-				HitNormal.Z = OutFloat[2];
+				OutNormal.X = OutFloat[0];
+				OutNormal.Y = OutFloat[1];
+				OutNormal.Z = OutFloat[2];
 
 				VectorStoreAligned(OutPositionSimd, OutFloat);
-				HitPosition.X = OutFloat[0];
-				HitPosition.Y = OutFloat[1];
-				HitPosition.Z = OutFloat[2];
-				TransformSweepOutputsHelper(TriMeshScale, HitNormal, HitPosition, LengthScale, Time, OutNormal, OutPosition, OutTime);
+				OutPosition.X = OutFloat[0];
+				OutPosition.Y = OutFloat[1];
+				OutPosition.Z = OutFloat[2];
 
 				OutFaceIndex = TriIdx;
 				VectorStoreFloat3(TriNormal, &OutFaceNormal);
-
+				
 				// We have an initial overlap (MTD of 0 counts). We cannot terminate early as there might be a "better" hit (one with a more negative Time).
-				if(Time <= 0)
+				FRealSingle UnscaledTime = OutTime / LengthScale;
+				if(UnscaledTime <= 0)
 				{
 					// In an ideal world, we'd update CurrentLength to 0 so as to prevent getting any non-zero TOI intersections,
 					// however GJKRaycast does not allow a ray length of 0. We could switch to GJKPenetration, 
@@ -1396,11 +1406,11 @@ struct FTriangleMeshSweepVisitor
 					// Doing this will always produce the correct result (in the edge case we get a TOI between 0 and espilon, we'll still pick the smallest TOI).
 					// This might be slightly slower in this edge case as GJK might've been able to terminate earlier.
 					// #TODO Replace this pathway with a call to GJKPenetration when a SIMD version is created.
-					CurDataLength = UE_KINDA_SMALL_NUMBER;
+					CurDataUnscaledLength = UE_KINDA_SMALL_NUMBER;
 					return true;
 				}
 
-				CurDataLength = static_cast<FRealSingle>(OutTime);
+				CurDataUnscaledLength = static_cast<FRealSingle>(UnscaledTime);
 			}
 		}
 
@@ -1440,7 +1450,7 @@ struct FTriangleMeshSweepVisitorCCD
 {
 	FTriangleMeshSweepVisitorCCD(const FTriangleMeshImplicitObject& InTriMesh, const TArray<TVec3<IdxType>>& InElements, const QueryGeomType& InQueryGeom,
 		const FVec3& InScaledDirNormalized, const FReal InLengthScale, const FReal InLength, const FRigidTransform3& InScaledStartTM, FVec3 InTriMeshScale, FReal InCullsBackFaceSweepsCode, const FReal InIgnorePenetration, const FReal InTargetPenetration)
-		: OutDistance(TNumericLimits<FRealSingle>::Max())
+		: OutTOI(TNumericLimits<FRealSingle>::Max())
 		, OutPhi(TNumericLimits<FRealSingle>::Max())
 		, OutFaceIndex(INDEX_NONE)
 		, TriMesh(InTriMesh)
@@ -1488,7 +1498,7 @@ struct FTriangleMeshSweepVisitorCCD
 		return true;
 	}
 
-	bool VisitSweep(const TSpatialVisitorData<int32>& VisitData, FRealSingle& CurDataLength)
+	bool VisitSweep(const TSpatialVisitorData<int32>& VisitData, FRealSingle& CurDataUnscaledLength)
 	{
 		const int32 TriIdx = VisitData.Payload;
 
@@ -1538,9 +1548,10 @@ struct FTriangleMeshSweepVisitorCCD
 			return true;
 		}
 
-		FRealSingle Distance;
+		const FRealSingle ScaledCurLength = LengthScale * CurDataUnscaledLength;
+		FRealSingle ScaledDistance;
 		VectorRegister4Float PositionSimd, NormalSimd;
-		if (GJKRaycast2ImplSimd(Tri, QueryGeom, RotationSimd, TranslationSimd, RayDirSimd, LengthScale * CurDataLength, Distance, PositionSimd, NormalSimd, true, GlobalVectorConstants::Float1000))
+		if (GJKRaycast2ImplSimd(Tri, QueryGeom, RotationSimd, TranslationSimd, RayDirSimd, ScaledCurLength, ScaledDistance, PositionSimd, NormalSimd, true, GlobalVectorConstants::Float1000))
 		{
 			const VectorRegister4Float DirDotNormalSimd = VectorDot3(RayDirSimd, NormalSimd);
 			FReal DirDotNormal;
@@ -1548,31 +1559,24 @@ struct FTriangleMeshSweepVisitorCCD
 
 			// Calculate the time to reach a depth of TargetPenetration
 			FReal TargetTOI, TargetPhi;
-			ComputeSweptContactTOIAndPhiAtTargetPenetration(DirDotNormal, LengthScale * CurDataLength, FReal(Distance), IgnorePenetration, TargetPenetration, TargetTOI, TargetPhi);
+			ComputeSweptContactTOIAndPhiAtTargetPenetration(DirDotNormal, ScaledCurLength, FReal(ScaledDistance), IgnorePenetration, TargetPenetration, TargetTOI, TargetPhi);
 
-			// TargetDistance is local scale, OutDistance is local scale.
-			const FRealSingle TargetDistance = FRealSingle(TargetTOI) * CurDataLength;
-			if (TargetDistance < CurDataLength)
+			if (ScaledDistance < ScaledCurLength)
 			{
-				CurDataLength = TargetDistance;
-
-				// Transform results back into world scale and save
-				FVec3 HitPosition;
-				FVec3 HitNormal;
-				FRealSingle TOI;	// Unused
-				VectorStoreFloat3(MakeVectorRegisterDouble(PositionSimd), &HitPosition);
-				VectorStoreFloat3(MakeVectorRegisterDouble(NormalSimd), &HitNormal);
-				TransformSweepOutputsHelper(TriMeshScale, HitNormal, HitPosition, LengthScale, TargetDistance, OutNormal, OutPosition, TOI);
-				OutDistance = TargetDistance;
+				CurDataUnscaledLength = FRealSingle(ScaledDistance) / LengthScale;
+				VectorStoreFloat3(MakeVectorRegisterDouble(PositionSimd), &OutPosition);
+				VectorStoreFloat3(MakeVectorRegisterDouble(NormalSimd), &OutNormal);
+
+				OutTOI = TargetTOI;
 				OutPhi = TargetPhi;
 				OutFaceIndex = TriIdx;
 			}
 		}
 
-		return (CurDataLength > 0);
+		return (CurDataUnscaledLength > 0);
 	}
 
-	FReal OutDistance;
+	FReal OutTOI;
 	FReal OutPhi;
 	FVec3 OutPosition;
 	FVec3 OutNormal;
@@ -1637,7 +1641,7 @@ FVec3 SafeInvScale(const FVec3& Scale)
 template <typename QueryGeomType>
 bool FTriangleMeshImplicitObject::SweepGeomImp(const QueryGeomType& QueryGeom, const FRigidTransform3& StartTM, const FVec3& Dir, 
 	const FReal Length, FReal& OutTime, FVec3& OutPosition, FVec3& OutNormal, int32& OutFaceIndex, FVec3& OutFaceNormal, const FReal Thickness, 
-	const bool bComputeMTD, FVec3 TriMeshScale) const
+	const bool bComputeMTD, FVec3 TriMeshScale, bool bOutputIncludesScale) const
 {
 	//QUICK_SCOPE_CYCLE_COUNTER(TrimeshSweep);
 	// Compute scaled sweep inputs to cache in visitor.
@@ -1661,11 +1665,10 @@ bool FTriangleMeshImplicitObject::SweepGeomImp(const QueryGeomType& QueryGeom, c
 		const FVec3 Inflation = QueryBounds.Extents() * InvTriMeshScale.GetAbs() * 0.5 + FVec3(Thickness);
 		FastBVH.template Sweep<VisitorType>(StartPoint, Dir, Length, Inflation, SQVisitor);
 
-		if(SQVisitor.OutTime <= Length)
+		// OutTime is always in scaled space
+		if (SQVisitor.OutTime <= Length * LengthScale)
 		{
-			OutTime = SQVisitor.OutTime;
-			OutPosition = SQVisitor.OutPosition;
-			OutNormal = SQVisitor.OutNormal;
+			TransformSweepOutputsHelper(TriMeshScale, SQVisitor.OutNormal, SQVisitor.OutPosition, LengthScale, SQVisitor.OutTime, bOutputIncludesScale, OutNormal, OutPosition, OutTime);
 			OutFaceIndex = SQVisitor.OutFaceIndex;
 			OutFaceNormal = GetFaceNormal(OutFaceIndex);
 			bHit = true;
@@ -1683,49 +1686,49 @@ bool FTriangleMeshImplicitObject::SweepGeomImp(const QueryGeomType& QueryGeom, c
 	return bHit;
 }
 
-bool FTriangleMeshImplicitObject::SweepGeom(const TSphere<FReal,3>& QueryGeom, const FRigidTransform3& StartTM, const FVec3& Dir, const FReal Length, FReal& OutTime, FVec3& OutPosition, FVec3& OutNormal, int32& OutFaceIndex, FVec3& OutFaceNormal, const FReal Thickness, const bool bComputeMTD, FVec3 TriMeshScale) const
+bool FTriangleMeshImplicitObject::SweepGeom(const TSphere<FReal,3>& QueryGeom, const FRigidTransform3& StartTM, const FVec3& Dir, const FReal Length, FReal& OutTime, FVec3& OutPosition, FVec3& OutNormal, int32& OutFaceIndex, FVec3& OutFaceNormal, const FReal Thickness, const bool bComputeMTD, FVec3 TriMeshScale, bool bOutputIncludesScale) const
 {
-	return SweepGeomImp(QueryGeom, StartTM, Dir, Length, OutTime, OutPosition, OutNormal, OutFaceIndex, OutFaceNormal, Thickness, bComputeMTD, TriMeshScale);
+	return SweepGeomImp(QueryGeom, StartTM, Dir, Length, OutTime, OutPosition, OutNormal, OutFaceIndex, OutFaceNormal, Thickness, bComputeMTD, TriMeshScale, bOutputIncludesScale);
 }
 
-bool FTriangleMeshImplicitObject::SweepGeom(const TBox<FReal, 3>& QueryGeom, const FRigidTransform3& StartTM, const FVec3& Dir, const FReal Length, FReal& OutTime, FVec3& OutPosition, FVec3& OutNormal, int32& OutFaceIndex, FVec3& OutFaceNormal, const FReal Thickness, const bool bComputeMTD, FVec3 TriMeshScale) const
+bool FTriangleMeshImplicitObject::SweepGeom(const TBox<FReal, 3>& QueryGeom, const FRigidTransform3& StartTM, const FVec3& Dir, const FReal Length, FReal& OutTime, FVec3& OutPosition, FVec3& OutNormal, int32& OutFaceIndex, FVec3& OutFaceNormal, const FReal Thickness, const bool bComputeMTD, FVec3 TriMeshScale, bool bOutputIncludesScale) const
 {
-	return SweepGeomImp(QueryGeom, StartTM, Dir, Length, OutTime, OutPosition, OutNormal, OutFaceIndex, OutFaceNormal, Thickness, bComputeMTD, TriMeshScale);
+	return SweepGeomImp(QueryGeom, StartTM, Dir, Length, OutTime, OutPosition, OutNormal, OutFaceIndex, OutFaceNormal, Thickness, bComputeMTD, TriMeshScale, bOutputIncludesScale);
 }
 
-bool FTriangleMeshImplicitObject::SweepGeom(const FCapsule& QueryGeom, const FRigidTransform3& StartTM, const FVec3& Dir, const FReal Length, FReal& OutTime, FVec3& OutPosition, FVec3& OutNormal, int32& OutFaceIndex, FVec3& OutFaceNormal, const FReal Thickness, const bool bComputeMTD, FVec3 TriMeshScale) const
+bool FTriangleMeshImplicitObject::SweepGeom(const FCapsule& QueryGeom, const FRigidTransform3& StartTM, const FVec3& Dir, const FReal Length, FReal& OutTime, FVec3& OutPosition, FVec3& OutNormal, int32& OutFaceIndex, FVec3& OutFaceNormal, const FReal Thickness, const bool bComputeMTD, FVec3 TriMeshScale, bool bOutputIncludesScale) const
 {
-	return SweepGeomImp(QueryGeom, StartTM, Dir, Length, OutTime, OutPosition, OutNormal, OutFaceIndex, OutFaceNormal, Thickness, bComputeMTD, TriMeshScale);
+	return SweepGeomImp(QueryGeom, StartTM, Dir, Length, OutTime, OutPosition, OutNormal, OutFaceIndex, OutFaceNormal, Thickness, bComputeMTD, TriMeshScale, bOutputIncludesScale);
 }
 
-bool FTriangleMeshImplicitObject::SweepGeom(const FConvex& QueryGeom, const FRigidTransform3& StartTM, const FVec3& Dir, const FReal Length, FReal& OutTime, FVec3& OutPosition, FVec3& OutNormal, int32& OutFaceIndex, FVec3& OutFaceNormal, const FReal Thickness, const bool bComputeMTD, FVec3 TriMeshScale) const
+bool FTriangleMeshImplicitObject::SweepGeom(const FConvex& QueryGeom, const FRigidTransform3& StartTM, const FVec3& Dir, const FReal Length, FReal& OutTime, FVec3& OutPosition, FVec3& OutNormal, int32& OutFaceIndex, FVec3& OutFaceNormal, const FReal Thickness, const bool bComputeMTD, FVec3 TriMeshScale, bool bOutputIncludesScale) const
 {
-	return SweepGeomImp(QueryGeom, StartTM, Dir, Length, OutTime, OutPosition, OutNormal, OutFaceIndex, OutFaceNormal, Thickness, bComputeMTD, TriMeshScale);
+	return SweepGeomImp(QueryGeom, StartTM, Dir, Length, OutTime, OutPosition, OutNormal, OutFaceIndex, OutFaceNormal, Thickness, bComputeMTD, TriMeshScale, bOutputIncludesScale);
 }
 
-bool FTriangleMeshImplicitObject::SweepGeom(const TImplicitObjectScaled<FSphere>& QueryGeom, const FRigidTransform3& StartTM, const FVec3& Dir, const FReal Length, FReal& OutTime, FVec3& OutPosition, FVec3& OutNormal, int32& OutFaceIndex, FVec3& OutFaceNormal, const FReal Thickness, const bool bComputeMTD, FVec3 TriMeshScale) const
+bool FTriangleMeshImplicitObject::SweepGeom(const TImplicitObjectScaled<FSphere>& QueryGeom, const FRigidTransform3& StartTM, const FVec3& Dir, const FReal Length, FReal& OutTime, FVec3& OutPosition, FVec3& OutNormal, int32& OutFaceIndex, FVec3& OutFaceNormal, const FReal Thickness, const bool bComputeMTD, FVec3 TriMeshScale, bool bOutputIncludesScale) const
 {
-	return SweepGeomImp(QueryGeom, StartTM, Dir, Length, OutTime, OutPosition, OutNormal, OutFaceIndex, OutFaceNormal, Thickness, bComputeMTD, TriMeshScale);
+	return SweepGeomImp(QueryGeom, StartTM, Dir, Length, OutTime, OutPosition, OutNormal, OutFaceIndex, OutFaceNormal, Thickness, bComputeMTD, TriMeshScale, bOutputIncludesScale);
 }
 
-bool FTriangleMeshImplicitObject::SweepGeom(const TImplicitObjectScaled<TBox<FReal, 3>>& QueryGeom, const FRigidTransform3& StartTM, const FVec3& Dir, const FReal Length, FReal& OutTime, FVec3& OutPosition, FVec3& OutNormal, int32& OutFaceIndex, FVec3& OutFaceNormal, const FReal Thickness, const bool bComputeMTD, FVec3 TriMeshScale) const
+bool FTriangleMeshImplicitObject::SweepGeom(const TImplicitObjectScaled<TBox<FReal, 3>>& QueryGeom, const FRigidTransform3& StartTM, const FVec3& Dir, const FReal Length, FReal& OutTime, FVec3& OutPosition, FVec3& OutNormal, int32& OutFaceIndex, FVec3& OutFaceNormal, const FReal Thickness, const bool bComputeMTD, FVec3 TriMeshScale, bool bOutputIncludesScale) const
 {
-	return SweepGeomImp(QueryGeom, StartTM, Dir, Length, OutTime, OutPosition, OutNormal, OutFaceIndex, OutFaceNormal, Thickness, bComputeMTD, TriMeshScale);
+	return SweepGeomImp(QueryGeom, StartTM, Dir, Length, OutTime, OutPosition, OutNormal, OutFaceIndex, OutFaceNormal, Thickness, bComputeMTD, TriMeshScale, bOutputIncludesScale);
 }
 
-bool FTriangleMeshImplicitObject::SweepGeom(const TImplicitObjectScaled<FCapsule>& QueryGeom, const FRigidTransform3& StartTM, const FVec3& Dir, const FReal Length, FReal& OutTime, FVec3& OutPosition, FVec3& OutNormal, int32& OutFaceIndex, FVec3& OutFaceNormal, const FReal Thickness, const bool bComputeMTD, FVec3 TriMeshScale) const
+bool FTriangleMeshImplicitObject::SweepGeom(const TImplicitObjectScaled<FCapsule>& QueryGeom, const FRigidTransform3& StartTM, const FVec3& Dir, const FReal Length, FReal& OutTime, FVec3& OutPosition, FVec3& OutNormal, int32& OutFaceIndex, FVec3& OutFaceNormal, const FReal Thickness, const bool bComputeMTD, FVec3 TriMeshScale, bool bOutputIncludesScale) const
 {
-	return SweepGeomImp(QueryGeom, StartTM, Dir, Length, OutTime, OutPosition, OutNormal, OutFaceIndex, OutFaceNormal, Thickness, bComputeMTD, TriMeshScale);
+	return SweepGeomImp(QueryGeom, StartTM, Dir, Length, OutTime, OutPosition, OutNormal, OutFaceIndex, OutFaceNormal, Thickness, bComputeMTD, TriMeshScale, bOutputIncludesScale);
 }
 
-bool FTriangleMeshImplicitObject::SweepGeom(const TImplicitObjectScaled<FConvex>& QueryGeom, const FRigidTransform3& StartTM, const FVec3& Dir, const FReal Length, FReal& OutTime, FVec3& OutPosition, FVec3& OutNormal, int32& OutFaceIndex, FVec3& OutFaceNormal, const FReal Thickness, const bool bComputeMTD, FVec3 TriMeshScale) const
+bool FTriangleMeshImplicitObject::SweepGeom(const TImplicitObjectScaled<FConvex>& QueryGeom, const FRigidTransform3& StartTM, const FVec3& Dir, const FReal Length, FReal& OutTime, FVec3& OutPosition, FVec3& OutNormal, int32& OutFaceIndex, FVec3& OutFaceNormal, const FReal Thickness, const bool bComputeMTD, FVec3 TriMeshScale, bool bOutputIncludesScale) const
 {
-	return SweepGeomImp(QueryGeom, StartTM, Dir, Length, OutTime, OutPosition, OutNormal, OutFaceIndex, OutFaceNormal, Thickness, bComputeMTD, TriMeshScale);
+	return SweepGeomImp(QueryGeom, StartTM, Dir, Length, OutTime, OutPosition, OutNormal, OutFaceIndex, OutFaceNormal, Thickness, bComputeMTD, TriMeshScale, bOutputIncludesScale);
 }
 
 
 template<typename QueryGeomType>
-bool FTriangleMeshImplicitObject::SweepGeomCCDImp(const QueryGeomType& QueryGeom, const FRigidTransform3& StartTM, const FVec3& Dir, const FReal Length, const FReal InIgnorePenetration, const FReal InTargetPenetration, FReal& OutTOI, FReal& OutPhi, FVec3& OutPosition, FVec3& OutNormal, int32& OutFaceIndex, FVec3& OutFaceNormal, const FVec3& TriMeshScale) const
+bool FTriangleMeshImplicitObject::SweepGeomCCDImp(const QueryGeomType& QueryGeom, const FRigidTransform3& StartTM, const FVec3& Dir, const FReal Length, const FReal InIgnorePenetration, const FReal InTargetPenetration, FReal& OutTOI, FReal& OutPhi, FVec3& OutPosition, FVec3& OutNormal, int32& OutFaceIndex, FVec3& OutFaceNormal, const FVec3& TriMeshScale, bool bOutputIncludesScale) const
 {
 	// Compute scaled sweep inputs to cache in visitor.
 	FVec3 ScaledDirNormalized;
@@ -1750,21 +1753,21 @@ bool FTriangleMeshImplicitObject::SweepGeomCCDImp(const QueryGeomType& QueryGeom
 		const FVec3 Inflation = QueryBounds.Extents() * InvTriMeshScale.GetAbs() * 0.5;
 		FastBVH.template Sweep<VisitorType>(StartPoint, Dir, Length, Inflation, SQVisitor);
 
-		FReal TOI = (Length > 0) ? SQVisitor.OutDistance / Length : FReal(0);
+		FReal TOI = SQVisitor.OutTOI;
 		FReal Phi = SQVisitor.OutPhi;
 
 		// @todo(chaos): Legacy path to be removed when fully tested. See ComputeSweptContactTOIAndPhiAtTargetPenetration
 		if (!CVars::bCCDNewTargetDepthMode)
 		{
-			LegacyComputeSweptContactTOIAndPhiAtTargetPenetration(FVec3::DotProduct(Dir, SQVisitor.OutNormal), LengthScale * Length, InIgnorePenetration, InTargetPenetration, TOI, Phi);
+			LegacyComputeSweptContactTOIAndPhiAtTargetPenetration(FVec3::DotProduct(ScaledDirNormalized, SQVisitor.OutNormal), LengthScale * Length, InIgnorePenetration, InTargetPenetration, TOI, Phi);
 		}
 
 		if (TOI <= 1)
 		{
 			OutTOI = TOI;
 			OutPhi = Phi;
-			OutPosition = SQVisitor.OutPosition;
-			OutNormal = SQVisitor.OutNormal;
+			FReal UnusedTime = 0;
+			TransformSweepOutputsHelper(TriMeshScale, SQVisitor.OutNormal, SQVisitor.OutPosition, LengthScale, UnusedTime, bOutputIncludesScale, OutNormal, OutPosition, UnusedTime);
 			OutFaceIndex = SQVisitor.OutFaceIndex;
 			OutFaceNormal = GetFaceNormal(OutFaceIndex);
 			bHit = true;
@@ -1782,44 +1785,44 @@ bool FTriangleMeshImplicitObject::SweepGeomCCDImp(const QueryGeomType& QueryGeom
 	return bHit;
 }
 
-bool FTriangleMeshImplicitObject::SweepGeomCCD(const FSphere& QueryGeom, const FRigidTransform3& StartTM, const FVec3& Dir, const FReal Length, const FReal IgnorePenetration, const FReal TargetPenetration, FReal& OutTOI, FReal& OutPhi, FVec3& OutPosition, FVec3& OutNormal, int32& OutFaceIndex, FVec3& OutFaceNormal, const FVec3& TriMeshScale) const
+bool FTriangleMeshImplicitObject::SweepGeomCCD(const FSphere& QueryGeom, const FRigidTransform3& StartTM, const FVec3& Dir, const FReal Length, const FReal IgnorePenetration, const FReal TargetPenetration, FReal& OutTOI, FReal& OutPhi, FVec3& OutPosition, FVec3& OutNormal, int32& OutFaceIndex, FVec3& OutFaceNormal, const FVec3& TriMeshScale, bool bOutputIncludesScale) const
 {
-	return SweepGeomCCDImp(QueryGeom, StartTM, Dir, Length, IgnorePenetration, TargetPenetration, OutTOI, OutPhi, OutPosition, OutNormal, OutFaceIndex, OutFaceNormal, TriMeshScale);
+	return SweepGeomCCDImp(QueryGeom, StartTM, Dir, Length, IgnorePenetration, TargetPenetration, OutTOI, OutPhi, OutPosition, OutNormal, OutFaceIndex, OutFaceNormal, TriMeshScale, bOutputIncludesScale);
 }
 
-bool FTriangleMeshImplicitObject::SweepGeomCCD(const TBox<FReal, 3>& QueryGeom, const FRigidTransform3& StartTM, const FVec3& Dir, const FReal Length, const FReal IgnorePenetration, const FReal TargetPenetration, FReal& OutTOI, FReal& OutPhi, FVec3& OutPosition, FVec3& OutNormal, int32& OutFaceIndex, FVec3& OutFaceNormal, const FVec3& TriMeshScale) const
+bool FTriangleMeshImplicitObject::SweepGeomCCD(const TBox<FReal, 3>& QueryGeom, const FRigidTransform3& StartTM, const FVec3& Dir, const FReal Length, const FReal IgnorePenetration, const FReal TargetPenetration, FReal& OutTOI, FReal& OutPhi, FVec3& OutPosition, FVec3& OutNormal, int32& OutFaceIndex, FVec3& OutFaceNormal, const FVec3& TriMeshScale, bool bOutputIncludesScale) const
 {
-	return SweepGeomCCDImp(QueryGeom, StartTM, Dir, Length, IgnorePenetration, TargetPenetration, OutTOI, OutPhi, OutPosition, OutNormal, OutFaceIndex, OutFaceNormal, TriMeshScale);
+	return SweepGeomCCDImp(QueryGeom, StartTM, Dir, Length, IgnorePenetration, TargetPenetration, OutTOI, OutPhi, OutPosition, OutNormal, OutFaceIndex, OutFaceNormal, TriMeshScale, bOutputIncludesScale);
 }
 
-bool FTriangleMeshImplicitObject::SweepGeomCCD(const FCapsule& QueryGeom, const FRigidTransform3& StartTM, const FVec3& Dir, const FReal Length, const FReal IgnorePenetration, const FReal TargetPenetration, FReal& OutTOI, FReal& OutPhi, FVec3& OutPosition, FVec3& OutNormal, int32& OutFaceIndex, FVec3& OutFaceNormal, const FVec3& TriMeshScale) const
+bool FTriangleMeshImplicitObject::SweepGeomCCD(const FCapsule& QueryGeom, const FRigidTransform3& StartTM, const FVec3& Dir, const FReal Length, const FReal IgnorePenetration, const FReal TargetPenetration, FReal& OutTOI, FReal& OutPhi, FVec3& OutPosition, FVec3& OutNormal, int32& OutFaceIndex, FVec3& OutFaceNormal, const FVec3& TriMeshScale, bool bOutputIncludesScale) const
 {
-	return SweepGeomCCDImp(QueryGeom, StartTM, Dir, Length, IgnorePenetration, TargetPenetration, OutTOI, OutPhi, OutPosition, OutNormal, OutFaceIndex, OutFaceNormal, TriMeshScale);
+	return SweepGeomCCDImp(QueryGeom, StartTM, Dir, Length, IgnorePenetration, TargetPenetration, OutTOI, OutPhi, OutPosition, OutNormal, OutFaceIndex, OutFaceNormal, TriMeshScale, bOutputIncludesScale);
 }
 
-bool FTriangleMeshImplicitObject::SweepGeomCCD(const FConvex& QueryGeom, const FRigidTransform3& StartTM, const FVec3& Dir, const FReal Length, const FReal IgnorePenetration, const FReal TargetPenetration, FReal& OutTOI, FReal& OutPhi, FVec3& OutPosition, FVec3& OutNormal, int32& OutFaceIndex, FVec3& OutFaceNormal, const FVec3& TriMeshScale) const
+bool FTriangleMeshImplicitObject::SweepGeomCCD(const FConvex& QueryGeom, const FRigidTransform3& StartTM, const FVec3& Dir, const FReal Length, const FReal IgnorePenetration, const FReal TargetPenetration, FReal& OutTOI, FReal& OutPhi, FVec3& OutPosition, FVec3& OutNormal, int32& OutFaceIndex, FVec3& OutFaceNormal, const FVec3& TriMeshScale, bool bOutputIncludesScale) const
 {
-	return SweepGeomCCDImp(QueryGeom, StartTM, Dir, Length, IgnorePenetration, TargetPenetration, OutTOI, OutPhi, OutPosition, OutNormal, OutFaceIndex, OutFaceNormal, TriMeshScale);
+	return SweepGeomCCDImp(QueryGeom, StartTM, Dir, Length, IgnorePenetration, TargetPenetration, OutTOI, OutPhi, OutPosition, OutNormal, OutFaceIndex, OutFaceNormal, TriMeshScale, bOutputIncludesScale);
 }
 
-bool FTriangleMeshImplicitObject::SweepGeomCCD(const TImplicitObjectScaled<FSphere>& QueryGeom, const FRigidTransform3& StartTM, const FVec3& Dir, const FReal Length, const FReal IgnorePenetration, const FReal TargetPenetration, FReal& OutTOI, FReal& OutPhi, FVec3& OutPosition, FVec3& OutNormal, int32& OutFaceIndex, FVec3& OutFaceNormal, const FVec3& TriMeshScale) const
+bool FTriangleMeshImplicitObject::SweepGeomCCD(const TImplicitObjectScaled<FSphere>& QueryGeom, const FRigidTransform3& StartTM, const FVec3& Dir, const FReal Length, const FReal IgnorePenetration, const FReal TargetPenetration, FReal& OutTOI, FReal& OutPhi, FVec3& OutPosition, FVec3& OutNormal, int32& OutFaceIndex, FVec3& OutFaceNormal, const FVec3& TriMeshScale, bool bOutputIncludesScale) const
 {
-	return SweepGeomCCDImp(QueryGeom, StartTM, Dir, Length, IgnorePenetration, TargetPenetration, OutTOI, OutPhi, OutPosition, OutNormal, OutFaceIndex, OutFaceNormal, TriMeshScale);
+	return SweepGeomCCDImp(QueryGeom, StartTM, Dir, Length, IgnorePenetration, TargetPenetration, OutTOI, OutPhi, OutPosition, OutNormal, OutFaceIndex, OutFaceNormal, TriMeshScale, bOutputIncludesScale);
 }
 
-bool FTriangleMeshImplicitObject::SweepGeomCCD(const TImplicitObjectScaled<TBox<FReal, 3>>& QueryGeom, const FRigidTransform3& StartTM, const FVec3& Dir, const FReal Length, const FReal IgnorePenetration, const FReal TargetPenetration, FReal& OutTOI, FReal& OutPhi, FVec3& OutPosition, FVec3& OutNormal, int32& OutFaceIndex, FVec3& OutFaceNormal, const FVec3& TriMeshScale) const
+bool FTriangleMeshImplicitObject::SweepGeomCCD(const TImplicitObjectScaled<TBox<FReal, 3>>& QueryGeom, const FRigidTransform3& StartTM, const FVec3& Dir, const FReal Length, const FReal IgnorePenetration, const FReal TargetPenetration, FReal& OutTOI, FReal& OutPhi, FVec3& OutPosition, FVec3& OutNormal, int32& OutFaceIndex, FVec3& OutFaceNormal, const FVec3& TriMeshScale, bool bOutputIncludesScale) const
 {
-	return SweepGeomCCDImp(QueryGeom, StartTM, Dir, Length, IgnorePenetration, TargetPenetration, OutTOI, OutPhi, OutPosition, OutNormal, OutFaceIndex, OutFaceNormal, TriMeshScale);
+	return SweepGeomCCDImp(QueryGeom, StartTM, Dir, Length, IgnorePenetration, TargetPenetration, OutTOI, OutPhi, OutPosition, OutNormal, OutFaceIndex, OutFaceNormal, TriMeshScale, bOutputIncludesScale);
 }
 
-bool FTriangleMeshImplicitObject::SweepGeomCCD(const TImplicitObjectScaled<FCapsule>& QueryGeom, const FRigidTransform3& StartTM, const FVec3& Dir, const FReal Length, const FReal IgnorePenetration, const FReal TargetPenetration, FReal& OutTOI, FReal& OutPhi, FVec3& OutPosition, FVec3& OutNormal, int32& OutFaceIndex, FVec3& OutFaceNormal, const FVec3& TriMeshScale) const
+bool FTriangleMeshImplicitObject::SweepGeomCCD(const TImplicitObjectScaled<FCapsule>& QueryGeom, const FRigidTransform3& StartTM, const FVec3& Dir, const FReal Length, const FReal IgnorePenetration, const FReal TargetPenetration, FReal& OutTOI, FReal& OutPhi, FVec3& OutPosition, FVec3& OutNormal, int32& OutFaceIndex, FVec3& OutFaceNormal, const FVec3& TriMeshScale, bool bOutputIncludesScale) const
 {
-	return SweepGeomCCDImp(QueryGeom, StartTM, Dir, Length, IgnorePenetration, TargetPenetration, OutTOI, OutPhi, OutPosition, OutNormal, OutFaceIndex, OutFaceNormal, TriMeshScale);
+	return SweepGeomCCDImp(QueryGeom, StartTM, Dir, Length, IgnorePenetration, TargetPenetration, OutTOI, OutPhi, OutPosition, OutNormal, OutFaceIndex, OutFaceNormal, TriMeshScale, bOutputIncludesScale);
 }
 
-bool FTriangleMeshImplicitObject::SweepGeomCCD(const TImplicitObjectScaled<FConvex>& QueryGeom, const FRigidTransform3& StartTM, const FVec3& Dir, const FReal Length, const FReal IgnorePenetration, const FReal TargetPenetration, FReal& OutTOI, FReal& OutPhi, FVec3& OutPosition, FVec3& OutNormal, int32& OutFaceIndex, FVec3& OutFaceNormal, const FVec3& TriMeshScale) const
+bool FTriangleMeshImplicitObject::SweepGeomCCD(const TImplicitObjectScaled<FConvex>& QueryGeom, const FRigidTransform3& StartTM, const FVec3& Dir, const FReal Length, const FReal IgnorePenetration, const FReal TargetPenetration, FReal& OutTOI, FReal& OutPhi, FVec3& OutPosition, FVec3& OutNormal, int32& OutFaceIndex, FVec3& OutFaceNormal, const FVec3& TriMeshScale, bool bOutputIncludesScale) const
 {
-	return SweepGeomCCDImp(QueryGeom, StartTM, Dir, Length, IgnorePenetration, TargetPenetration, OutTOI, OutPhi, OutPosition, OutNormal, OutFaceIndex, OutFaceNormal, TriMeshScale);
+	return SweepGeomCCDImp(QueryGeom, StartTM, Dir, Length, IgnorePenetration, TargetPenetration, OutTOI, OutPhi, OutPosition, OutNormal, OutFaceIndex, OutFaceNormal, TriMeshScale, bOutputIncludesScale);
 }
 
 
diff --git a/Engine/Source/Runtime/Experimental/Chaos/Public/Chaos/ImplicitObjectScaled.h b/Engine/Source/Runtime/Experimental/Chaos/Public/Chaos/ImplicitObjectScaled.h
--- a/Engine/Source/Runtime/Experimental/Chaos/Public/Chaos/ImplicitObjectScaled.h
+++ b/Engine/Source/Runtime/Experimental/Chaos/Public/Chaos/ImplicitObjectScaled.h
@@ -662,22 +662,44 @@ public:
 			const T UnscaledLength = Length * LengthScale;
 			const TVector<T, d> UnscaledDir = UnscaledDirDenorm * LengthScaleInv;
 
-			TVector<T, d> UnscaledPosition;
-			TVector<T, d> UnscaledNormal;
-			T UnscaledTime;
+			
 
 			TRigidTransform<T, d> BToATMNoScale(BToATM.GetLocation() * FVec3(MInvScale), BToATM.GetRotation());
-		
-			if (MObject->SweepGeom(B, BToATMNoScale, UnscaledDir, UnscaledLength, UnscaledTime, UnscaledPosition, UnscaledNormal, OutFaceIndex, OutFaceNormal, Thickness, bComputeMTD, MScale))
+
+			constexpr bool bIncludeScaleInResults = std::is_same_v<TConcrete, FTriangleMeshImplicitObject>;
+			if constexpr (!bIncludeScaleInResults)
 			{
-				const T NewTime = LengthScaleInv * UnscaledTime;
-				//We double check that NewTime < Length because of potential precision issues. When that happens we always keep the shortest hit first
-				if (NewTime < Length)
+				TVector<T, d> UnscaledPosition;
+				TVector<T, d> UnscaledNormal;
+				T UnscaledTime;
+				if (MObject->SweepGeom(B, BToATMNoScale, UnscaledDir, UnscaledLength, UnscaledTime, UnscaledPosition, UnscaledNormal, OutFaceIndex, OutFaceNormal, Thickness, bComputeMTD, MScale))
 				{
-					OutTime = NewTime;
-					LocalPosition = MScale * UnscaledPosition;
-					LocalNormal = (MInvScale * UnscaledNormal).GetSafeNormal();
-					return true;
+					const T NewTime = LengthScaleInv * UnscaledTime;
+					//We double check that NewTime < Length because of potential precision issues. When that happens we always keep the shortest hit first
+					if (NewTime < Length)
+					{
+						OutTime = NewTime;
+						LocalPosition = MScale * UnscaledPosition;
+						LocalNormal = (MInvScale * UnscaledNormal).GetSafeNormal();
+						return true;
+					}
+				}
+			}
+			else
+			{
+				TVector<T, d> OutScaledPosition;
+				TVector<T, d> OutScaledNormal;
+				T OutScaledTime;
+				if (MObject->SweepGeom(B, BToATMNoScale, UnscaledDir, UnscaledLength, OutScaledTime, OutScaledPosition, OutScaledNormal, OutFaceIndex, OutFaceNormal, Thickness, bComputeMTD, MScale, bIncludeScaleInResults))
+				{
+					//We double check that OutScaledTime < Length because of potential precision issues. When that happens we always keep the shortest hit first
+					if (OutScaledTime < Length)
+					{
+						OutTime = OutScaledTime;
+						LocalPosition = OutScaledPosition;
+						LocalNormal = OutScaledNormal.GetSafeNormal();
+						return true;
+					}
 				}
 			}
 		}
@@ -700,22 +722,42 @@ public:
 			const T UnscaledLength = Length * LengthScale;
 			const TVector<T, d> UnscaledDir = UnscaledDirDenorm * LengthScaleInv;
 
-			TVector<T, d> UnscaledPosition;
-			TVector<T, d> UnscaledNormal;
 			T TOI;
 			T Phi;
 
 			TRigidTransform<T, d> BToATMNoScale(BToATM.GetLocation() * FVec3(MInvScale), BToATM.GetRotation());
-
-			if (MObject->SweepGeomCCD(B, BToATMNoScale, UnscaledDir, UnscaledLength, IgnorePenetration, TargetPenetration, TOI, Phi, UnscaledPosition, UnscaledNormal, OutFaceIndex, OutFaceNormal, MScale))
+			
+			constexpr bool bIncludeScaleInResults = std::is_same_v<TConcrete, FTriangleMeshImplicitObject>;
+			if constexpr (!bIncludeScaleInResults)
 			{
-				if (TOI < FReal(1))
+				TVector<T, d> UnscaledPosition;
+				TVector<T, d> UnscaledNormal;
+				if (MObject->SweepGeomCCD(B, BToATMNoScale, UnscaledDir, UnscaledLength, IgnorePenetration, TargetPenetration, TOI, Phi, UnscaledPosition, UnscaledNormal, OutFaceIndex, OutFaceNormal, MScale))
 				{
-					OutTOI = TOI;
-					OutPhi = Phi;
-					LocalPosition = MScale * UnscaledPosition;
-					LocalNormal = (MInvScale * UnscaledNormal).GetSafeNormal();
-					return true;
+					if (TOI < FReal(1))
+					{
+						OutTOI = TOI;
+						OutPhi = Phi;
+						LocalPosition = MScale * UnscaledPosition;
+						LocalNormal = (MInvScale * UnscaledNormal).GetSafeNormal();
+						return true;
+					}
+				}
+			}
+			else
+			{
+				TVector<T, d> ScaledPosition;
+				TVector<T, d> ScaledNormal;
+				if (MObject->SweepGeomCCD(B, BToATMNoScale, UnscaledDir, UnscaledLength, IgnorePenetration, TargetPenetration, TOI, Phi, ScaledPosition, ScaledNormal, OutFaceIndex, OutFaceNormal, MScale, bIncludeScaleInResults))
+				{
+					if (TOI < FReal(1))
+					{
+						OutTOI = TOI;
+						OutPhi = Phi;
+						LocalPosition = ScaledPosition;
+						LocalNormal = ScaledNormal.GetSafeNormal();
+						return true;
+					}
 				}
 			}
 		}
diff --git a/Engine/Source/Runtime/Experimental/Chaos/Public/Chaos/TriangleMeshImplicitObject.h b/Engine/Source/Runtime/Experimental/Chaos/Public/Chaos/TriangleMeshImplicitObject.h
--- a/Engine/Source/Runtime/Experimental/Chaos/Public/Chaos/TriangleMeshImplicitObject.h
+++ b/Engine/Source/Runtime/Experimental/Chaos/Public/Chaos/TriangleMeshImplicitObject.h
@@ -552,25 +552,25 @@ namespace Chaos
 		CHAOS_API bool OverlapGeom(const TImplicitObjectScaled<FCapsule>& QueryGeom, const FRigidTransform3& QueryTM, const FReal Thickness, FMTDInfo* OutMTD = nullptr, FVec3 TriMeshScale = FVec3(1.0f)) const;
 		CHAOS_API bool OverlapGeom(const TImplicitObjectScaled<FConvex>& QueryGeom, const FRigidTransform3& QueryTM, const FReal Thickness, FMTDInfo* OutMTD = nullptr, FVec3 TriMeshScale = FVec3(1.0f)) const;
 
-		CHAOS_API bool SweepGeom(const FSphere& QueryGeom, const FRigidTransform3& StartTM, const FVec3& Dir, const FReal Length, FReal& OutTime, FVec3& OutPosition, FVec3& OutNormal, int32& OutFaceIndex, FVec3& OutFaceNormal, const FReal Thickness = 0, const bool bComputeMTD = false, FVec3 TriMeshScale = FVec3(1.0f)) const;
-		CHAOS_API bool SweepGeom(const TBox<FReal, 3>& QueryGeom, const FRigidTransform3& StartTM, const FVec3& Dir, const FReal Length, FReal& OutTime, FVec3& OutPosition, FVec3& OutNormal, int32& OutFaceIndex, FVec3& OutFaceNormal, const FReal Thickness = 0, const bool bComputeMTD = false, FVec3 TriMeshScale = FVec3(1.0f)) const;
-		CHAOS_API bool SweepGeom(const FCapsule& QueryGeom, const FRigidTransform3& StartTM, const FVec3& Dir, const FReal Length, FReal& OutTime, FVec3& OutPosition, FVec3& OutNormal, int32& OutFaceIndex, FVec3& OutFaceNormal, const FReal Thickness = 0, const bool bComputeMTD = false, FVec3 TriMeshScale = FVec3(1.0f)) const;
-		CHAOS_API bool SweepGeom(const FConvex& QueryGeom, const FRigidTransform3& StartTM, const FVec3& Dir, const FReal Length, FReal& OutTime, FVec3& OutPosition, FVec3& OutNormal, int32& OutFaceIndex, FVec3& OutFaceNormal, const FReal Thickness = 0, const bool bComputeMTD = false, FVec3 TriMeshScale = FVec3(1.0f)) const;
+		CHAOS_API bool SweepGeom(const FSphere& QueryGeom, const FRigidTransform3& StartTM, const FVec3& Dir, const FReal Length, FReal& OutTime, FVec3& OutPosition, FVec3& OutNormal, int32& OutFaceIndex, FVec3& OutFaceNormal, const FReal Thickness = 0, const bool bComputeMTD = false, FVec3 TriMeshScale = FVec3(1.0f), bool bOutputIncludesScale = false) const;
+		CHAOS_API bool SweepGeom(const TBox<FReal, 3>& QueryGeom, const FRigidTransform3& StartTM, const FVec3& Dir, const FReal Length, FReal& OutTime, FVec3& OutPosition, FVec3& OutNormal, int32& OutFaceIndex, FVec3& OutFaceNormal, const FReal Thickness = 0, const bool bComputeMTD = false, FVec3 TriMeshScale = FVec3(1.0f), bool bOutputIncludesScale = false) const;
+		CHAOS_API bool SweepGeom(const FCapsule& QueryGeom, const FRigidTransform3& StartTM, const FVec3& Dir, const FReal Length, FReal& OutTime, FVec3& OutPosition, FVec3& OutNormal, int32& OutFaceIndex, FVec3& OutFaceNormal, const FReal Thickness = 0, const bool bComputeMTD = false, FVec3 TriMeshScale = FVec3(1.0f), bool bOutputIncludesScale = false) const;
+		CHAOS_API bool SweepGeom(const FConvex& QueryGeom, const FRigidTransform3& StartTM, const FVec3& Dir, const FReal Length, FReal& OutTime, FVec3& OutPosition, FVec3& OutNormal, int32& OutFaceIndex, FVec3& OutFaceNormal, const FReal Thickness = 0, const bool bComputeMTD = false, FVec3 TriMeshScale = FVec3(1.0f), bool bOutputIncludesScale = false) const;
 
-		CHAOS_API bool SweepGeom(const TImplicitObjectScaled<FSphere>& QueryGeom, const FRigidTransform3& StartTM, const FVec3& Dir, const FReal Length, FReal& OutTime, FVec3& OutPosition, FVec3& OutNormal, int32& OutFaceIndex, FVec3& OutFaceNormal, const FReal Thickness = 0, const bool bComputeMTD = false, FVec3 TriMeshScale = FVec3(1.0f)) const;
-		CHAOS_API bool SweepGeom(const TImplicitObjectScaled<TBox<FReal, 3>>& QueryGeom, const FRigidTransform3& StartTM, const FVec3& Dir, const FReal Length, FReal& OutTime, FVec3& OutPosition, FVec3& OutNormal, int32& OutFaceIndex, FVec3& OutFaceNormal, const FReal Thickness = 0, const bool bComputeMTD = false, FVec3 TriMeshScale = FVec3(1.0f)) const;
-		CHAOS_API bool SweepGeom(const TImplicitObjectScaled<FCapsule>& QueryGeom, const FRigidTransform3& StartTM, const FVec3& Dir, const FReal Length, FReal& OutTime, FVec3& OutPosition, FVec3& OutNormal, int32& OutFaceIndex, FVec3& OutFaceNormal, const FReal Thickness = 0, const bool bComputeMTD = false, FVec3 TriMeshScale = FVec3(1.0f)) const;
-		CHAOS_API bool SweepGeom(const TImplicitObjectScaled<FConvex>& QueryGeom, const FRigidTransform3& StartTM, const FVec3& Dir, const FReal Length, FReal& OutTime, FVec3& OutPosition, FVec3& OutNormal, int32& OutFaceIndex, FVec3& OutFaceNormal, const FReal Thickness = 0, const bool bComputeMTD = false, FVec3 TriMeshScale = FVec3(1.0f)) const;
+		CHAOS_API bool SweepGeom(const TImplicitObjectScaled<FSphere>& QueryGeom, const FRigidTransform3& StartTM, const FVec3& Dir, const FReal Length, FReal& OutTime, FVec3& OutPosition, FVec3& OutNormal, int32& OutFaceIndex, FVec3& OutFaceNormal, const FReal Thickness = 0, const bool bComputeMTD = false, FVec3 TriMeshScale = FVec3(1.0f), bool bOutputIncludesScale = false) const;
+		CHAOS_API bool SweepGeom(const TImplicitObjectScaled<TBox<FReal, 3>>& QueryGeom, const FRigidTransform3& StartTM, const FVec3& Dir, const FReal Length, FReal& OutTime, FVec3& OutPosition, FVec3& OutNormal, int32& OutFaceIndex, FVec3& OutFaceNormal, const FReal Thickness = 0, const bool bComputeMTD = false, FVec3 TriMeshScale = FVec3(1.0f), bool bOutputIncludesScale = false) const;
+		CHAOS_API bool SweepGeom(const TImplicitObjectScaled<FCapsule>& QueryGeom, const FRigidTransform3& StartTM, const FVec3& Dir, const FReal Length, FReal& OutTime, FVec3& OutPosition, FVec3& OutNormal, int32& OutFaceIndex, FVec3& OutFaceNormal, const FReal Thickness = 0, const bool bComputeMTD = false, FVec3 TriMeshScale = FVec3(1.0f), bool bOutputIncludesScale = false) const;
+		CHAOS_API bool SweepGeom(const TImplicitObjectScaled<FConvex>& QueryGeom, const FRigidTransform3& StartTM, const FVec3& Dir, const FReal Length, FReal& OutTime, FVec3& OutPosition, FVec3& OutNormal, int32& OutFaceIndex, FVec3& OutFaceNormal, const FReal Thickness = 0, const bool bComputeMTD = false, FVec3 TriMeshScale = FVec3(1.0f), bool bOutputIncludesScale = false) const;
 
 		// Sweep used for CCD. Ignores triangles we penetrate by less than IgnorePenetration, and calculate the TOI for a depth of TargetPenetration. If both are zero, this is equivalent to SweepGeom
-		CHAOS_API bool SweepGeomCCD(const FSphere& QueryGeom, const FRigidTransform3& StartTM, const FVec3& Dir, const FReal Length, const FReal IgnorePenetration, const FReal TargetPenetration, FReal& OutTOI, FReal& OutPhi, FVec3& OutPosition, FVec3& OutNormal, int32& OutFaceIndex, FVec3& OutFaceNormal, const FVec3& TriMeshScale = FVec3(1.0f)) const;
-		CHAOS_API bool SweepGeomCCD(const TBox<FReal, 3>& QueryGeom, const FRigidTransform3& StartTM, const FVec3& Dir, const FReal Length, const FReal IgnorePenetration, const FReal TargetPenetration, FReal& OutTOI, FReal& OutPhi, FVec3& OutPosition, FVec3& OutNormal, int32& OutFaceIndex, FVec3& OutFaceNormal, const FVec3& TriMeshScale = FVec3(1.0f)) const;
-		CHAOS_API bool SweepGeomCCD(const FCapsule& QueryGeom, const FRigidTransform3& StartTM, const FVec3& Dir, const FReal Length, const FReal IgnorePenetration, const FReal TargetPenetration, FReal& OutTOI, FReal& OutPhi, FVec3& OutPosition, FVec3& OutNormal, int32& OutFaceIndex, FVec3& OutFaceNormal, const FVec3& TriMeshScale = FVec3(1.0f)) const;
-		CHAOS_API bool SweepGeomCCD(const FConvex& QueryGeom, const FRigidTransform3& StartTM, const FVec3& Dir, const FReal Length, const FReal IgnorePenetration, const FReal TargetPenetration, FReal& OutTOI, FReal& OutPhi, FVec3& OutPosition, FVec3& OutNormal, int32& OutFaceIndex, FVec3& OutFaceNormal, const FVec3& TriMeshScale = FVec3(1.0f)) const;
-		CHAOS_API bool SweepGeomCCD(const TImplicitObjectScaled<FSphere>& QueryGeom, const FRigidTransform3& StartTM, const FVec3& Dir, const FReal Length, const FReal IgnorePenetration, const FReal TargetPenetration, FReal& OutTOI, FReal& OutPhi, FVec3& OutPosition, FVec3& OutNormal, int32& OutFaceIndex, FVec3& OutFaceNormal, const FVec3& TriMeshScale = FVec3(1.0f)) const;
-		CHAOS_API bool SweepGeomCCD(const TImplicitObjectScaled<TBox<FReal, 3>>& QueryGeom, const FRigidTransform3& StartTM, const FVec3& Dir, const FReal Length, const FReal IgnorePenetration, const FReal TargetPenetration, FReal& OutTOI, FReal& OutPhi, FVec3& OutPosition, FVec3& OutNormal, int32& OutFaceIndex, FVec3& OutFaceNormal, const FVec3& TriMeshScale = FVec3(1.0f)) const;
-		CHAOS_API bool SweepGeomCCD(const TImplicitObjectScaled<FCapsule>& QueryGeom, const FRigidTransform3& StartTM, const FVec3& Dir, const FReal Length, const FReal IgnorePenetration, const FReal TargetPenetration, FReal& OutTOI, FReal& OutPhi, FVec3& OutPosition, FVec3& OutNormal, int32& OutFaceIndex, FVec3& OutFaceNormal, const FVec3& TriMeshScale = FVec3(1.0f)) const;
-		CHAOS_API bool SweepGeomCCD(const TImplicitObjectScaled<FConvex>& QueryGeom, const FRigidTransform3& StartTM, const FVec3& Dir, const FReal Length, const FReal IgnorePenetration, const FReal TargetPenetration, FReal& OutTOI, FReal& OutPhi, FVec3& OutPosition, FVec3& OutNormal, int32& OutFaceIndex, FVec3& OutFaceNormal, const FVec3& TriMeshScale = FVec3(1.0f)) const;
+		CHAOS_API bool SweepGeomCCD(const FSphere& QueryGeom, const FRigidTransform3& StartTM, const FVec3& Dir, const FReal Length, const FReal IgnorePenetration, const FReal TargetPenetration, FReal& OutTOI, FReal& OutPhi, FVec3& OutPosition, FVec3& OutNormal, int32& OutFaceIndex, FVec3& OutFaceNormal, const FVec3& TriMeshScale = FVec3(1.0f), bool bOutputIncludesScale = false) const;
+		CHAOS_API bool SweepGeomCCD(const TBox<FReal, 3>& QueryGeom, const FRigidTransform3& StartTM, const FVec3& Dir, const FReal Length, const FReal IgnorePenetration, const FReal TargetPenetration, FReal& OutTOI, FReal& OutPhi, FVec3& OutPosition, FVec3& OutNormal, int32& OutFaceIndex, FVec3& OutFaceNormal, const FVec3& TriMeshScale = FVec3(1.0f), bool bOutputIncludesScale = false) const;
+		CHAOS_API bool SweepGeomCCD(const FCapsule& QueryGeom, const FRigidTransform3& StartTM, const FVec3& Dir, const FReal Length, const FReal IgnorePenetration, const FReal TargetPenetration, FReal& OutTOI, FReal& OutPhi, FVec3& OutPosition, FVec3& OutNormal, int32& OutFaceIndex, FVec3& OutFaceNormal, const FVec3& TriMeshScale = FVec3(1.0f), bool bOutputIncludesScale = false) const;
+		CHAOS_API bool SweepGeomCCD(const FConvex& QueryGeom, const FRigidTransform3& StartTM, const FVec3& Dir, const FReal Length, const FReal IgnorePenetration, const FReal TargetPenetration, FReal& OutTOI, FReal& OutPhi, FVec3& OutPosition, FVec3& OutNormal, int32& OutFaceIndex, FVec3& OutFaceNormal, const FVec3& TriMeshScale = FVec3(1.0f), bool bOutputIncludesScale = false) const;
+		CHAOS_API bool SweepGeomCCD(const TImplicitObjectScaled<FSphere>& QueryGeom, const FRigidTransform3& StartTM, const FVec3& Dir, const FReal Length, const FReal IgnorePenetration, const FReal TargetPenetration, FReal& OutTOI, FReal& OutPhi, FVec3& OutPosition, FVec3& OutNormal, int32& OutFaceIndex, FVec3& OutFaceNormal, const FVec3& TriMeshScale = FVec3(1.0f), bool bOutputIncludesScale = false) const;
+		CHAOS_API bool SweepGeomCCD(const TImplicitObjectScaled<TBox<FReal, 3>>& QueryGeom, const FRigidTransform3& StartTM, const FVec3& Dir, const FReal Length, const FReal IgnorePenetration, const FReal TargetPenetration, FReal& OutTOI, FReal& OutPhi, FVec3& OutPosition, FVec3& OutNormal, int32& OutFaceIndex, FVec3& OutFaceNormal, const FVec3& TriMeshScale = FVec3(1.0f), bool bOutputIncludesScale = false) const;
+		CHAOS_API bool SweepGeomCCD(const TImplicitObjectScaled<FCapsule>& QueryGeom, const FRigidTransform3& StartTM, const FVec3& Dir, const FReal Length, const FReal IgnorePenetration, const FReal TargetPenetration, FReal& OutTOI, FReal& OutPhi, FVec3& OutPosition, FVec3& OutNormal, int32& OutFaceIndex, FVec3& OutFaceNormal, const FVec3& TriMeshScale = FVec3(1.0f), bool bOutputIncludesScale = false) const;
+		CHAOS_API bool SweepGeomCCD(const TImplicitObjectScaled<FConvex>& QueryGeom, const FRigidTransform3& StartTM, const FVec3& Dir, const FReal Length, const FReal IgnorePenetration, const FReal TargetPenetration, FReal& OutTOI, FReal& OutPhi, FVec3& OutPosition, FVec3& OutNormal, int32& OutFaceIndex, FVec3& OutFaceNormal, const FVec3& TriMeshScale = FVec3(1.0f), bool bOutputIncludesScale = false) const;
 
 		CHAOS_API bool GJKContactPoint(const FSphere& QueryGeom, const FRigidTransform3& QueryTM, const FReal Thickness, FVec3& Location, FVec3& Normal, FReal& Penetration, int32& FaceIndex) const;
 		CHAOS_API bool GJKContactPoint(const TBox<FReal, 3>& QueryGeom, const FRigidTransform3& QueryTM, const FReal Thickness, FVec3& Location, FVec3& Normal, FReal& Penetration, int32& FaceIndex) const;
@@ -917,10 +917,10 @@ namespace Chaos
 		bool OverlapGeomImp(const QueryGeomType& QueryGeom, const FRigidTransform3& QueryTM, const FReal Thickness, FMTDInfo* OutMTD = nullptr, FVec3 TriMeshScale = FVec3(1.0f)) const;
 
 		template<typename QueryGeomType>
-		bool SweepGeomImp(const QueryGeomType& QueryGeom, const FRigidTransform3& StartTM, const FVec3& Dir, const FReal Length, FReal& OutTime, FVec3& OutPosition, FVec3& OutNormal, int32& OutFaceIndex, FVec3& OutFaceNormal, const FReal Thickness, const bool bComputeMTD, FVec3 TriMeshScale = FVec3(1.0f)) const;
+		bool SweepGeomImp(const QueryGeomType& QueryGeom, const FRigidTransform3& StartTM, const FVec3& Dir, const FReal Length, FReal& OutTime, FVec3& OutPosition, FVec3& OutNormal, int32& OutFaceIndex, FVec3& OutFaceNormal, const FReal Thickness, const bool bComputeMTD, FVec3 TriMeshScale = FVec3(1.0f), bool bOutputIncludesScale = false) const;
 
 		template <typename QueryGeomType>
-		bool SweepGeomCCDImp(const QueryGeomType& QueryGeom, const FRigidTransform3& StartTM, const FVec3& Dir, const FReal Length, const FReal IgnorePenetration, const FReal TargetPenetration, FReal& OutTOI, FReal& OutPhi, FVec3& OutPosition, FVec3& OutNormal, int32& OutFaceIndex, FVec3& OutFaceNormal, const FVec3& TriMeshScale) const;
+		bool SweepGeomCCDImp(const QueryGeomType& QueryGeom, const FRigidTransform3& StartTM, const FVec3& Dir, const FReal Length, const FReal IgnorePenetration, const FReal TargetPenetration, FReal& OutTOI, FReal& OutPhi, FVec3& OutPosition, FVec3& OutNormal, int32& OutFaceIndex, FVec3& OutFaceNormal, const FVec3& TriMeshScale, bool bOutputIncludesScale) const;
 
 		template <typename IdxType>
 		bool RaycastImp(const TArray<TVec3<IdxType>>& Elements, const FVec3& StartPoint, const FVec3& Dir, const FReal Length, const FReal Thickness, FReal& OutTime, FVec3& OutPosition, FVec3& OutNormal, int32& OutFaceIndex) const;
