Geometry collection : fix crash with replication
- Lambdas running on the Physics thread were capturing member variables of the component and be schedulefor later execution
- if the component was deleted before the execution , the lambda would crash accessing the component pointer
- now the lambdas are using a weakPtr on the component to validate its existence before using it

#Synced-CL 48227990

diff --git a/Engine/Source/Runtime/Experimental/GeometryCollectionEngine/Private/GeometryCollection/GeometryCollectionComponent.cpp b/Engine/Source/Runtime/Experimental/GeometryCollectionEngine/Private/GeometryCollection/GeometryCollectionComponent.cpp
--- a/Engine/Source/Runtime/Experimental/GeometryCollectionEngine/Private/GeometryCollection/GeometryCollectionComponent.cpp
+++ b/Engine/Source/Runtime/Experimental/GeometryCollectionEngine/Private/GeometryCollection/GeometryCollectionComponent.cpp
@@ -3086,11 +3086,17 @@ void UGeometryCollectionComponent::ProcessRepDataOnPT()
 		AActor* Owner = GetOwner();
 		const bool bReplicateMovement = Owner ? GetOwner()->IsReplicatingMovement() : true;
 
-		// Capture - VersionProcessed - LastHardsnapTimeInMs - OneOffActivatedProcessed
-		// which are only used in that function and in the reset function and both are executed on the Physics Thread
-		CurrSolver->EnqueueCommandImmediate([PhysicsProxy = PhysicsProxy, RepDataCopy = RepData, &VersionProcessed = VersionProcessed, &LastHardsnapTimeInMs = LastHardsnapTimeInMs, &OneOffActivatedProcessed = OneOffActivatedProcessed, bReplicateMovement](Chaos::FReal DeltaTime, Chaos::FReal SimTime)
+		// Capture a weak component pointer to make sure we can check for validity when executing the lambda 
+		// This is needed to access VersionProcessed, LastHardsnapTimeInMs, OneOffActivatedProcessed in Read/Write
+		// Note : those properties are only used in that function and in the reset function and both are executed on the Physics Thread
+
+		TWeakObjectPtr<UGeometryCollectionComponent> WeakThis(this);
+		CurrSolver->EnqueueCommandImmediate([WeakThis, PhysicsProxy = PhysicsProxy, RepDataCopy = RepData, bReplicateMovement](Chaos::FReal DeltaTime, Chaos::FReal SimTime)
 		{
-			::ProcessRepDataCommon(DeltaTime, SimTime, RepDataCopy, PhysicsProxy, VersionProcessed, OneOffActivatedProcessed, LastHardsnapTimeInMs, bReplicateMovement);
+			if (TStrongObjectPtr<UGeometryCollectionComponent> StrongThis = WeakThis.Pin())
+			{
+				::ProcessRepDataCommon(DeltaTime, SimTime, RepDataCopy, PhysicsProxy, StrongThis->VersionProcessed, StrongThis->OneOffActivatedProcessed, StrongThis->LastHardsnapTimeInMs, bReplicateMovement);
+			}
 		});
 	}
 }
@@ -3102,11 +3108,19 @@ void UGeometryCollectionComponent::ProcessRepStateDataOnPT()
 	if (Chaos::FPhysicsSolver* CurrSolver = GetSolver(*this))
 	{
 		// enqueue a callback to process the replicated state data on the physics thread 
+		// Capture a weak component pointer to make sure we can check for validity when executing the lambda 
+		// This is needed to access VersionProcessed in Read/Write
+		// Note : this property is only used in that function and in the reset function and both are executed on the Physics Thread
+
+		TWeakObjectPtr<UGeometryCollectionComponent> WeakThis(this);
 		CurrSolver->EnqueueCommandImmediate(
-			[PhysicsProxy = PhysicsProxy, RepStateDataCopy = RepStateData, &VersionProcessed = VersionProcessed]
+			[WeakThis, PhysicsProxy = PhysicsProxy, RepStateDataCopy = RepStateData]
 			(Chaos::FReal /*DeltaTime*/, Chaos::FReal /*SimTime*/)
 			{
-				::ProcessRepStateDataCommon(PhysicsProxy, RepStateDataCopy, VersionProcessed);
+				if (TStrongObjectPtr<UGeometryCollectionComponent> StrongThis = WeakThis.Pin())
+				{
+					::ProcessRepStateDataCommon(PhysicsProxy, RepStateDataCopy, StrongThis->VersionProcessed);
+				}
 			});
 	}
 }
@@ -3121,11 +3135,19 @@ void UGeometryCollectionComponent::ProcessRepDynamicDataOnPT()
 		const bool bReplicateMovement = Owner ? GetOwner()->IsReplicatingMovement() : true;
 
 		// enqueue a callback to process the replicated dynamic data on the physics thread 
+		// Capture a weak component pointer to make sure we can check for validity when executing the lambda 
+		// This is needed to access VersionProcessed, LastHardsnapTimeInMs, OneOffActivatedProcessed in Read/Write
+		// Note : those properties are only used in that function and in the reset function and both are executed on the Physics Thread
+
+		TWeakObjectPtr<UGeometryCollectionComponent> WeakThis(this);
 		CurrSolver->EnqueueCommandImmediate(
-			[PhysicsProxy = PhysicsProxy, RepDynamicDataCopy = RepDynamicData, &VersionProcessed = DynamicRepDataVersionProcessed, bReplicateMovement, &LastHardsnapTimeInMs = LastHardsnapTimeInMs]
+			[WeakThis, PhysicsProxy = PhysicsProxy, RepDynamicDataCopy = RepDynamicData, bReplicateMovement]
 			(Chaos::FReal DeltaTime, Chaos::FReal SimTime)
 			{
-				::ProcessRepDynamicDataCommon(PhysicsProxy, RepDynamicDataCopy, VersionProcessed, bReplicateMovement, SimTime, DeltaTime, LastHardsnapTimeInMs);
+				if (TStrongObjectPtr<UGeometryCollectionComponent> StrongThis = WeakThis.Pin())
+				{
+					::ProcessRepDynamicDataCommon(PhysicsProxy, RepDynamicDataCopy, StrongThis->DynamicRepDataVersionProcessed, bReplicateMovement, SimTime, DeltaTime, StrongThis->LastHardsnapTimeInMs);
+				}
 			});
 	}
 }
