add default values to dataflow ustructs (FDataflowVariable::Type, FDataflowVariable::ContainerType, and FAssetRandomSeedDataflowNode::RandomSeed) to fix initialization warning

#Synced-CL 48266510

diff --git a/Engine/Plugins/Experimental/GeometryCollectionPlugin/Source/GeometryCollectionNodes/Public/Dataflow/GeometryCollectionAssetNodes.h b/Engine/Plugins/Experimental/GeometryCollectionPlugin/Source/GeometryCollectionNodes/Public/Dataflow/GeometryCollectionAssetNodes.h
--- a/Engine/Plugins/Experimental/GeometryCollectionPlugin/Source/GeometryCollectionNodes/Public/Dataflow/GeometryCollectionAssetNodes.h
+++ b/Engine/Plugins/Experimental/GeometryCollectionPlugin/Source/GeometryCollectionNodes/Public/Dataflow/GeometryCollectionAssetNodes.h
@@ -461,7 +461,7 @@ private:
 
 	/** Random seed value */
 	UPROPERTY(meta = (DataflowOutput))
-	int32 RandomSeed;
+	int32 RandomSeed = 0;
 
 	virtual void Evaluate(UE::Dataflow::FContext& Context, const FDataflowOutput* Out) const override;
 };
diff --git a/Engine/Source/Runtime/Experimental/Dataflow/Engine/Public/Dataflow/DataflowBlueprintLibrary.h b/Engine/Source/Runtime/Experimental/Dataflow/Engine/Public/Dataflow/DataflowBlueprintLibrary.h
--- a/Engine/Source/Runtime/Experimental/Dataflow/Engine/Public/Dataflow/DataflowBlueprintLibrary.h
+++ b/Engine/Source/Runtime/Experimental/Dataflow/Engine/Public/Dataflow/DataflowBlueprintLibrary.h
@@ -20,10 +20,10 @@ public:
 	FName VariableName;
 
 	UPROPERTY(EditAnywhere, BlueprintReadWrite, Category = "Dataflow")
-	EPropertyBagPropertyType Type;
+	EPropertyBagPropertyType Type = EPropertyBagPropertyType::None;
 
 	UPROPERTY(EditAnywhere, BlueprintReadWrite, Category = "Dataflow")
-	EPropertyBagContainerType ContainerType;
+	EPropertyBagContainerType ContainerType = EPropertyBagContainerType::None;
 
 	UPROPERTY(EditAnywhere, BlueprintReadWrite, Category = "Dataflow")
 	FString Value;
