Dumping assets to json commands
#Synced-CL 48294064

diff --git a/Engine/Source/Runtime/Experimental/DumpPackageToJson/DumpPackageToJson.Build.cs b/Engine/Source/Runtime/Experimental/DumpPackageToJson/DumpPackageToJson.Build.cs
new file mode 100644
--- /dev/null
+++ b/Engine/Source/Runtime/Experimental/DumpPackageToJson/DumpPackageToJson.Build.cs
@@ -0,0 +1,25 @@
+// Copyright Epic Games, Inc. All Rights Reserved.
+
+using UnrealBuildTool;
+
+public class DumpPackageToJson : ModuleRules
+{
+	public DumpPackageToJson(ReadOnlyTargetRules Target) : base(Target)
+	{
+		PrivateIncludePathModuleNames.AddRange(new string[] {
+			"StorageServerClient",
+		});
+
+		PrivateDependencyModuleNames.AddRange(
+			new string[] {
+				"Core",
+				"CoreUObject",
+				"JsonObjectGraph",
+				"StorageServerClient",
+			}
+		);
+
+		PCHUsage = PCHUsageMode.NoPCHs;
+		bEnableExceptions = true;
+	}
+}
diff --git a/Engine/Source/Runtime/Experimental/DumpPackageToJson/Private/DetachedStorageServerIoDispatcherBackend.cpp b/Engine/Source/Runtime/Experimental/DumpPackageToJson/Private/DetachedStorageServerIoDispatcherBackend.cpp
new file mode 100644
--- /dev/null
+++ b/Engine/Source/Runtime/Experimental/DumpPackageToJson/Private/DetachedStorageServerIoDispatcherBackend.cpp
@@ -0,0 +1,143 @@
+﻿// Copyright Epic Games, Inc. All Rights Reserved.
+
+#include "DetachedStorageServerIoDispatcherBackend.h"
+
+#if !UE_BUILD_SHIPPING
+
+#include "DumpPackageToJsonModule.h"
+
+#include "Serialization/MemoryReader.h"
+#include "Serialization/MemoryWriter.h"
+#include "Serialization/ZenPackageHeader.h"
+
+namespace UE::DumpPackageToJson
+{
+
+void FDetachedStorageServerIoDispatcherBackend::SetInner(TSharedPtr<IIoDispatcherBackend> InInner)
+{
+	if (!InInner.IsValid() && Inner.IsValid() && SavedContext.IsValid())
+	{
+		Inner->Shutdown();
+	}
+	Inner = InInner;
+	if (Inner.IsValid() && SavedContext.IsValid())
+	{
+		Inner->Initialize(SavedContext.ToSharedRef());
+	}
+}
+
+void FDetachedStorageServerIoDispatcherBackend::Initialize(TSharedRef<const FIoDispatcherBackendContext> Context)
+{
+	SavedContext = Context;
+	Inner->Initialize(Context);
+}
+
+void FDetachedStorageServerIoDispatcherBackend::Shutdown()
+{
+	Inner->Shutdown();
+	SavedContext.Reset();
+}
+
+void FDetachedStorageServerIoDispatcherBackend::ResolveIoRequests(FIoRequestList Requests, FIoRequestList& OutUnresolved)
+{
+	for (FIoRequestImpl& Request : Requests)
+	{
+		Request.ChunkId = Remapping->Map(Request.ChunkId);
+	}
+
+	FIoRequestList OutUnresolvedMapped = {};
+	Inner->ResolveIoRequests(MoveTemp(Requests), OutUnresolvedMapped);
+
+	OutUnresolved.AddTail(MoveTemp(OutUnresolvedMapped));
+}
+
+bool FDetachedStorageServerIoDispatcherBackend::DoesChunkExist(const FIoChunkId& ChunkId) const
+{
+	return Inner->DoesChunkExist(Remapping->Map(ChunkId));
+}
+
+TIoStatusOr<uint64> FDetachedStorageServerIoDispatcherBackend::GetSizeForChunk(const FIoChunkId& ChunkId) const
+{
+	return Inner->GetSizeForChunk(Remapping->Map(ChunkId));
+}
+
+FIoRequestImpl* FDetachedStorageServerIoDispatcherBackend::GetCompletedIoRequests()
+{
+	FIoRequestImpl* CompletedMapped = Inner->GetCompletedIoRequests();
+
+	for (FIoRequestList::FIterator It(CompletedMapped); It; ++It)
+	{
+		FIoRequestImpl& Request = (*It);
+
+		const bool bIsExportBundleData = Request.ChunkId.GetChunkType() == EIoChunkType::ExportBundleData;
+		const bool bIsChunkIndex0 = NETWORK_ORDER16(*reinterpret_cast<const uint16*>(&Request.ChunkId.GetData()[8])) == 0;
+
+		Request.ChunkId = Remapping->Remap(Request.ChunkId);
+
+		if (bIsExportBundleData && bIsChunkIndex0 && Request.HasBuffer())
+		{
+			PatchPackageNameInZenPackageHeader(Request.GetBuffer());
+		}
+	}
+
+	return CompletedMapped;
+}
+
+void FDetachedStorageServerIoDispatcherBackend::PatchPackageNameInZenPackageHeader(FIoBuffer& Buffer)
+{
+	// There is a check in AsyncLoading2.cpp comparing requested package name with one in package header data: 
+	// check(Package->Desc.PackageIdToLoad == FPackageId::FromName(Package->HeaderData.PackageName));
+	//
+	// To avoid failing the check we either need to modify AsyncLoading2.cpp to support our use case,
+	// or patch package header in place. This code does in-place patching of zen package header name map.
+	// We can only patch with a new map of same size, hence requirement that temp patch name matches the length of original package name.
+
+	uint8* PackageHeaderDataPtr = Buffer.GetData();
+	const FZenPackageSummary* PackageSummary = reinterpret_cast<const FZenPackageSummary*>(PackageHeaderDataPtr);
+
+	// reading package header, see FZenPackageHeader
+	TArrayView PackageHeaderDataView(PackageHeaderDataPtr + sizeof(FZenPackageSummary), PackageSummary->HeaderSize - sizeof(FZenPackageSummary));
+	FMemoryReaderView PackageHeaderDataReader(PackageHeaderDataView);
+
+	TOptional<FZenPackageVersioningInfo> VersioningInfo;
+	if (PackageSummary->bHasVersioningInfo)
+	{
+		VersioningInfo.Emplace();
+		PackageHeaderDataReader << VersioningInfo.GetValue();
+	}
+
+	const FZenPackageVersioningInfo* LocalVersioningInfo = VersioningInfo.GetPtrOrNull();
+	if (LocalVersioningInfo == nullptr || LocalVersioningInfo->PackageVersion >= EUnrealEngineObjectUE5Version::VERSE_CELLS)
+	{
+		FZenPackageCellOffsets CellOffsets;
+		PackageHeaderDataReader << CellOffsets.CellImportMapOffset;
+		PackageHeaderDataReader << CellOffsets.CellExportMapOffset;
+	}
+
+	const int64 NameMapOffset = PackageHeaderDataReader.Tell();
+	TArray<FDisplayNameEntryId> NameMap = LoadNameBatch(PackageHeaderDataReader);
+	const int64 NameMapSize = PackageHeaderDataReader.Tell() - NameMapOffset;
+
+	const FString OriginalName = PackageSummary->Name.ResolveName(NameMap).ToString();
+	const FString TempName = Remapping->Remap(OriginalName);
+	const FName TempNameFName = *TempName;
+	NameMap[PackageSummary->Name.GetIndex()] = FDisplayNameEntryId(TempNameFName);
+
+	TArray<uint8> NewNameMapBytes;
+	FMemoryWriter NewNameMapWriter(NewNameMapBytes);
+
+	SaveNameBatch(NameMap, NewNameMapWriter);
+
+	if (NewNameMapBytes.Num() != NameMapSize)
+	{
+		UE_LOG(LogDumpPackageToJson, Error, TEXT("Unable to patch zen header in place, old name table size %lli new name table size %d"), NameMapSize, NewNameMapBytes.Num());
+		return;
+	}
+
+	// replace name map in-place
+	FMemory::Memcpy(PackageHeaderDataView.GetData() + NameMapOffset, NewNameMapBytes.GetData(), NewNameMapBytes.Num());
+}
+
+}
+
+#endif
\ No newline at end of file
diff --git a/Engine/Source/Runtime/Experimental/DumpPackageToJson/Private/DetachedStorageServerIoDispatcherBackend.h b/Engine/Source/Runtime/Experimental/DumpPackageToJson/Private/DetachedStorageServerIoDispatcherBackend.h
new file mode 100644
--- /dev/null
+++ b/Engine/Source/Runtime/Experimental/DumpPackageToJson/Private/DetachedStorageServerIoDispatcherBackend.h
@@ -0,0 +1,56 @@
+﻿// Copyright Epic Games, Inc. All Rights Reserved.
+
+#pragma once
+
+#include "Misc/Build.h"
+
+#if !UE_BUILD_SHIPPING
+
+#include "CoreTypes.h"
+#include "IO/IoDispatcherBackend.h"
+#include "PackageNameRemapping.h"
+
+namespace UE::DumpPackageToJson
+{
+
+// io dispatcher backend that remaps incoming chunk id's to inner chunk id's and back 
+class FDetachedStorageServerIoDispatcherBackend final : public IIoDispatcherBackend
+{
+public:
+	FDetachedStorageServerIoDispatcherBackend(TSharedRef<FPackageNameRemapping> InRemapping)
+		: Remapping(InRemapping)
+	{
+	}
+	~FDetachedStorageServerIoDispatcherBackend() = default;
+
+	void SetInner(TSharedPtr<IIoDispatcherBackend> InInner);
+
+	virtual void Initialize(TSharedRef<const FIoDispatcherBackendContext> Context) override;
+	virtual void Shutdown() override;
+	virtual void ResolveIoRequests(FIoRequestList Requests, FIoRequestList& OutUnresolved) override;
+	virtual void CancelIoRequest(FIoRequestImpl* Request) override {}
+	virtual void UpdatePriorityForIoRequest(FIoRequestImpl* Request) override {};
+	virtual bool DoesChunkExist(const FIoChunkId& ChunkId) const override;
+	virtual TIoStatusOr<uint64> GetSizeForChunk(const FIoChunkId& ChunkId) const override;
+	virtual FIoRequestImpl* GetCompletedIoRequests() override;
+
+	virtual TIoStatusOr<FIoMappedRegion> OpenMapped(const FIoChunkId& ChunkId, const FIoReadOptions& Options) override
+	{
+		return FIoStatus(EIoErrorCode::NotFound);
+	}
+	virtual const TCHAR* GetName() const override
+	{
+		return TEXT("DetachedStorageServerIoDispatcherBackend");
+	}
+
+private:
+	TSharedPtr<IIoDispatcherBackend> Inner;
+	TSharedRef<FPackageNameRemapping> Remapping;
+	TSharedPtr<const FIoDispatcherBackendContext> SavedContext;
+
+	void PatchPackageNameInZenPackageHeader(FIoBuffer& Buffer);
+};
+
+}
+
+#endif
\ No newline at end of file
diff --git a/Engine/Source/Runtime/Experimental/DumpPackageToJson/Private/DetachedStorageServerPackageStoreBackend.cpp b/Engine/Source/Runtime/Experimental/DumpPackageToJson/Private/DetachedStorageServerPackageStoreBackend.cpp
new file mode 100644
--- /dev/null
+++ b/Engine/Source/Runtime/Experimental/DumpPackageToJson/Private/DetachedStorageServerPackageStoreBackend.cpp
@@ -0,0 +1,28 @@
+﻿// Copyright Epic Games, Inc. All Rights Reserved.
+
+#include "DetachedStorageServerPackageStoreBackend.h"
+
+#if !UE_BUILD_SHIPPING
+
+namespace UE::DumpPackageToJson
+{
+
+void FDetachedStorageServerPackageStoreBackend::SetInner(TSharedPtr<IPackageStoreBackend> InInner)
+{
+	Inner = InInner;
+}
+
+EPackageStoreEntryStatus FDetachedStorageServerPackageStoreBackend::GetPackageStoreEntry(FPackageId PackageId, FName PackageName, FPackageStoreEntry& OutPackageStoreEntry)
+{
+	if (!Inner.IsValid())
+	{
+		return EPackageStoreEntryStatus::Missing;
+	}
+
+	// TODO we could also remap imported package id's
+	return Inner->GetPackageStoreEntry(Remapping->Map(PackageId), PackageName, OutPackageStoreEntry);
+}
+
+}
+
+#endif
\ No newline at end of file
diff --git a/Engine/Source/Runtime/Experimental/DumpPackageToJson/Private/DetachedStorageServerPackageStoreBackend.h b/Engine/Source/Runtime/Experimental/DumpPackageToJson/Private/DetachedStorageServerPackageStoreBackend.h
new file mode 100644
--- /dev/null
+++ b/Engine/Source/Runtime/Experimental/DumpPackageToJson/Private/DetachedStorageServerPackageStoreBackend.h
@@ -0,0 +1,45 @@
+﻿// Copyright Epic Games, Inc. All Rights Reserved.
+
+#pragma once
+
+#include "Misc/Build.h"
+
+#if !UE_BUILD_SHIPPING
+
+#include "CoreTypes.h"
+#include "PackageNameRemapping.h"
+#include "Serialization/PackageStore.h"
+
+namespace UE::DumpPackageToJson
+{
+
+// package store backend that remaps incoming chunk id's to inner chunk id's and back
+class FDetachedStorageServerPackageStoreBackend	: public IPackageStoreBackend
+{
+public:
+	FDetachedStorageServerPackageStoreBackend(TSharedRef<FPackageNameRemapping> InRemapping)
+		: Remapping(InRemapping)
+	{
+	}
+	virtual ~FDetachedStorageServerPackageStoreBackend() = default;
+
+	void SetInner(TSharedPtr<IPackageStoreBackend> InInner);
+
+	virtual void OnMounted(TSharedRef<const FPackageStoreBackendContext> Context) override {}
+	virtual void BeginRead() override {}
+	virtual void EndRead() override {}
+
+	virtual EPackageStoreEntryStatus GetPackageStoreEntry(FPackageId PackageId, FName PackageName, FPackageStoreEntry& OutPackageStoreEntry) override;
+	virtual bool GetPackageRedirectInfo(FPackageId PackageId, FName& OutSourcePackageName, FPackageId& OutRedirectedToPackageId) override
+	{
+		return false;
+	}
+
+private:
+	TSharedPtr<IPackageStoreBackend> Inner;
+	TSharedRef<FPackageNameRemapping> Remapping;
+};
+
+}
+
+#endif
\ No newline at end of file
diff --git a/Engine/Source/Runtime/Experimental/DumpPackageToJson/Private/DumpPackageToJsonModule.cpp b/Engine/Source/Runtime/Experimental/DumpPackageToJson/Private/DumpPackageToJsonModule.cpp
new file mode 100644
--- /dev/null
+++ b/Engine/Source/Runtime/Experimental/DumpPackageToJson/Private/DumpPackageToJsonModule.cpp
@@ -0,0 +1,249 @@
+﻿// Copyright Epic Games, Inc. All Rights Reserved.
+
+#include "DumpPackageToJsonModule.h"
+#include "Modules/ModuleManager.h"
+
+#if !UE_BUILD_SHIPPING
+
+#include "DetachedStorageServerIoDispatcherBackend.h"
+#include "DetachedStorageServerPackageStoreBackend.h"
+#include "PackageNameRemapping.h"
+#include "StorageServerClientModule.h"
+#include "JsonObjectGraph/Stringify.h"
+#include "Logging/StructuredLog.h"
+#include "Misc/FileHelper.h"
+#include "Misc/Paths.h"
+#include "Serialization/PackageStore.h"
+#include "UObject/LinkerInstancingContext.h"
+
+DEFINE_LOG_CATEGORY(LogDumpPackageToJson);
+
+namespace UE::DumpPackageToJson
+{
+static FAutoConsoleCommand DumpPackageToJsonConnectToZenCommand(
+	TEXT("DumpPackageToJson.ConnectToZen"),
+	TEXT("Connects to a zen server to be able to dump cooked packages, pass [zen_project] [zen_platform] ([ip]) ([port]).")
+	TEXT("This connection is only used for DumpPackage commands and doesn't influence rest of the engine."),
+	FConsoleCommandWithArgsDelegate::CreateLambda([](const TArray<FString>& Args)
+	{
+		if (Args.Num() < 2 || Args.Num() > 4)
+		{
+			UE_LOG(LogDumpPackageToJson, Warning, TEXT("Invalid amount of arguments passed"));
+			return;
+		}
+
+		auto& Module = FModuleManager::Get().GetModuleChecked<UE::DumpPackageToJson::FDumpPackageToJsonModule>("DumpPackageToJson");
+
+		const FString Project = Args[0];
+		const FString Platform = Args[1];
+		const FString Host = Args.Num() >= 3 ? Args[2] : TEXT("127.0.0.1");
+
+		uint16 Port = 8558;
+		if (Args.Num() >= 4 && !LexTryParseString(Port, *Args[3]))
+		{
+			UE_LOG(LogDumpPackageToJson, Warning, TEXT("Invalid format of [port]"));
+			return;
+		}
+
+		Module.Connect(Host, Port, Project, Platform);
+	}));
+
+static FAutoConsoleCommand DumpPackageToJsonDisconnectCommand(
+	TEXT("DumpPackageToJson.Disconnect"),
+	TEXT("Disconnects from a zen server"),
+	FConsoleCommandWithArgsDelegate::CreateLambda([](const TArray<FString>& Args)
+	{
+		auto& Module = FModuleManager::Get().GetModuleChecked<UE::DumpPackageToJson::FDumpPackageToJsonModule>("DumpPackageToJson");
+		Module.Disconnect();
+	}));
+
+
+static FAutoConsoleCommand DumpPackageToJsonCommand(
+	TEXT("DumpPackageToJson"),
+	TEXT("Dump json representation of a package to log, pass [package_name] ([file_path]).")
+	TEXT("Will dump remote package if DumpPackageToJson.ConnectToZen was executed before, otherwise will try to use package_name in current context.")
+	TEXT("Will write json to a file if [file_path] is provided."),
+	FConsoleCommandWithArgsDelegate::CreateLambda([](const TArray<FString>& Args)
+	{
+		if (Args.Num() < 1 || Args.Num() > 2)
+		{
+			UE_LOG(LogDumpPackageToJson, Warning, TEXT("Invalid amount of arguments passed"));
+			return;
+		}
+
+
+		auto& Module = FModuleManager::Get().GetModuleChecked<UE::DumpPackageToJson::FDumpPackageToJsonModule>("DumpPackageToJson");
+
+		const FString PackageName = Args[0];
+		const FString Json = Module.IsConnected() ? Module.RemotePackageToJson(PackageName) : Module.LocalPackageToJson(PackageName);
+
+		if (Args.Num() >= 2)
+		{
+			FString FilePath = Args[1];
+
+			FPaths::NormalizeFilename(FilePath);
+
+			if (FFileHelper::SaveStringToFile(Json, *FilePath))
+			{
+				UE_LOG(LogDumpPackageToJson, Display, TEXT("Dumped package to '%s'"), *FilePath);
+			}
+			else
+			{
+				UE_LOG(LogDumpPackageToJson, Error, TEXT("Failed to write to '%s'"), *FilePath);
+			}
+		}
+		else
+		{
+			UE_LOG(LogDumpPackageToJson, Display, TEXT("%s"), *Json);
+		}
+	}));
+
+bool FDumpPackageToJsonModule::Connect(FString Host, uint16 Port, FString Project, FString Platform)
+{
+	if (IsConnected())
+	{
+		Disconnect();
+	}
+
+	IStorageServerClientModule& StorageServerClientModule = FModuleManager::Get().LoadModuleChecked<IStorageServerClientModule>("StorageServerClient");
+
+	IStorageServerClientModule::FDetachedPlatformFileResult StorageServer = {};
+
+	if (!StorageServerClientModule.CreateDetachedPlatformFile(Host, Port, Project, Platform, StorageServer))
+	{
+		UE_LOG(LogDumpPackageToJson, Error, TEXT("Can't connect to storage server"));
+		return false;
+	}
+
+	StorageServerPlatformFile = MoveTemp(StorageServer.PlatformFile);
+	StorageServerPackageStoreBackend = MoveTemp(StorageServer.PackageStoreBackend);
+	StorageServerIoDispatcherBackend = MoveTemp(StorageServer.IoDispatcherBackend);
+
+	// TODO do we want remapping platform file as well?
+
+	if (!Remapper.IsValid())
+	{
+		Remapper = MakeShared<FPackageNameRemapping>();
+	}
+
+	if (RemappingIoDispatcherBackend.IsValid())
+	{
+		RemappingIoDispatcherBackend->SetInner(StorageServerIoDispatcherBackend);
+	}
+	else
+	{
+		RemappingIoDispatcherBackend = MakeShared<FDetachedStorageServerIoDispatcherBackend>(Remapper.ToSharedRef());
+		RemappingIoDispatcherBackend->SetInner(StorageServerIoDispatcherBackend);
+		FIoDispatcher::Get().Mount(RemappingIoDispatcherBackend.ToSharedRef());
+	}
+
+	if (RemappingPackageStoreBackend.IsValid())
+	{
+		RemappingPackageStoreBackend->SetInner(StorageServerPackageStoreBackend);
+	}
+	else
+	{
+		RemappingPackageStoreBackend = MakeShared<FDetachedStorageServerPackageStoreBackend>(Remapper.ToSharedRef());
+		RemappingPackageStoreBackend->SetInner(StorageServerPackageStoreBackend);
+		FPackageStore::Get().Mount(RemappingPackageStoreBackend.ToSharedRef());
+	}
+	return true;
+}
+
+void FDumpPackageToJsonModule::Disconnect()
+{
+	if (!IsConnected())
+	{
+		return;
+	}
+
+	RemappingIoDispatcherBackend->SetInner(TSharedPtr<IIoDispatcherBackend>());
+	RemappingPackageStoreBackend->SetInner(TSharedPtr<IPackageStoreBackend>());
+
+	StorageServerPackageStoreBackend.Reset();
+	StorageServerIoDispatcherBackend.Reset();
+	StorageServerPlatformFile.Reset();
+}
+
+FString FDumpPackageToJsonModule::RemotePackageToJson(FString OriginalPackageName)
+{
+	const FString TempPackageNameTemplate = FString::Printf(TEXT("/Temp/ToJson%u"), TempPackageIndex++);
+	if (OriginalPackageName.Len() < TempPackageNameTemplate.Len())
+	{
+		// see FDetachedStorageServerIoDispatcherBackend::PatchPackageNameInZenPackageHeader for more info
+		UE_LOG(LogDumpPackageToJson, Error, TEXT("Package name '%s' is too short to patch, shortest supported name is '%s'"), *OriginalPackageName, *TempPackageNameTemplate);
+		return FString("");
+	}
+
+	const FString TempPackageName = FString::Printf(TEXT("%s%0*d"), *TempPackageNameTemplate, OriginalPackageName.Len() - TempPackageNameTemplate.Len(), 0);
+	ensure(TempPackageName.Len() == OriginalPackageName.Len());
+
+	Remapper->AddRemap(TempPackageName, OriginalPackageName);
+
+	const FPackagePath TempPackagePath = FPackagePath::FromPackageNameChecked(TempPackageName);
+
+	const FPackageId TempPackageId = FPackageId::FromName(TempPackagePath.GetPackageFName());
+
+	FPackageStoreEntry PackageStoreEntry = {};
+	EPackageStoreEntryStatus PackageStoreEntryStatus;
+
+	{
+		FPackageStoreReadScope _(FPackageStore::Get());
+		PackageStoreEntryStatus = FPackageStore::Get().GetPackageStoreEntry(TempPackageId, TempPackagePath.GetPackageFName(), PackageStoreEntry);
+	}
+
+	if (PackageStoreEntryStatus != EPackageStoreEntryStatus::Ok)
+	{
+		UE_LOG(LogDumpPackageToJson, Error, TEXT("Failed to resolve temp package id %s for package name '%s'"), *LexToString(TempPackageId), *TempPackageName);
+		Remapper->ClearRemaps();
+		return {};
+	}
+
+	UE_LOG(LogDumpPackageToJson, Display, TEXT("Package %s imports %i packages"), *OriginalPackageName, PackageStoreEntry.ImportedPackageIds.Num());
+	for (FPackageId ImportedPackageId : PackageStoreEntry.ImportedPackageIds)
+	{
+		UE_LOG(LogDumpPackageToJson, Display, TEXT("Imports package id %s"), *LexToString(ImportedPackageId));
+	}
+
+	UPackage* Package = LoadPackage(nullptr, *TempPackageName, LOAD_NoVerify | LOAD_SkipLoadImportedPackages | LOAD_DisableEngineVersionChecks, nullptr, nullptr);
+	Remapper->ClearRemaps();
+
+	if (!Package)
+	{
+		UE_LOG(LogDumpPackageToJson, Error, TEXT("Failed to load temp package id %s for package name '%s'"), *LexToString(TempPackageId), *TempPackageName);
+		return {};
+	}
+
+	return UPackageToJson(Package);
+}
+
+FString FDumpPackageToJsonModule::LocalPackageToJson(FString PackageName)
+{
+	UPackage* Package = LoadPackage(nullptr, *PackageName, LOAD_NoVerify | LOAD_SkipLoadImportedPackages);
+
+	return Package ? UPackageToJson(Package) : FString();
+}
+
+FString FDumpPackageToJsonModule::UPackageToJson(UPackage* Package)
+{
+	FUtf8String Json = {};
+
+	try
+	{
+		Json = UE::JsonObjectGraph::Stringify(
+		 {Package}
+		);
+	}
+	catch (...)
+	{
+		UE_LOG(LogDumpPackageToJson, Error, TEXT("Failed to stingify package '%s' due to unknown error"), *Package->GetName());
+	}
+
+	return StringCast<TCHAR>(*Json).Get();
+}
+
+}
+
+#endif
+
+IMPLEMENT_MODULE(UE::DumpPackageToJson::FDumpPackageToJsonModule, DumpPackageToJson);
diff --git a/Engine/Source/Runtime/Experimental/DumpPackageToJson/Private/DumpPackageToJsonModule.h b/Engine/Source/Runtime/Experimental/DumpPackageToJson/Private/DumpPackageToJsonModule.h
new file mode 100644
--- /dev/null
+++ b/Engine/Source/Runtime/Experimental/DumpPackageToJson/Private/DumpPackageToJsonModule.h
@@ -0,0 +1,71 @@
+﻿// Copyright Epic Games, Inc. All Rights Reserved.
+
+#pragma once
+
+#include "IDumpPackageToJsonModule.h"
+
+#include "Templates/SharedPointer.h"
+#include "Templates/UniquePtr.h"
+#include "UObject/UObjectGlobals.h"
+#include "UObject/Package.h"
+
+struct IIoDispatcherBackend;
+class IPackageStoreBackend;
+class IStorageServerPlatformFile;
+
+#if !UE_BUILD_SHIPPING
+
+DECLARE_LOG_CATEGORY_EXTERN(LogDumpPackageToJson, Log, All);
+
+#endif
+
+namespace UE::DumpPackageToJson
+{
+
+#if !UE_BUILD_SHIPPING
+
+class FPackageNameRemapping;
+class FDetachedStorageServerIoDispatcherBackend;
+class FDetachedStorageServerPackageStoreBackend;
+
+class FDumpPackageToJsonModule : public IDumpPackageToJsonModule
+{
+public:
+	virtual bool Connect(FString Host, uint16 Port, FString Project, FString Platform) override;
+	virtual void Disconnect() override;
+	virtual bool IsConnected() override
+	{
+		return StorageServerPlatformFile.IsValid();
+	}
+	virtual FString RemotePackageToJson(FString PackageName) override;
+	virtual FString LocalPackageToJson(FString PackageName) override;
+
+private:
+	TUniquePtr<IStorageServerPlatformFile> StorageServerPlatformFile;
+	TSharedPtr<IPackageStoreBackend> StorageServerPackageStoreBackend;
+	TSharedPtr<IIoDispatcherBackend> StorageServerIoDispatcherBackend;
+
+	TSharedPtr<FPackageNameRemapping> Remapper;
+	TSharedPtr<FDetachedStorageServerIoDispatcherBackend> RemappingIoDispatcherBackend;
+	TSharedPtr<FDetachedStorageServerPackageStoreBackend> RemappingPackageStoreBackend;
+
+	uint32 TempPackageIndex = 0;
+
+	static FString UPackageToJson(UPackage* Package); 
+};
+
+#else
+
+class FDumpPackageToJsonModule : public IDumpPackageToJsonModule
+{
+public:
+	virtual bool Connect(FString /*Host*/, uint16 /*Port*/, FString /*Project*/, FString /*Platform*/) override {return false;}
+	virtual void Disconnect() override {}
+	virtual bool IsConnected() override {return false;}
+	virtual FString RemotePackageToJson(FString /*PackageName*/) override {return TEXT("");}
+	virtual FString LocalPackageToJson(FString /*PackageName*/) override {return TEXT("");}
+};
+
+#endif
+
+}
diff --git a/Engine/Source/Runtime/Experimental/DumpPackageToJson/Private/PackageNameRemapping.cpp b/Engine/Source/Runtime/Experimental/DumpPackageToJson/Private/PackageNameRemapping.cpp
new file mode 100644
--- /dev/null
+++ b/Engine/Source/Runtime/Experimental/DumpPackageToJson/Private/PackageNameRemapping.cpp
@@ -0,0 +1,99 @@
+﻿// Copyright Epic Games, Inc. All Rights Reserved.
+
+#include "PackageNameRemapping.h"
+
+#if !UE_BUILD_SHIPPING
+
+namespace UE::DumpPackageToJson
+{
+
+void FPackageNameRemapping::AddRemap(FString TempName, FString OriginalName)
+{
+	const FPackageId TempPackageId = PackageIdFromString(TempName);
+	const FPackageId OriginalPackageId = PackageIdFromString(OriginalName);
+
+	TempNameToOriginal.Add(TempName, OriginalName);
+	OriginalNameToTemp.Add(OriginalName, TempName);
+
+	TempPackageIdToOriginal.Add(TempPackageId, OriginalPackageId);
+	OriginalPackageIdToName.Add(OriginalPackageId, TempPackageId);
+}
+
+void FPackageNameRemapping::ClearRemaps()
+{
+	TempNameToOriginal.Empty();
+	OriginalNameToTemp.Empty();
+	TempPackageIdToOriginal.Empty();
+	OriginalPackageIdToName.Empty();
+}
+
+FString FPackageNameRemapping::Map(FString TempName) const
+{
+	const FString OriginalName = TempNameToOriginal.FindRef(TempName, TEXT("%InvalidPackage%"));
+	return OriginalName;
+}
+
+FPackageId FPackageNameRemapping::Map(FPackageId TempPackageId) const
+{
+	const FPackageId OriginalPackageId = TempPackageIdToOriginal.FindRef(TempPackageId, FPackageId::FromValue(0));
+	return OriginalPackageId;
+}
+
+FIoChunkId FPackageNameRemapping::Map(FIoChunkId TempChunkId) const
+{
+	const FPackageId TempPackageId = PackageIdFromChunkId(TempChunkId);
+	const FPackageId OriginalPackageId = Map(TempPackageId);
+	const FIoChunkId OriginalChunkId = ReplaceChunkIdsPackageId(TempChunkId, OriginalPackageId);
+
+	// UE_LOG(LogDumpPackageToJson, Display, TEXT("Mapping %s -> %s"), *LexToString(TempChunkId), *LexToString(OriginalChunkId));
+	return OriginalChunkId;
+}
+
+FString FPackageNameRemapping::Remap(FString OriginalName) const
+{
+	const FString TempName = OriginalNameToTemp.FindRef(OriginalName, TEXT("%InvalidPackage%"));
+	return TempName;
+}
+
+FPackageId FPackageNameRemapping::Remap(FPackageId OriginalPackageId) const
+{
+	const FPackageId TempPackageId = OriginalPackageIdToName.FindRef(OriginalPackageId, FPackageId::FromValue(0));
+	return TempPackageId;
+}
+
+FIoChunkId FPackageNameRemapping::Remap(FIoChunkId OriginalChunkId) const
+{
+	const FPackageId OriginalPackageId = PackageIdFromChunkId(OriginalChunkId);
+	const FPackageId TempPackageId = Remap(OriginalPackageId);
+	const FIoChunkId TempChunkId = ReplaceChunkIdsPackageId(OriginalChunkId, TempPackageId);
+
+	//UE_LOG(LogDumpPackageToJson, Display, TEXT("Remapping %s -> %s"), *LexToString(OriginalChunkId), *LexToString(TempChunkId));
+	return TempChunkId;
+}
+
+FPackageId FPackageNameRemapping::PackageIdFromString(FString String)
+{
+	// TODO add FPackageId::FromString?
+	String.ToLowerInline();
+	const uint64 Value = CityHash64(reinterpret_cast<const char*>(String.GetCharArray().GetData()), String.NumBytesWithoutNull());
+	return FPackageId::FromValue(Value);
+}
+
+FPackageId FPackageNameRemapping::PackageIdFromChunkId(FIoChunkId ChunkId)
+{
+	uint64 Result = 0;
+	FMemory::Memcpy(&Result, ChunkId.GetData(), sizeof(uint64));
+	return FPackageId::FromValue(Result);
+}
+
+FIoChunkId FPackageNameRemapping::ReplaceChunkIdsPackageId(FIoChunkId ChunkId, FPackageId NewPackageId)
+{
+	const uint64 Value = NewPackageId.Value();
+	// TODO add better way
+	FMemory::Memcpy(const_cast<uint8*>(ChunkId.GetData()), &Value, sizeof(uint64));
+	return ChunkId;
+}
+
+}
+
+#endif
\ No newline at end of file
diff --git a/Engine/Source/Runtime/Experimental/DumpPackageToJson/Private/PackageNameRemapping.h b/Engine/Source/Runtime/Experimental/DumpPackageToJson/Private/PackageNameRemapping.h
new file mode 100644
--- /dev/null
+++ b/Engine/Source/Runtime/Experimental/DumpPackageToJson/Private/PackageNameRemapping.h
@@ -0,0 +1,46 @@
+﻿// Copyright Epic Games, Inc. All Rights Reserved.
+
+#pragma once
+
+#include "Misc/Build.h"
+
+#if !UE_BUILD_SHIPPING
+
+#include "Containers/UnrealString.h"
+#include "IO/IoChunkId.h"
+#include "IO/PackageId.h"
+
+namespace UE::DumpPackageToJson
+{
+
+// Helper class to map between temporary package name/id <-> original package name/id
+class FPackageNameRemapping
+{
+public:
+	void AddRemap(FString TempName, FString OriginalName);
+	void ClearRemaps();
+
+	// Temp -> Original
+	FString Map(FString TempName) const;
+	FPackageId Map(FPackageId TempPackageId) const;
+	FIoChunkId Map(FIoChunkId TempChunkId) const;
+
+	// Original -> Temp
+	FString Remap(FString OriginalName) const;
+	FPackageId Remap(FPackageId OriginalPackageId) const;
+	FIoChunkId Remap(FIoChunkId OriginalChunkId) const;
+
+private:
+	TMap<FString, FString> TempNameToOriginal;
+	TMap<FString, FString> OriginalNameToTemp;
+	TMap<FPackageId, FPackageId> TempPackageIdToOriginal;
+	TMap<FPackageId, FPackageId> OriginalPackageIdToName;
+
+	static FPackageId PackageIdFromString(FString String);
+	static FPackageId PackageIdFromChunkId(FIoChunkId ChunkId);
+	static FIoChunkId ReplaceChunkIdsPackageId(FIoChunkId ChunkId, FPackageId NewPackageId);
+};
+
+}
+
+#endif
\ No newline at end of file
diff --git a/Engine/Source/Runtime/Experimental/DumpPackageToJson/Public/IDumpPackageToJsonModule.h b/Engine/Source/Runtime/Experimental/DumpPackageToJson/Public/IDumpPackageToJsonModule.h
new file mode 100644
--- /dev/null
+++ b/Engine/Source/Runtime/Experimental/DumpPackageToJson/Public/IDumpPackageToJsonModule.h
@@ -0,0 +1,21 @@
+﻿// Copyright Epic Games, Inc. All Rights Reserved.
+
+#pragma once
+
+#include "Modules/ModuleInterface.h"
+#include "Containers/UnrealString.h"
+
+namespace UE::DumpPackageToJson
+{
+
+class IDumpPackageToJsonModule : public IModuleInterface
+{
+public:
+	virtual bool Connect(FString Host, uint16 Port, FString Project, FString Platform) = 0;
+	virtual void Disconnect() = 0;
+	virtual bool IsConnected() = 0;
+	virtual FString RemotePackageToJson(FString PackageName) = 0;
+	virtual FString LocalPackageToJson(FString PackageName) = 0;
+};
+
+}
diff --git a/Engine/Source/Runtime/Launch/Launch.Build.cs b/Engine/Source/Runtime/Launch/Launch.Build.cs
--- a/Engine/Source/Runtime/Launch/Launch.Build.cs
+++ b/Engine/Source/Runtime/Launch/Launch.Build.cs
@@ -271,5 +271,15 @@ public class Launch : ModuleRules
 				PublicDependencyModuleNames.Add("CookOnTheFly");
 			}
 		}
+
+		if (Target.Configuration != UnrealTargetConfiguration.Shipping && Target.IsInPlatformGroup(UnrealPlatformGroup.Windows))
+		{
+			PrivateDependencyModuleNames.Add("DumpPackageToJson");
+			PrivateDefinitions.Add("UE_WITH_DUMP_PACKAGE_TO_JSON=1");
+		}
+		else
+		{
+			PrivateDefinitions.Add("UE_WITH_DUMP_PACKAGE_TO_JSON=0");
+		}
 	}
 }
diff --git a/Engine/Source/Runtime/Launch/Private/LaunchEngineLoop.cpp b/Engine/Source/Runtime/Launch/Private/LaunchEngineLoop.cpp
--- a/Engine/Source/Runtime/Launch/Private/LaunchEngineLoop.cpp
+++ b/Engine/Source/Runtime/Launch/Private/LaunchEngineLoop.cpp
@@ -331,6 +331,10 @@ static const FColor SystemColor = FColor::Orange;
 static const FColor CoreDelegatesColor = FColor::Cyan;
 #endif // ENABLE_NAMED_EVENTS
 
+#ifndef UE_WITH_DUMP_PACKAGE_TO_JSON
+#define UE_WITH_DUMP_PACKAGE_TO_JSON 0
+#endif
+
 bool GIsConsoleExecutable = false;
 
 static TAutoConsoleVariable<int32> CVarDoAsyncEndOfFrameTasksRandomize(
@@ -4626,6 +4630,13 @@ bool FEngineLoop::LoadStartupCoreModules()
 
 	FModuleManager::Get().LoadModule(TEXT("MassEntity"));
 
+#if UE_WITH_DUMP_PACKAGE_TO_JSON
+	if (FParse::Param(FCommandLine::Get(), TEXT("DumpPackageToJson")))
+	{
+		FModuleManager::Get().LoadModule("DumpPackageToJson");
+	}
+#endif
+
 	return bSuccess;
 }
 
diff --git a/Engine/Source/Runtime/StorageServerClient/Internal/StorageServerClientModule.h b/Engine/Source/Runtime/StorageServerClient/Internal/StorageServerClientModule.h
--- a/Engine/Source/Runtime/StorageServerClient/Internal/StorageServerClientModule.h
+++ b/Engine/Source/Runtime/StorageServerClient/Internal/StorageServerClientModule.h
@@ -6,8 +6,10 @@
 #include "HAL/IPlatformFileModule.h"
 #include "HAL/PlatformFileManager.h"
 #include "IStorageServerPlatformFile.h"
+#include "IO/IoDispatcherBackend.h"
 #include "Modules/ModuleInterface.h"
 #include "Modules/ModuleManager.h"
+#include "Serialization/PackageStore.h"
 
 #if !UE_BUILD_SHIPPING
 
@@ -24,6 +26,16 @@ public:
 	}
 
 	virtual IStorageServerPlatformFile* TryCreateCustomPlatformFile(FStringView StoreDirectory, IPlatformFile* Inner) = 0;
+
+	// Creates a platform file, package store and io dispatcher backends,
+	// without mounting them in FPackageStore and FIoDispatcher.
+	struct FDetachedPlatformFileResult
+	{
+		TUniquePtr<IStorageServerPlatformFile> PlatformFile;
+		TSharedPtr<IPackageStoreBackend> PackageStoreBackend;
+		TSharedPtr<IIoDispatcherBackend> IoDispatcherBackend;
+	};
+	virtual bool CreateDetachedPlatformFile(FString Host, uint16 Port, FString Project, FString Platform, FDetachedPlatformFileResult& OutResult) = 0;
 };
 
 #endif // !UE_BUILD_SHIPPING
diff --git a/Engine/Source/Runtime/StorageServerClient/Private/StorageServerClientModule.cpp b/Engine/Source/Runtime/StorageServerClient/Private/StorageServerClientModule.cpp
--- a/Engine/Source/Runtime/StorageServerClient/Private/StorageServerClientModule.cpp
+++ b/Engine/Source/Runtime/StorageServerClient/Private/StorageServerClientModule.cpp
@@ -25,6 +25,7 @@ public:
 
 	// ~Begin IStorageServerClientModule
 	virtual IStorageServerPlatformFile* TryCreateCustomPlatformFile(FStringView StoreDirectory, IPlatformFile* Inner) override;
+	virtual bool CreateDetachedPlatformFile(FString Host, uint16 Port, FString Project, FString Platform, FDetachedPlatformFileResult& OutResult) override;
 	// ~End IStorageServerClientModule
 
 	IPlatformFile* CachedDefaultPlatformFileInstance = nullptr;
@@ -66,6 +67,39 @@ IStorageServerPlatformFile* FStorageServerClientModule::TryCreateCustomPlatformF
 	return nullptr;
 }
 
+bool FStorageServerClientModule::CreateDetachedPlatformFile(FString Host, uint16 Port, FString Project, FString Platform, FDetachedPlatformFileResult& OutResult)
+{
+	OutResult = {};
+
+	TUniquePtr<FStorageServerPlatformFile> StorageServerPlatformFile = MakeUnique<FStorageServerPlatformFile>();
+	StorageServerPlatformFile->SetAllowPackageIo(false);
+	StorageServerPlatformFile->SetAbortOnConnectionFailure(false);
+
+	TStringBuilder<1024> CmdLine;
+	CmdLine.Appendf(TEXT("-ZenStoreProject=%s "), *Project);
+	CmdLine.Appendf(TEXT("-ZenStorePlatform=%s "), *Platform);
+	CmdLine.Appendf(TEXT("-ZenStoreHost=%s "), *Host);
+	CmdLine.Appendf(TEXT("-ZenStorePort=%u "), Port);
+	CmdLine.Appendf(TEXT("-ZenDisableCache"));
+
+	const TCHAR* CmdLineStr = CmdLine.ToString();
+	
+	if (StorageServerPlatformFile->ShouldBeUsed(nullptr, CmdLineStr) && StorageServerPlatformFile->Initialize(nullptr, CmdLineStr))
+	{
+		StorageServerPlatformFile->InitializeAfterProjectFilePath();
+
+		if (StorageServerPlatformFile->IsConnectionValid())
+		{
+			OutResult.PackageStoreBackend = StorageServerPlatformFile->CreatePackageStoreBackend();
+			OutResult.IoDispatcherBackend = StorageServerPlatformFile->CreateIoDispatcherBackend();
+			OutResult.PlatformFile = MoveTemp(StorageServerPlatformFile);
+			return true;
+		}
+	}
+
+	return false;
+}
+
 IMPLEMENT_MODULE(FStorageServerClientModule, StorageServerClient)
 
 #endif // !UE_BUILD_SHIPPING
diff --git a/Engine/Source/Runtime/StorageServerClient/Private/StorageServerPlatformFile.cpp b/Engine/Source/Runtime/StorageServerClient/Private/StorageServerPlatformFile.cpp
--- a/Engine/Source/Runtime/StorageServerClient/Private/StorageServerPlatformFile.cpp
+++ b/Engine/Source/Runtime/StorageServerClient/Private/StorageServerPlatformFile.cpp
@@ -1210,7 +1210,7 @@ bool FStorageServerPlatformFile::ShouldBeUsed(IPlatformFile* Inner, const TCHAR*
 	}
 
 	FString Host;
-	if (FParse::Value(FCommandLine::Get(), TEXT("-ZenStoreHost="), Host))
+	if (FParse::Value(CmdLine, TEXT("-ZenStoreHost="), Host))
 	{
 		UE_LOG(LogStorageServerPlatformFile, Display, TEXT("Adding connection settings from command line: -ZenStoreHost='%s'"), *Host);
 		if (!Host.ParseIntoArray(HostAddrs, TEXT("+"), true))
@@ -1452,6 +1452,25 @@ void FStorageServerPlatformFile::UpdateFileList()
 	}
 }
 
+TSharedPtr<IPackageStoreBackend> FStorageServerPlatformFile::CreatePackageStoreBackend() const
+{
+#if WITH_COTF
+	if (CookOnTheFlyServerConnection)
+	{
+		return MakeShared<FCookOnTheFlyPackageStoreBackend>(*CookOnTheFlyServerConnection.Get());
+	}
+	else
+#endif
+	{
+		return MakeShared<FStorageServerPackageStoreBackend>(*Connection.Get());
+	}
+}
+
+TSharedPtr<IIoDispatcherBackend> FStorageServerPlatformFile::CreateIoDispatcherBackend() const
+{
+	return MakeShared<FStorageServerIoDispatcherBackend>(*Connection.Get());
+}
+
 bool FStorageServerPlatformFile::FileExists(const TCHAR* Filename)
 {
 	TStringBuilder<1024> StorageServerFilename;
diff --git a/Engine/Source/Runtime/StorageServerClient/Private/StorageServerPlatformFile.h b/Engine/Source/Runtime/StorageServerClient/Private/StorageServerPlatformFile.h
--- a/Engine/Source/Runtime/StorageServerClient/Private/StorageServerPlatformFile.h
+++ b/Engine/Source/Runtime/StorageServerClient/Private/StorageServerPlatformFile.h
@@ -8,6 +8,9 @@
 #include "Containers/DirectoryTree.h"
 #include "IO/IoDispatcher.h"
 #include "IStorageServerPlatformFile.h"
+#include "StorageServerConnection.h"
+#include "StorageServerIoDispatcherBackend.h"
+#include "StorageServerPackageStore.h"
 
 #if !UE_BUILD_SHIPPING
 
@@ -126,7 +129,17 @@ public:
 	{
 		CustomProjectStorePath = InProjectStorePath;
 	}
+	
 	void UpdateFileList() override;
+
+	bool IsConnectionValid() const
+	{
+		return Connection.IsValid();
+	}
+
+	TSharedPtr<IPackageStoreBackend> CreatePackageStoreBackend() const;
+	TSharedPtr<IIoDispatcherBackend> CreateIoDispatcherBackend() const;
+
 private:
 	friend class FStorageServerFileHandle;
 	template<typename ParentVisitorClass, typename DataType> friend class FUniqueDirectoryStatVisitor;
