Global on-demand TOC file.
- as a first step to get on-demand TOC files more aligned with regular container TOCs this change
enables generating a global.uondemandtoc file next to the global.utoc file in the content PAKs folder with the paramater -GlobalContainerPath=<Path>
- extended IasTool Upload with a parameter to specify a path where the global TOC should be, typically {Project}/Content/Paks/global.uondemandtoc
- this path can be derived from the input parameters but it serves as a way to preserve the old way of storing the global TOC
- extended config parsing to look for the global TOC file at startup
- added command line argument -Iax.Zen when using Zen as a CDN

#rnx

#Synced-CL 48026973

diff --git a/Engine/Source/Runtime/Experimental/IoStore/OnDemand/Private/OnDemandConfig.cpp b/Engine/Source/Runtime/Experimental/IoStore/OnDemand/Private/OnDemandConfig.cpp
--- a/Engine/Source/Runtime/Experimental/IoStore/OnDemand/Private/OnDemandConfig.cpp
+++ b/Engine/Source/Runtime/Experimental/IoStore/OnDemand/Private/OnDemandConfig.cpp
@@ -2,6 +2,7 @@
 
 #include "OnDemandConfig.h"
 
+#include "HAL/FileManager.h"
 #include "HAL/PlatformProcess.h"
 #include "IO/IoStatus.h"
 #include "IO/IoStoreOnDemand.h"
@@ -277,36 +278,50 @@ TIoStatusOr<FOnDemandIoStoreConfig> TryParseConfig(const TCHAR* CommandLine)
 	}
 #endif
 
-	const FString ConfigFileName = "IoStoreOnDemand.ini";
-	const FString ConfigPath = FPaths::Combine(TEXT("Cloud"), ConfigFileName);
-
-	if (FPlatformMisc::FileExistsInPlatformPackage(ConfigPath) == false)
+	// Check for global TOC
+	bool bGlobalToc = false;
 	{
-		FIoStatus Status = FIoStatusBuilder(EIoErrorCode::NotFound)
-			<< TEXT("Failed to find config file '")
-			<< ConfigPath
-			<< TEXT("'");
-		return Status;
+		const FString GlobalTocPath = FString::Printf(TEXT("%sPaks/global.uondemandtoc"), *FPaths::ProjectContentDir());
+		if (bGlobalToc = IFileManager::Get().FileExists(*GlobalTocPath); bGlobalToc)
+		{
+			OutConfig.TocFilePath = GlobalTocPath;
+		}
 	}
 
-	const FString ConfigContent = FPlatformMisc::LoadTextFileFromPlatformPackage(ConfigPath);
-	if (TryParseConfigContent(ConfigContent, ConfigFileName, OutConfig) == false || OutConfig.IsValid() == false)
+	// Look for the global TOC in the cloud directory (legacy path)
+	if (bGlobalToc == false)
 	{
-		FIoStatus Status = FIoStatusBuilder(EIoErrorCode::ReadError)
-			<< TEXT("Failed to read config file '")
-			<< ConfigPath
-			<< TEXT("'");
-		return Status;
-	}
+		const FString ConfigFileName = "IoStoreOnDemand.ini";
+		const FString ConfigPath = FPaths::Combine(TEXT("Cloud"), ConfigFileName);
 
-	{
-		TStringBuilder<256> TocFilePath;
-		FPathViews::Append(TocFilePath, TEXT("Cloud"), FPaths::GetBaseFilename(OutConfig.TocPath));
-		TocFilePath.Append(TEXT(".iochunktoc"));
+		if (FPlatformMisc::FileExistsInPlatformPackage(ConfigPath) == false)
+		{
+			FIoStatus Status = FIoStatusBuilder(EIoErrorCode::NotFound)
+				<< TEXT("Failed to find config file '")
+				<< ConfigPath
+				<< TEXT("'");
+			return Status;
+		}
 
-		if (FPlatformMisc::FileExistsInPlatformPackage(*TocFilePath))
+		const FString ConfigContent = FPlatformMisc::LoadTextFileFromPlatformPackage(ConfigPath);
+		if (TryParseConfigContent(ConfigContent, ConfigFileName, OutConfig) == false || OutConfig.IsValid() == false)
 		{
-			OutConfig.TocFilePath = TocFilePath;
+			FIoStatus Status = FIoStatusBuilder(EIoErrorCode::ReadError)
+				<< TEXT("Failed to read config file '")
+				<< ConfigPath
+				<< TEXT("'");
+			return Status;
+		}
+
+		{
+			TStringBuilder<256> TocFilePath;
+			FPathViews::Append(TocFilePath, TEXT("Cloud"), FPaths::GetBaseFilename(OutConfig.TocPath));
+			TocFilePath.Append(TEXT(".iochunktoc"));
+
+			if (FPlatformMisc::FileExistsInPlatformPackage(*TocFilePath))
+			{
+				OutConfig.TocFilePath = TocFilePath;
+			}
 		}
 	}
 
@@ -321,6 +336,10 @@ TIoStatusOr<FOnDemandIoStoreConfig> TryParseConfig(const TCHAR* CommandLine)
 			Host.TrimStartAndEndInline();
 		}
 	}
+	else if (FParse::Param(CommandLine, TEXT("Iax.Zen")))
+	{
+		OutConfig.ServiceUrls.Add(TEXT("http://localhost:8558"));
+	}
 #endif
 
 	return OutConfig;
diff --git a/Engine/Source/Runtime/Experimental/IoStore/OnDemand/Private/Tool/Upload.cpp b/Engine/Source/Runtime/Experimental/IoStore/OnDemand/Private/Tool/Upload.cpp
--- a/Engine/Source/Runtime/Experimental/IoStore/OnDemand/Private/Tool/Upload.cpp
+++ b/Engine/Source/Runtime/Experimental/IoStore/OnDemand/Private/Tool/Upload.cpp
@@ -46,6 +46,7 @@ struct FUploadParams
 	FString TargetPlatform;
 	FString EncryptionKeyName;
 	FString HostGroupName;
+	FString GlobalContainerPath;
 	int32 MaxConcurrentUploads = 16;
 	int32 MaxTocListCount = 10000;
 	int32 MaxTocDownloadCount = 100;
@@ -67,9 +68,12 @@ struct FUploadParams
 
 FIoStatus FUploadParams::Validate() const
 {
-	if (bWriteTocToDisk && TocOutputDir.IsEmpty())
+	if (GlobalContainerPath.IsEmpty())
 	{
-		return FIoStatus(EIoErrorCode::InvalidParameter, TEXT("Cmdline param 'WriteToDisk' requires a valid 'ConfigFilePath' param as well"));
+		if (bWriteTocToDisk && TocOutputDir.IsEmpty())
+		{
+			return FIoStatus(EIoErrorCode::InvalidParameter, TEXT("Cmdline param 'WriteToDisk' requires a valid 'ConfigFilePath' param as well"));
+		}
 	}
 
 	if (!AccessKey.IsEmpty() && SecretKey.IsEmpty())
@@ -128,6 +132,7 @@ static FUploadParams BuildUploadParams(const FContext& Context)
 	Ret.BuildVersion			= Context.Get<FStringView>(TEXT("-BuildVersion"),			Ret.BuildVersion);
 	Ret.TargetPlatform			= Context.Get<FStringView>(TEXT("-TargetPlatform"),			Ret.TargetPlatform);
 	Ret.HostGroupName			= Context.Get<FStringView>(TEXT("-HostGroupName"),			Ret.HostGroupName);
+	Ret.GlobalContainerPath		= Context.Get<FStringView>(TEXT("-GlobalContainerPath"),	Ret.GlobalContainerPath);
 	Ret.bWriteTocToDisk			= Context.Get<bool>(TEXT("-WriteTocToDisk"),				Ret.bWriteTocToDisk);
 	Ret.bPerContainerTocs		= Context.Get<bool>(TEXT("-PerContainerTocs"),				Ret.bPerContainerTocs);
 	Ret.bIgnoreContainerHeader	= Context.Get<bool>(TEXT("-IgnoreContainerHeader"),			Ret.bIgnoreContainerHeader);
@@ -175,6 +180,12 @@ static FUploadParams BuildUploadParams(const FContext& Context)
 	Ret.MaxTocListCount		= FMath::Max(FMath::Min(Ret.MaxTocListCount, MaxTocListCount), MinTocListCount);
 	Ret.MaxTocDownloadCount = FMath::Max(FMath::Min(Ret.MaxTocDownloadCount, MaxTocDownloadCount), 1);
 
+	// Make sure this is a directory path
+	if (FPaths::GetExtension(Ret.GlobalContainerPath).IsEmpty() == false)
+	{
+		Ret.GlobalContainerPath	= FPaths::GetPath(Ret.GlobalContainerPath);
+	}
+
 	return Ret;
 }
 
@@ -804,16 +815,30 @@ static TIoStatusOr<FUploadResult> UploadContainerFiles(
 			return FIoStatus(EIoErrorCode::WriteError, TEXT("Failed to upload TOC"));
 		}
 
-		if (UploadParams.bWriteTocToDisk)
+		// Write a single global.uondemandtoc next to global.utoc or in the Cloud directory (legacy)
 		{
-			// Write a single .iochunktoc containing all on demand data for the current build
 			TStringBuilder<512> OnDemandTocFilePath;
-			FPathViews::Append(OnDemandTocFilePath, UploadParams.TocOutputDir, UploadResult.TocHash);
-			OnDemandTocFilePath << TEXT(".iochunktoc");
+			if (UploadParams.GlobalContainerPath.IsEmpty() == false)
+			{
+				FPathViews::Append(OnDemandTocFilePath, UploadParams.GlobalContainerPath, TEXT("global.uondemandtoc"));
+			}
+			else if (UploadParams.bWriteTocToDisk)
+			{
+				FPathViews::Append(OnDemandTocFilePath, UploadParams.TocOutputDir, UploadResult.TocHash);
+				OnDemandTocFilePath << TEXT(".iochunktoc");
+			}
 
-			if (SaveArrayToFile(TArrayView64<uint8>(Ar.GetData(), Ar.TotalSize()), OnDemandTocFilePath.ToString()) == false)
+			if (OnDemandTocFilePath.Len() > 0)
 			{
-				UE_LOG(LogIoStore, Error, TEXT("Failed to save on demand toc file '%s'"), OnDemandTocFilePath.ToString());
+				FPathViews::NormalizeFilename(OnDemandTocFilePath);
+				if (SaveArrayToFile(TArrayView64<uint8>(Ar.GetData(), Ar.TotalSize()), OnDemandTocFilePath.ToString()))
+				{
+					UE_LOG(LogIas, Display, TEXT("Saved on-demand TOC file '%s'"), OnDemandTocFilePath.ToString());
+				}
+				else
+				{
+					UE_LOG(LogIas, Error, TEXT("Failed to save on-demand TOC file '%s'"), OnDemandTocFilePath.ToString());
+				}
 			}
 		}
 	}
@@ -994,26 +1019,27 @@ static FCommand UploadCommand(
 	TEXT("Upload"),
 	TEXT("Used to upload IoStore containers to the cloud and convert to on-demand"),
 	{
-		TArgument<FStringView>(TEXT("ContainerGlob"),	TEXT("Path globbed to discover input containers")),
-		TArgument<FStringView>(TEXT("-CryptoKeys"),		TEXT("JSON-format keyring for input containers")),
-		TArgument<FStringView>(TEXT("-BuildVersion"),	TEXT("Optional build version to embed it TOC")),
-		TArgument<FStringView>(TEXT("-TargetPlatform"),	TEXT("If given, embedded in the output TOC")),
-		TArgument<FStringView>(TEXT("-ConfigFilePath"),	TEXT("Path to the config file to write runtime parameters to")),
-		TArgument<FStringView>(TEXT("-DistributionUrl"),TEXT("URL for IAS to use at runtime. Output to config file")),
-		TArgument<FStringView>(TEXT("-FallbackUrl"),	TEXT("Alternative fallack for DistributionUrl")),
-		TArgument<FStringView>(TEXT("-BucketPrefix"),	TEXT("Path to prefix to bucket objects")),
-		TArgument<FStringView>(TEXT("-HostGroupName"),	TEXT("Host group name or URL")),
-		TArgument<bool>(TEXT("-KeepContainerFiles"),	TEXT("Do not delete container files after upload")),
-		TArgument<bool>(TEXT("-KeepPakFiles"),			TEXT("Do not delete the springboard pak files")),
-		TArgument<bool>(TEXT("-WriteTocToDisk"),		TEXT("Output the TOC to disk as well as uploading")),
-		TArgument<bool>(TEXT("-PerContainerTocs"),		TEXT("Whether to generate TOC's for each container file(s)")),
-		TArgument<bool>(TEXT("-IgnoreContainerFlags"),	TEXT("Whether to ignore the OnDemand container flag")),
-		TArgument<bool>(TEXT("-StreamOnDemand"),		TEXT("Set the content to be streamed on-demand")),
-		TArgument<bool>(TEXT("-InstallOnDemand"),		TEXT("Set the content to be installed on-demand")),
-		TArgument<bool>(TEXT("-IgnoreContainerHeader"),	TEXT("Do not serialize the container header as part of the TOC file.")),
-		TArgument<int32>(TEXT("-MaxConcurrentUploads"),	TEXT("Number of simultaneous uploads")),
-		TArgument<int32>(TEXT("-MaxTocListCount"),		TEXT("Maximum number of TOC file(s) to list from the bucket")),
-		TArgument<int32>(TEXT("-MaxTocDownloadCount"),	TEXT("Maximum number of TOC file(s) to download")),
+		TArgument<FStringView>(TEXT("ContainerGlob"),			TEXT("Path globbed to discover input containers")),
+		TArgument<FStringView>(TEXT("-CryptoKeys"),				TEXT("JSON-format keyring for input containers")),
+		TArgument<FStringView>(TEXT("-BuildVersion"),			TEXT("Optional build version to embed it TOC")),
+		TArgument<FStringView>(TEXT("-TargetPlatform"),			TEXT("If given, embedded in the output TOC")),
+		TArgument<FStringView>(TEXT("-ConfigFilePath"),			TEXT("Path to the config file to write runtime parameters to")),
+		TArgument<FStringView>(TEXT("-DistributionUrl"),		TEXT("URL for IAS to use at runtime. Output to config file")),
+		TArgument<FStringView>(TEXT("-FallbackUrl"),			TEXT("Alternative fallack for DistributionUrl")),
+		TArgument<FStringView>(TEXT("-BucketPrefix"),			TEXT("Path to prefix to bucket objects")),
+		TArgument<FStringView>(TEXT("-HostGroupName"),			TEXT("Host group name or URL")),
+		TArgument<FStringView>(TEXT("-GlobalContainerPath"),	TEXT("Directory path to the content PAK folder")),
+		TArgument<bool>(TEXT("-KeepContainerFiles"),			TEXT("Do not delete container files after upload")),
+		TArgument<bool>(TEXT("-KeepPakFiles"),					TEXT("Do not delete the springboard pak files")),
+		TArgument<bool>(TEXT("-WriteTocToDisk"),				TEXT("Output the TOC to disk as well as uploading")),
+		TArgument<bool>(TEXT("-PerContainerTocs"),				TEXT("Whether to generate TOC's for each container file(s)")),
+		TArgument<bool>(TEXT("-IgnoreContainerFlags"),			TEXT("Whether to ignore the OnDemand container flag")),
+		TArgument<bool>(TEXT("-StreamOnDemand"),				TEXT("Set the content to be streamed on-demand")),
+		TArgument<bool>(TEXT("-InstallOnDemand"),				TEXT("Set the content to be installed on-demand")),
+		TArgument<bool>(TEXT("-IgnoreContainerHeader"),			TEXT("Do not serialize the container header as part of the TOC file.")),
+		TArgument<int32>(TEXT("-MaxConcurrentUploads"),			TEXT("Number of simultaneous uploads")),
+		TArgument<int32>(TEXT("-MaxTocListCount"),				TEXT("Maximum number of TOC file(s) to list from the bucket")),
+		TArgument<int32>(TEXT("-MaxTocDownloadCount"),			TEXT("Maximum number of TOC file(s) to download")),
 		S3Arguments,
 	}
 );
