Minor improvements to the IAX on screen debug display for the IAS cache.

#rnx

- Remove 'stats' from the start of the line as it provided no useful info.
- Show the percentage of the cache being used after the number of bytes being used.
- Replace  use of FString::Printf with the string builder version to skip out a memory allocation.

#Synced-CL 47893047

diff --git a/Engine/Source/Runtime/Experimental/IoStore/OnDemand/Private/Statistics.cpp b/Engine/Source/Runtime/Experimental/IoStore/OnDemand/Private/Statistics.cpp
--- a/Engine/Source/Runtime/Experimental/IoStore/OnDemand/Private/Statistics.cpp
+++ b/Engine/Source/Runtime/Experimental/IoStore/OnDemand/Private/Statistics.cpp
@@ -919,18 +919,22 @@ void FOnDemandIoBackendStats::PrintOnScreenStatistics(FCoreDelegates::FSeverityM
 	// Print IAS Cache
 	{
 		TStringBuilder<256> CachingText;
-		CachingText << TEXT("IAS Cache Stats: ");
+		CachingText << TEXT("IAS Cache: ");
 		if (BackendStatus.IsCacheEnabled())
 		{
+			const float CacheUsagePercent = GCacheMaxBytes > 0 ? 100.f * (float(GCacheCachedBytes.Get()) / float(GCacheMaxBytes)) : 0.f;
+
 			CachingText << TEXT("Cached: ") << FText::AsMemory(GCacheCachedBytes.Get()).ToString();
+			CachingText.Appendf(TEXT(" (%.1f%%)"), CacheUsagePercent);
 			CachingText << TEXT(" | Rejected: ") << FText::AsMemory(GCacheRejectBytes.Get()).ToString();
 			CachingText << TEXT(" | Read: ") << FText::AsMemory(GCacheReadBytes.Get()).ToString();
 			CachingText << TEXT(" (") << GCacheGetCount.Get() << TEXT(")");
 			CachingText << TEXT(" (Boot: ") << GCacheBootMs << TEXT("ms)");
+
 			const double SuspendedSeconds = GCacheSuspendedSeconds.Get();
 			if (SuspendedSeconds > 0.0)
 			{
-				CachingText << FString::Printf(TEXT(" (Suspended: %.2lfs)"), SuspendedSeconds);
+				CachingText.Appendf(TEXT(" (Suspended: %.2lfs)"), SuspendedSeconds);
 			}
 		}
 		else
