Collision Filter Refactor Part 5.2: Extend collision filters to 64 bit internally.
- Remove old SimWord3 and BodyIndex
- Extend the masks to 64 bit
- Extend collision channel index to 6 bits.
- Updated a lot of tests
Note: The 64 bit extension does not work through the engine as data is lost when converting to the legacy format which is what's stored currently. A separate pass will be done to update that data when serialization is also deal with.
#Synced-CL 48295927

diff --git a/Engine/Source/Runtime/Experimental/Chaos/Private/Chaos/CollisionFilterData.cpp b/Engine/Source/Runtime/Experimental/Chaos/Private/Chaos/CollisionFilterData.cpp
--- a/Engine/Source/Runtime/Experimental/Chaos/Private/Chaos/CollisionFilterData.cpp
+++ b/Engine/Source/Runtime/Experimental/Chaos/Private/Chaos/CollisionFilterData.cpp
@@ -6,8 +6,8 @@
 
 namespace Chaos::Filter
 {
-	constexpr uint8 NumFlagsBits = 21;
-	constexpr uint8 NumCollisionChannelBits = 5;
+	constexpr uint8 NumFlagsBits = 20;
+	constexpr uint8 NumCollisionChannelBits = 6;
 	constexpr uint8 NumMaskFilterBits = 6;
 	static_assert(NumFlagsBits + NumMaskFilterBits + NumCollisionChannelBits == 32);
 
@@ -18,12 +18,17 @@ namespace Chaos::Filter
 	constexpr uint32 CollisionChannelMask = ((1 << NumCollisionChannelBits) - 1) << CollisionChannelBitsOffset;
 	constexpr uint32 MaskFilterMask = ((1 << NumMaskFilterBits) - 1) << MaskFilterBitsOffset;
 
-	// Note: The full mask is currently limited to 32 bits during the transition period
-	constexpr uint64 FullChannelMask = 0xFFFFFFFF;
+	constexpr uint64 FullChannelMask = std::numeric_limits<uint64>::max();
 
 	constexpr uint32 LegacyQueryFlagsMask = (uint32)(EFilterFlags::SimpleCollision | EFilterFlags::ComplexCollision);
 	constexpr uint32 LegacySimFlagsMask = (uint32)(EFilterFlags::All) & ~LegacyQueryFlagsMask;
 
+	constexpr uint8 NumLegacyFlagsBits = 21;
+	constexpr uint8 NumLegacyCollisionChannelBits = 5;
+	static_assert(NumLegacyFlagsBits + NumMaskFilterBits + NumLegacyCollisionChannelBits == 32);
+	constexpr uint32 LegacyCollisionChannelBitsOffset = NumLegacyFlagsBits;
+	constexpr uint32 LegacyCollisionChannelMask = ((1 << NumLegacyCollisionChannelBits) - 1) << LegacyCollisionChannelBitsOffset;
+
 	uint32 GetFlags(uint32 Word3)
 	{
 		const uint32 FilterFlags = (Word3 & FilterFlagsMask);
@@ -90,6 +95,30 @@ namespace Chaos::Filter
 		Word3 |= MaskFilter << MaskFilterBitsOffset;
 	}
 
+	uint32 LegacyToNewWord3(const uint32 LegacyWord3)
+	{
+		// Split the word into the collision channel and everything else.
+		const uint32 LegacyChannels = LegacyWord3 & LegacyCollisionChannelMask;
+		const uint32 LegacyWord3Remainder = LegacyWord3 & ~LegacyCollisionChannelMask;
+		// To build the new word 3, we have to change from 5 bits to 6 bits for the channel index (left aligned).
+		// Simply right shift and steal the old upper flags bit (unused).
+		const uint32 NewChannels = (LegacyChannels >> 1);
+		const uint32 NewWord3 = LegacyWord3Remainder | NewChannels;
+		return NewWord3;
+	}
+
+	uint32 NewToLegacyWord3(const uint32 NewWord3)
+	{
+		// Split the word into the collision channel and everything else.
+		const uint32 NewChannels = NewWord3 & CollisionChannelMask;
+		const uint32 NewWord3Remainder = NewWord3 & ~CollisionChannelMask;
+		// To build the legacy word3, we have to change from 6 bits to 5 bits for the channel index (left aligned). 
+		// To do this, shift 1 to the left and mask the word (chopping off the top bit so it doesn't interfere with the mask filter).
+		const uint32 LegacyChannels = (NewChannels << 1) & LegacyCollisionChannelMask;
+		const uint32 LegacyWord3 = NewWord3Remainder | LegacyChannels;
+		return LegacyWord3;
+	}
+
 	FInstanceData::FInstanceData(const uint32 OwnerId, const uint32 ComponentId)
 		: OwnerId(OwnerId)
 		, ComponentId(ComponentId)
@@ -98,7 +127,7 @@ namespace Chaos::Filter
 
 	bool FInstanceData::IsValid() const
 	{
-		return OwnerId != 0 || ComponentId != 0 || BodyIndex != 0;
+		return OwnerId != 0 || ComponentId != 0;
 	}
 
 	uint32 FInstanceData::GetActorId() const
@@ -138,12 +167,12 @@ namespace Chaos::Filter
 
 	bool FShapeFilterData::IsValid() const
 	{
-		return QueryBlockChannels != 0 || QueryOverlapChannels != 0 || SimBlockChannels != 0 || Word3 != 0 || SimWord3 != 0;
+		return QueryBlockChannels != 0 || QueryOverlapChannels != 0 || SimBlockChannels != 0 || Word3 != 0;
 	}
 
 	bool FShapeFilterData::IsSimValid() const
 	{
-		return SimBlockChannels != 0 || SimWord3 != 0;
+		return SimBlockChannels != 0 || Word3 != 0;
 	}
 
 	EFilterFlags FShapeFilterData::GetFlags() const
@@ -154,7 +183,6 @@ namespace Chaos::Filter
 	void FShapeFilterData::SetFlags(EFilterFlags InFlags)
 	{
 		Chaos::Filter::SetFlags(Word3, (uint32)InFlags);
-		Chaos::Filter::SetFlags(SimWord3, (uint32)InFlags);
 	}
 
 	bool FShapeFilterData::HasFlag(EFilterFlags InFlag) const
@@ -170,7 +198,6 @@ namespace Chaos::Filter
 	void FShapeFilterData::SetMaskFilter(uint8 MaskFilter)
 	{
 		Chaos::Filter::SetMaskFilter(Word3, MaskFilter);
-		Chaos::Filter::SetMaskFilter(SimWord3, MaskFilter);
 	}
 
 	uint8 FShapeFilterData::GetCollisionChannelIndex() const
@@ -226,21 +253,22 @@ namespace Chaos::Filter
 
 	uint32 FShapeFilterData::GetLegacyWord3() const
 	{
-		return Word3;
+		return NewToLegacyWord3(Word3);
 	}
 
 	void FShapeUnionFilterData::Combine(const FShapeFilterData& ShapeFilter, bool bQueryEnabled, bool bSimEnabled)
 	{
 		if (bQueryEnabled)
 		{
+			// Note: Query Channels are unused in broad filtering. 
+			// Including them would increase the size of this filter which would increase the acceleration handle size.
 			QueryBlockChannels |= ShapeFilter.QueryBlockChannels;
 			QueryOverlapChannels |= ShapeFilter.QueryOverlapChannels;
-			QueryChannelMask |= (uint32)ShapeFilter.GetCollisionChannelMask();
 		}
 		if (bSimEnabled)
 		{
 			SimBlockChannels |= ShapeFilter.SimBlockChannels;
-			SimChannelMask |= (uint32)ShapeFilter.GetCollisionChannelMask();
+			SimChannelMask |= ShapeFilter.GetCollisionChannelMask();
 		}
 	}
 
@@ -376,11 +404,6 @@ namespace Chaos::Filter
 		}
 	}
 
-	void FQueryFilterData::SetMaskChannelAndFlags(const uint8 IgnoreMask, const uint8 ChannelIndex, const EFilterFlags FilterFlags)
-	{
-		Word3 = ((uint32)IgnoreMask << MaskFilterBitsOffset) | ((uint32)ChannelIndex << CollisionChannelBitsOffset) | (uint32)FilterFlags;
-	}
-
 	ENarrowFilterResult FQueryFilterData::ChannelTypeNarrowFilter(const FShapeFilterData& ShapeFilter, bool bSimBlock) const
 	{
 		const uint64 QueryChannelMask = GetCollisionChannelMask();
@@ -475,41 +498,38 @@ namespace Chaos::Filter
 	FShapeFilterBuilder& FShapeFilterBuilder::SetCollisionChannelIndex(const uint8 ChannelIndex)
 	{
 		SetCollisionChannel(ShapeFilterData.Word3, (uint32)ChannelIndex);
-		ShapeFilterData.SimWord3 = ShapeFilterData.Word3;
 		return *this;
 	}
 
 	FShapeFilterBuilder& FShapeFilterBuilder::SetBlockChannelMask(const uint64 BlockMask)
 	{
-		ShapeFilterData.SimBlockChannels = (uint32)BlockMask;
-		ShapeFilterData.QueryBlockChannels = (uint32)BlockMask;
+		ShapeFilterData.SimBlockChannels = BlockMask;
+		ShapeFilterData.QueryBlockChannels = BlockMask;
 		return *this;
 	}
 
 	FShapeFilterBuilder& FShapeFilterBuilder::SetBlockChannelMasks(const uint64 SimBlockMask, const uint64 QueryBlockMask)
 	{
-		ShapeFilterData.SimBlockChannels = (uint32)SimBlockMask;
-		ShapeFilterData.QueryBlockChannels = (uint32)QueryBlockMask;
+		ShapeFilterData.SimBlockChannels = SimBlockMask;
+		ShapeFilterData.QueryBlockChannels = QueryBlockMask;
 		return *this;
 	}
 
 	FShapeFilterBuilder& FShapeFilterBuilder::SetOverlapChannelMask(const uint64 OverlapMask)
 	{
-		ShapeFilterData.QueryOverlapChannels = (uint32)OverlapMask;
+		ShapeFilterData.QueryOverlapChannels = OverlapMask;
 		return *this;
 	}
 
 	FShapeFilterBuilder& FShapeFilterBuilder::SetMaskFilter(const uint8 InMaskFilter = 0)
 	{
 		ShapeFilterData.SetMaskFilter(InMaskFilter);
-		ShapeFilterData.SimWord3 = ShapeFilterData.Word3;
 		return *this;
 	}
 
 	FShapeFilterBuilder& FShapeFilterBuilder::SetFilterFlags(const EFilterFlags InFilterFlags)
 	{
 		ShapeFilterData.SetFlags(InFilterFlags);
-		ShapeFilterData.SimWord3 = ShapeFilterData.Word3;
 		return *this;
 	}
 
@@ -541,22 +561,11 @@ namespace Chaos::Filter
 		return Builder.Build();
 	}
 
-	FShapeFilterData FShapeFilterBuilder::BuildFromLegacyQueryFilter(const FCollisionFilterData& QueryFilter)
-	{
-		FShapeFilterData Result;
-		Result.QueryBlockChannels = QueryFilter.Word1;
-		Result.QueryOverlapChannels = QueryFilter.Word2;
-		Result.Word3 = QueryFilter.Word3;
-		Result.SimWord3 = QueryFilter.Word3;
-		return Result;
-	}
-
 	FShapeFilterData FShapeFilterBuilder::BuildFromLegacySimFilter(const FCollisionFilterData& SimFilter)
 	{
 		FShapeFilterData Result;
 		Result.SimBlockChannels = SimFilter.Word1;
-		Result.Word3 = SimFilter.Word3;
-		Result.SimWord3 = SimFilter.Word3;
+		Result.Word3 = LegacyToNewWord3(SimFilter.Word3);
 		return Result;
 	}
 
@@ -569,9 +578,9 @@ namespace Chaos::Filter
 		if (QueryFilterData.Word3 != 0 && SimFilterData.Word3 != 0)
 		{
 			const uint32 SimFilterMask = SimFilterData.Word3 & MaskFilterMask;
-			const uint32 SimCollisionChannel = SimFilterData.Word3 & CollisionChannelMask;
+			const uint32 SimCollisionChannel = SimFilterData.Word3 & LegacyCollisionChannelMask;
 			const uint32 QueryFilterMask = QueryFilterData.Word3 & MaskFilterMask;
-			const uint32 QueryCollisionChannel = QueryFilterData.Word3 & CollisionChannelMask;
+			const uint32 QueryCollisionChannel = QueryFilterData.Word3 & LegacyCollisionChannelMask;
 			ensure(SimFilterMask == QueryFilterMask);
 			ensure(SimCollisionChannel == QueryCollisionChannel);
 		}
@@ -581,13 +590,11 @@ namespace Chaos::Filter
 		FCombinedShapeFilterData CombinedShapeFilterData;
 		CombinedShapeFilterData.InstanceData.OwnerId = QueryFilterData.Word0;
 		CombinedShapeFilterData.InstanceData.ComponentId = SimFilterData.Word2;
-		CombinedShapeFilterData.InstanceData.BodyIndex = SimFilterData.Word0;
 
 		CombinedShapeFilterData.ShapeFilterData.QueryBlockChannels = QueryFilterData.Word1;
 		CombinedShapeFilterData.ShapeFilterData.QueryOverlapChannels = QueryFilterData.Word2;
-		CombinedShapeFilterData.ShapeFilterData.Word3 = Word3;
+		CombinedShapeFilterData.ShapeFilterData.Word3 = LegacyToNewWord3(Word3);
 		CombinedShapeFilterData.ShapeFilterData.SimBlockChannels = SimFilterData.Word1;
-		CombinedShapeFilterData.ShapeFilterData.SimWord3 = Word3;
 
 		return CombinedShapeFilterData;
 	}
@@ -608,20 +615,16 @@ namespace Chaos::Filter
 
 	void FShapeFilterBuilder::GetLegacyShapeFilter(const FCombinedShapeFilterData& CombinedShapeFilterData, FCollisionFilterData& OutQueryFilterData, FCollisionFilterData& OutSimFilterData)
 	{
-		// Grab whichever word3 is valid (temporary safeguard for legacy).
 		const FShapeFilterData& ShapeFilterData = CombinedShapeFilterData.ShapeFilterData;
-		const uint32 Word3 = ShapeFilterData.Word3 != 0 ? ShapeFilterData.Word3 : ShapeFilterData.SimWord3;
-		// We expect the word3's to always be equal in the transition period before removal of SimWord3.
-		ensure(ShapeFilterData.Word3 == ShapeFilterData.SimWord3);
+		const uint32 Word3 = NewToLegacyWord3(ShapeFilterData.Word3);
 
 		OutQueryFilterData.Word0 = CombinedShapeFilterData.InstanceData.OwnerId;
 		OutSimFilterData.Word2 = CombinedShapeFilterData.InstanceData.ComponentId;
-		OutSimFilterData.Word0 = CombinedShapeFilterData.InstanceData.BodyIndex;
 
-		OutQueryFilterData.Word1 = CombinedShapeFilterData.ShapeFilterData.QueryBlockChannels;
-		OutQueryFilterData.Word2 = CombinedShapeFilterData.ShapeFilterData.QueryOverlapChannels;
+		OutQueryFilterData.Word1 = (uint32)CombinedShapeFilterData.ShapeFilterData.QueryBlockChannels;
+		OutQueryFilterData.Word2 = (uint32)CombinedShapeFilterData.ShapeFilterData.QueryOverlapChannels;
 		OutQueryFilterData.Word3 = Word3;
-		OutSimFilterData.Word1 = CombinedShapeFilterData.ShapeFilterData.SimBlockChannels;
+		OutSimFilterData.Word1 = (uint32)CombinedShapeFilterData.ShapeFilterData.SimBlockChannels;
 		OutSimFilterData.Word3 = Word3;
 	}
 
@@ -646,7 +649,7 @@ namespace Chaos::Filter
 		if (SimWord3 != 0)
 		{
 			const uint32 MaskFilter = SimWord3 & MaskFilterMask;
-			const uint32 ChannelIndex = SimWord3 & CollisionChannelMask;
+			const uint32 ChannelIndex = SimWord3 & LegacyCollisionChannelMask;
 			const uint32 Flags = (SimWord3 & LegacySimFlagsMask) | (QueryWord3 & LegacyQueryFlagsMask);
 			return MaskFilter | ChannelIndex | Flags;
 		}
@@ -658,7 +661,7 @@ namespace Chaos::Filter
 		if (QueryWord3 != 0)
 		{
 			const uint32 MaskFilter = QueryWord3 & MaskFilterMask;
-			const uint32 ChannelIndex = QueryWord3 & CollisionChannelMask;
+			const uint32 ChannelIndex = QueryWord3 & LegacyCollisionChannelMask;
 			const uint32 Flags = (SimWord3 & LegacySimFlagsMask) | (QueryWord3 & LegacyQueryFlagsMask);
 			return MaskFilter | ChannelIndex | Flags;
 		}
@@ -669,13 +672,11 @@ namespace Chaos::Filter
 	{
 		OutCombinedShapeFilterData.InstanceData.OwnerId = QueryFilterData.Word0;
 		OutCombinedShapeFilterData.InstanceData.ComponentId = SimFilterData.Word2;
-		OutCombinedShapeFilterData.InstanceData.BodyIndex = SimFilterData.Word0;
 
 		OutCombinedShapeFilterData.ShapeFilterData.QueryBlockChannels = QueryFilterData.Word1;
 		OutCombinedShapeFilterData.ShapeFilterData.QueryOverlapChannels = QueryFilterData.Word2;
-		OutCombinedShapeFilterData.ShapeFilterData.Word3 = Word3;
+		OutCombinedShapeFilterData.ShapeFilterData.Word3 = LegacyToNewWord3(Word3);
 		OutCombinedShapeFilterData.ShapeFilterData.SimBlockChannels = SimFilterData.Word1;
-		OutCombinedShapeFilterData.ShapeFilterData.SimWord3 = Word3;
 	}
 
 	FQueryFilterData FQueryFilterBuilder::BuildFromLegacyQueryFilter(const FCollisionFilterData& QueryFilterData)
@@ -684,7 +685,7 @@ namespace Chaos::Filter
 		Result.Word0 = QueryFilterData.Word0;
 		Result.Word1 = QueryFilterData.Word1;
 		Result.Word2 = QueryFilterData.Word2;
-		Result.Word3 = QueryFilterData.Word3;
+		Result.Word3 = LegacyToNewWord3(QueryFilterData.Word3);
 		return Result;
 	}
 
@@ -692,9 +693,9 @@ namespace Chaos::Filter
 	{
 		FCollisionFilterData Result;
 		Result.Word0 = QueryFilterData.Word0;
-		Result.Word1 = QueryFilterData.Word1;
-		Result.Word2 = QueryFilterData.Word2;
-		Result.Word3 = QueryFilterData.Word3;
+		Result.Word1 = (uint32)QueryFilterData.Word1;
+		Result.Word2 = (uint32)QueryFilterData.Word2;
+		Result.Word3 = NewToLegacyWord3(QueryFilterData.Word3);
 		return Result;
 	}
 
@@ -711,7 +712,7 @@ namespace Chaos::Filter
 
 	FQueryObjectFilterBuilder& FQueryObjectFilterBuilder::SetObjectTypes(uint64 ObjectTypes)
 	{
-		QueryFilterData.Word1 = (uint32)ObjectTypes;
+		QueryFilterData.Word1 = ObjectTypes;
 		return *this;
 	}
 
@@ -772,13 +773,13 @@ namespace Chaos::Filter
 
 	FQueryTraceFilterBuilder& FQueryTraceFilterBuilder::SetBlockChannelMask(const uint64 BlockChannelMask)
 	{
-		QueryFilterData.Word1 = (uint32)BlockChannelMask;
+		QueryFilterData.Word1 = BlockChannelMask;
 		return *this;
 	}
 
 	FQueryTraceFilterBuilder& FQueryTraceFilterBuilder::SetOverlapChannelMask(const uint64 OverlapChannelMask)
 	{
-		QueryFilterData.Word2 = (uint32)OverlapChannelMask;
+		QueryFilterData.Word2 = OverlapChannelMask;
 		return *this;
 	}
 
diff --git a/Engine/Source/Runtime/Experimental/Chaos/Public/Chaos/CollisionFilterData.h b/Engine/Source/Runtime/Experimental/Chaos/Public/Chaos/CollisionFilterData.h
--- a/Engine/Source/Runtime/Experimental/Chaos/Public/Chaos/CollisionFilterData.h
+++ b/Engine/Source/Runtime/Experimental/Chaos/Public/Chaos/CollisionFilterData.h
@@ -117,8 +117,6 @@ namespace Chaos::Filter
 
 		uint32 OwnerId = 0;
 		uint32 ComponentId = 0;
-		// TODO @ JoshD: Remove if possible
-		uint32 BodyIndex = 0;
 	};
 
 	struct FShapeFilterData
@@ -159,12 +157,10 @@ namespace Chaos::Filter
 		friend FCombinedShapeFilterData;
 		friend struct FShapeUnionFilterData;
 
-		uint32 QueryBlockChannels = 0;
-		uint32 QueryOverlapChannels = 0;
-		uint32 SimBlockChannels = 0;
+		uint64 QueryBlockChannels = 0;
+		uint64 QueryOverlapChannels = 0;
+		uint64 SimBlockChannels = 0;
 		uint32 Word3 = 0;
-		// TODO @ JoshD: Remove this eventually. This is kept here for legacy conversions until usage is cleaned up.
-		uint32 SimWord3 = 0;
 	};
 
 	struct FShapeUnionFilterData
@@ -182,11 +178,10 @@ namespace Chaos::Filter
 
 	private:
 		friend class FShapeFilterBuilder;
-		uint32 QueryBlockChannels = 0;
-		uint32 QueryOverlapChannels = 0;
-		uint32 QueryChannelMask = 0;
-		uint32 SimBlockChannels = 0;
-		uint32 SimChannelMask = 0;
+		uint64 QueryBlockChannels = 0;
+		uint64 QueryOverlapChannels = 0;
+		uint64 SimBlockChannels = 0;
+		uint64 SimChannelMask = 0;
 	};
 
 	struct FQueryFilterData
@@ -234,14 +229,14 @@ namespace Chaos::Filter
 		friend class FQueryObjectFilterBuilder;
 		friend class FQueryTraceFilterBuilder;
 
-		void SetMaskChannelAndFlags(const uint8 IgnoreMask, const uint8 ChannelIndex, const EFilterFlags FilterFlags);
 		ENarrowFilterResult ChannelTypeNarrowFilter(const FShapeFilterData& ShapeFilter, bool bSimBlock) const;
 		ENarrowFilterResult ObjectTypeNarrowFilter(const FShapeFilterData& ShapeFilter) const;
 
 		uint32 Word0 = 0;
-		uint32 Word1 = 0;
-		uint32 Word2 = 0;
 		uint32 Word3 = 0;
+		// Note: These are the old word names, but they're arranged for memory packing
+		uint64 Word1 = 0;
+		uint64 Word2 = 0;
 	};
 
 	struct FCombinedShapeFilterData
@@ -286,7 +281,6 @@ namespace Chaos::Filter
 		CHAOS_API static FShapeFilterData BuildBlockAll(const EFilterFlags InFilterFlags = EFilterFlags::None);
 		CHAOS_API static FShapeFilterData BuildSimBlockAll(const EFilterFlags InFilterFlags = EFilterFlags::None);
 
-		UE_INTERNAL CHAOS_API static FShapeFilterData BuildFromLegacyQueryFilter(const FCollisionFilterData& QueryFilter);
 		UE_INTERNAL CHAOS_API static FShapeFilterData BuildFromLegacySimFilter(const FCollisionFilterData& SimFilter);
 		UE_INTERNAL CHAOS_API static FCombinedShapeFilterData BuildFromLegacyShapeFilter(const FCollisionFilterData& QueryFilterData, const FCollisionFilterData& SimFilterData);
 
