iax: Journal cache could get stuck if the cursor was precisely at the end of the journal on load.

#rnx

#Synced-CL Unknown

diff --git a/Engine/Source/Runtime/Experimental/IoStore/OnDemand/Private/JournaledCache.cpp b/Engine/Source/Runtime/Experimental/IoStore/OnDemand/Private/JournaledCache.cpp
--- a/Engine/Source/Runtime/Experimental/IoStore/OnDemand/Private/JournaledCache.cpp
+++ b/Engine/Source/Runtime/Experimental/IoStore/OnDemand/Private/JournaledCache.cpp
@@ -537,6 +537,13 @@ void FDiskJournal::GetPath(TStringBuilder<64>& Out)
 ////////////////////////////////////////////////////////////////////////////////
 FDiskPhrase FDiskJournal::OpenPhrase(uint32 DataSize)
 {
+	// We may have ended up exactly on the end of the journal file, immediately
+	// wrap in that case.
+	if (Cursor >= MaxSize)
+	{
+		Cursor = 0;
+	}
+
 	check((Cursor & (sizeof(FDataEntry) - 1)) == 0);
 
 	Entries.Add(FDataEntry{});
@@ -613,13 +620,6 @@ int32 FDiskJournal::Flush()
 		Cursor += Size;
 	}
 
-	// We may end up exactly on the end of the journal file,
-	// immediately wrap in that case.
-	if (Cursor >= MaxSize)
-	{
-		Cursor = 0;
-	}
-
 	Entries.Reset();
 	OpenJrnFile();
 
