Minor improvement to code documentation.

#rnx

#Synced-CL 48346827

diff --git a/Engine/Source/Runtime/Experimental/IoStore/OnDemand/Private/JournaledCache.cpp b/Engine/Source/Runtime/Experimental/IoStore/OnDemand/Private/JournaledCache.cpp
--- a/Engine/Source/Runtime/Experimental/IoStore/OnDemand/Private/JournaledCache.cpp
+++ b/Engine/Source/Runtime/Experimental/IoStore/OnDemand/Private/JournaledCache.cpp
@@ -682,7 +682,7 @@ FDiskCache::FDiskCache(FString&& Path, uint64 InMaxDataSize, uint32 InJournalSiz
 {
 	check(MaxDataSize >= HASH_CHECKSUM_SIZE);
 
-	// Align down to keep to some assumptions
+	// Remove the space reserved for the journal and align down to the nearest MiB  keep to some assumptions
 	MaxDataSize = (MaxDataSize - Journal.GetMaxSize()) & ~((1ull << 20) - 1);
 
 	OpenDataFile();
