Fix a number of spelling mistakes in code documentation.

#rnx

#Synced-CL 48346827

diff --git a/Engine/Source/Runtime/Experimental/IoStore/OnDemandCore/Public/IO/IoStoreOnDemand.h b/Engine/Source/Runtime/Experimental/IoStore/OnDemandCore/Public/IO/IoStoreOnDemand.h
--- a/Engine/Source/Runtime/Experimental/IoStore/OnDemandCore/Public/IO/IoStoreOnDemand.h
+++ b/Engine/Source/Runtime/Experimental/IoStore/OnDemandCore/Public/IO/IoStoreOnDemand.h
@@ -331,7 +331,7 @@ class FOnDemandInstallRequest final
 public:
 	/** Creates an invalid install request. */
 	UE_API FOnDemandInstallRequest();
-	/** Move constructable. */;
+	/** Move constructor. */;
 	UE_API FOnDemandInstallRequest(FOnDemandInstallRequest&&);
 	/** Move assignable. */
 	UE_API FOnDemandInstallRequest& operator=(FOnDemandInstallRequest&&);
@@ -341,7 +341,7 @@ public:
 	UE_API virtual EStatus GetStatus() const override;
 	/** Cancel the install request. */
 	UE_API void Cancel();
-	/** Update priority of the instlal request. */
+	/** Update priority of the install request. */
 	UE_API void UpdatePriority(int32 NewPriority);
 
 private:
@@ -370,7 +370,7 @@ struct FOnDemandPurgeArgs
 {
 	/** Purge options. */
 	EOnDemandPurgeOptions Options = EOnDemandPurgeOptions::None;
-	/** Optional Size to purge. If not set will purge all unfreferenced blocks. */
+	/** Optional Size to purge. If not set will purge all unreferenced blocks. */
 	TOptional<uint64> BytesToPurge;
 };
 
@@ -403,7 +403,7 @@ struct FOnDemandDefragArgs
 {
 	/** Defrag options. */
 	EOnDemandDefragOptions Options = EOnDemandDefragOptions::None;
-	/** Optional Size to free. If not set, will degrag all blocks and free all unreferenced chunks. */
+	/** Optional Size to free. If not set, will defrag all blocks and free all unreferenced chunks. */
 	TOptional<uint64> BytesToFree;
 };
 
@@ -440,7 +440,7 @@ struct FOnDemandGetInstallSizeArgs
 	FString MountId;
 	/** Get total install size for the specified tag(s) and optionally matching the mount ID. */
 	TArray<FString> TagSets;
-	/** Get total intall size for the specified package IDs. */
+	/** Get total install size for the specified package IDs. */
 	TArray<FPackageId> PackageIds;
 	/** Options */
 	EOnDemandGetInstallSizeOptions Options = EOnDemandGetInstallSizeOptions::None;
@@ -467,7 +467,7 @@ enum class EOnDemandIsOnDemandOptions
 };
 ENUM_CLASS_FLAGS(EOnDemandIsOnDemandOptions);
 
-/** Arguments for checking if any on-demand content exists in the dependecy graph. */
+/** Arguments for checking if any on-demand content exists in the dependency graph. */
 struct FOnDemandIsOnDemandArgs
 {
 	/** Return true if any of the specified package IDs has on-demand dependencies. */
@@ -522,7 +522,7 @@ struct FOnDemandInstallCacheUsage
 	UE_API friend FStringBuilderBase& operator<<(FStringBuilderBase& Sb, const FOnDemandInstallCacheUsage& CacheUsage);
 };
 
-/** Holds information about the stremaing cache usage */
+/** Holds information about the streaming cache usage */
 struct FOnDemandStreamingCacheUsage
 {
 	uint64 MaxSize = 0;
@@ -677,7 +677,7 @@ public:
 	/**
 	 * Query whether any of the specified packages have on-demand dependencies.
 	 * @param Args	Package set to query.
-	 * @return		Results indicating whether on-demand depedencies exist.
+	 * @return		Results indicating whether on-demand dependencies exist.
 	 */
 	UE_API TIoStatusOr<FOnDemandIsOnDemandResult> GetIsOnDemand(const FOnDemandIsOnDemandArgs& Args) const;
 	/** Returns the total cache size in bytes. */
