[Chaos] Remove modification flag from FContactPairModifier::GetWorldContactLocations. This function doesn't modify the contact and was causing the contact to be removed from the cache

#Synced-CL 48390192

diff --git a/Engine/Source/Runtime/Experimental/Chaos/Private/Chaos/ContactModification.cpp b/Engine/Source/Runtime/Experimental/Chaos/Private/Chaos/ContactModification.cpp
--- a/Engine/Source/Runtime/Experimental/Chaos/Private/Chaos/ContactModification.cpp
+++ b/Engine/Source/Runtime/Experimental/Chaos/Private/Chaos/ContactModification.cpp
@@ -146,13 +146,11 @@ namespace Chaos
 
 	void FContactPairModifier::GetWorldContactLocations(int32 ContactPointIdx, FVec3& OutLocation0, FVec3& OutLocation1) const
 	{
-		Constraint->SetModifierApplied();
-
 		const FRigidTransform3& ShapeTransform0 = Constraint->GetShapeWorldTransform0();
 		const FRigidTransform3& ShapeTransform1 = Constraint->GetShapeWorldTransform1();
 
 		TArrayView<FManifoldPoint> ManifoldPoints = Constraint->GetManifoldPoints();
-		FManifoldPoint& ManifoldPoint = ManifoldPoints[ContactPointIdx];
+		const FManifoldPoint& ManifoldPoint = ManifoldPoints[ContactPointIdx];
 		OutLocation0 = ShapeTransform0.TransformPositionNoScale(FVec3(ManifoldPoint.ContactPoint.ShapeContactPoints[0]));
 		OutLocation1 = ShapeTransform1.TransformPositionNoScale(FVec3(ManifoldPoint.ContactPoint.ShapeContactPoints[1]));
 	}
