Add option to quantize per-face bound in the chaos trimesh bvh, to optionally save ~40% memory per complex collision mesh, w/ some additional perf cost for queries

#Synced-CL 48316605

diff --git a/Engine/Source/Runtime/Engine/Private/PhysicsEngine/Experimental/ChaosCooking.cpp b/Engine/Source/Runtime/Engine/Private/PhysicsEngine/Experimental/ChaosCooking.cpp
--- a/Engine/Source/Runtime/Engine/Private/PhysicsEngine/Experimental/ChaosCooking.cpp
+++ b/Engine/Source/Runtime/Engine/Private/PhysicsEngine/Experimental/ChaosCooking.cpp
@@ -137,7 +137,12 @@ namespace Chaos
 
 				TUniquePtr<TArray<int32>> OutFaceRemapPtr = MakeUnique<TArray<int32>>(OutFaceRemap);
 				TUniquePtr<TArray<int32>> OutVertexRemapPtr = Chaos::TriMeshPerPolySupport ? MakeUnique<TArray<int32>>(OutVertexRemap) : nullptr;
-				Chaos::FTriangleMeshImplicitObjectPtr TriangleMesh( new Chaos::FTriangleMeshImplicitObject(MoveTemp(TriMeshParticles), MoveTemp(Triangles), MoveTemp(MaterialIndices), MoveTemp(OutFaceRemapPtr), MoveTemp(OutVertexRemapPtr)));
+
+				Chaos::FTriangleMeshImplicitObject::FBuildSettings MeshBuildSettings;
+				MeshBuildSettings.PerFaceBounds = Desc.bReducedMemoryRepresentation ? Chaos::FTrimeshBVH::EPerFaceBoundsRepresentation::Byte : Chaos::FTrimeshBVH::EPerFaceBoundsRepresentation::VectorizedFloat;
+				constexpr bool bCullsBackFaceRaycast = false;
+				Chaos::FTriangleMeshImplicitObjectPtr TriangleMesh( new Chaos::FTriangleMeshImplicitObject(MoveTemp(TriMeshParticles), MoveTemp(Triangles), MoveTemp(MaterialIndices), MoveTemp(OutFaceRemapPtr), MoveTemp(OutVertexRemapPtr),
+					bCullsBackFaceRaycast, MeshBuildSettings));
 
 				// Propagate remapped indices from the FTriangleMeshImplicitObject back to the remap array
 				CopyUpdatedFaceRemapFromTriangleMesh(*TriangleMesh.GetReference(), OutFaceRemap);
@@ -320,7 +325,11 @@ namespace Chaos
 
 				TUniquePtr<TArray<int32>> OutFaceRemapPtr = MakeUnique<TArray<int32>>(OutFaceRemap);
 				TUniquePtr<TArray<int32>> OutVertexRemapPtr = Chaos::TriMeshPerPolySupport ? MakeUnique<TArray<int32>>(OutVertexRemap) : nullptr;
-				Chaos::FTriangleMeshImplicitObject* TriangleMesh = new Chaos::FTriangleMeshImplicitObject(MoveTemp(TriMeshParticles), MoveTemp(Triangles), MoveTemp(MaterialIndices), MoveTemp(OutFaceRemapPtr), MoveTemp(OutVertexRemapPtr));
+				Chaos::FTriangleMeshImplicitObject::FBuildSettings MeshBuildSettings;
+				MeshBuildSettings.PerFaceBounds = InParams.TriangleMeshDesc.bReducedMemoryRepresentation ? Chaos::FTrimeshBVH::EPerFaceBoundsRepresentation::Byte : Chaos::FTrimeshBVH::EPerFaceBoundsRepresentation::VectorizedFloat;
+				constexpr bool bCullsBackFaceRaycast = false;
+				Chaos::FTriangleMeshImplicitObject* TriangleMesh = new Chaos::FTriangleMeshImplicitObject(MoveTemp(TriMeshParticles), MoveTemp(Triangles), MoveTemp(MaterialIndices), MoveTemp(OutFaceRemapPtr), MoveTemp(OutVertexRemapPtr),
+					bCullsBackFaceRaycast, MeshBuildSettings);
 				OutTriangleMeshes.Emplace(TriangleMesh);
 
 				// Propagate remapped indices from the FTriangleMeshImplicitObject back to the remap array
diff --git a/Engine/Source/Runtime/Experimental/Chaos/Private/Chaos/TriangleMeshImplicitObject.cpp b/Engine/Source/Runtime/Experimental/Chaos/Private/Chaos/TriangleMeshImplicitObject.cpp
--- a/Engine/Source/Runtime/Experimental/Chaos/Private/Chaos/TriangleMeshImplicitObject.cpp
+++ b/Engine/Source/Runtime/Experimental/Chaos/Private/Chaos/TriangleMeshImplicitObject.cpp
@@ -813,6 +813,112 @@ TArray<int32> FTrimeshBVH::FindAllIntersections (const FAABB3& Intersection) con
 }
 
 
+void FTrimeshBVH::Serialize(FChaosArchive& Ar)
+{
+	Ar.UsingCustomVersion(FUE5MainStreamObjectVersion::GUID);
+	Ar << Nodes;
+	Ar << FaceBounds;
+	if (Ar.CustomVer(FUE5MainStreamObjectVersion::GUID) < FUE5MainStreamObjectVersion::RemoveTriangleMeshBVHFaces)
+	{
+		TArray<int32> TmpFaces;
+		Ar << TmpFaces;
+	}
+}
+
+void FTrimeshBVH::InitializeQuantizedBounds(const TParticles<FRealSingle, 3>& Particles, const FTrimeshIndexBuffer& IndexBuffer)
+{
+	// Make sure we don't have the full bounds in memory if we're using quantized bounds
+	FaceBounds.Empty();
+
+	// Helper to convert vectorized AABB to a standard AABB
+	auto StoreAABBVectorized = [](const FAABBVectorized& VecBounds) -> TAABB<FRealSingle, 3>
+	{
+		FVector3f StoreMin, StoreMax;
+		VectorStoreFloat3(VecBounds.GetMin(), &StoreMin.X);
+		VectorStoreFloat3(VecBounds.GetMax(), &StoreMax.X);
+		return TAABB<FRealSingle, 3>(StoreMin, StoreMax);
+	};
+
+	auto GetTriBounds = [&Particles]<typename IndexBufferType>(const IndexBufferType& Buffer, int32 Idx)
+	{
+		FAABB3f Bounds;
+		Bounds.GrowToInclude(Particles.GetX(Buffer[Idx].X));
+		Bounds.GrowToInclude(Particles.GetX(Buffer[Idx].Y));
+		Bounds.GrowToInclude(Particles.GetX(Buffer[Idx].Z));
+		return Bounds;
+	};
+	
+	QuantizedFaceBoundData.SetNumUninitialized(IndexBuffer.GetNumTriangles() * 6 + 1);
+	for (const FNode& Node : Nodes)
+	{
+		for (int32 ChildIndex = 0; ChildIndex < 2; ++ChildIndex)
+		{
+			const FChildData& ChildData = Node.Children;
+			const int32 FaceIndex = ChildData.GetChildOrFaceIndex(ChildIndex);
+			if (FaceIndex != INDEX_NONE)
+			{
+				const uint32 FaceCount = ChildData.GetFaceCount(ChildIndex);
+				if (FaceCount > 0)
+				{
+					TAABB<FRealSingle, 3> NodeBounds = StoreAABBVectorized(ChildData.GetBounds(ChildIndex));
+					uint8* QuantizedBounds = &QuantizedFaceBoundData[FaceIndex * 6];
+					for (uint32 Idx = FaceIndex, End = FaceIndex + FaceCount; Idx < End; ++Idx, QuantizedBounds += 6)
+					{
+						FAABB3f TriBounds = IndexBuffer.RequiresLargeIndices() ?
+							GetTriBounds(IndexBuffer.GetLargeIndexBuffer(), Idx) :
+							GetTriBounds(IndexBuffer.GetSmallIndexBuffer(), Idx);
+						SetQuantizedBounds(QuantizedBounds, NodeBounds, TriBounds);
+					}
+				}
+			}
+		}
+	}
+}
+
+bool FTrimeshBVH::NeedsQuantizedBounds()
+{
+	return !Nodes.IsEmpty() && FaceBounds.IsEmpty();
+}
+
+void FTrimeshBVH::ResetFaceBounds(int32 NumTris, bool bUseQuantized)
+{
+	if (bUseQuantized)
+	{
+		FaceBounds.Empty();
+		QuantizedFaceBoundData.Reset(NumTris * 6 + 1);
+	}
+	else
+	{
+		QuantizedFaceBoundData.Empty();
+		FaceBounds.Reset(NumTris);
+	}
+}
+
+void FTrimeshBVH::AddQuantizedFaceBound(const TAABB<FRealSingle, 3>& NodeBounds, const TAABB<FRealSingle, 3>& TriBounds)
+{
+	QuantizedFaceBoundData.AddUninitialized(6 + (int32)(QuantizedFaceBoundData.IsEmpty()));
+	// Note we always want to leave the last byte for padding, so we use Last(6) instead of Last(5) here
+	SetQuantizedBounds(&QuantizedFaceBoundData.Last(6), NodeBounds, TriBounds);
+}
+
+void FTrimeshBVH::SetQuantizedBounds(uint8* QuantizedData, const TAABB<FRealSingle, 3>& NodeBounds, const TAABB<FRealSingle, 3>& TriBounds)
+{
+	for (int32 Idx = 0; Idx < 3; ++Idx)
+	{
+		FRealSingle NodeBoundsRange = NodeBounds.Max()[Idx] - NodeBounds.Min()[Idx];
+		if (NodeBoundsRange < UE_SMALL_NUMBER)
+		{
+			QuantizedData[Idx] = 0;
+			QuantizedData[Idx + 3] = 255;
+		}
+		else
+		{
+			QuantizedData[Idx] = static_cast<uint8>(FMath::Clamp(FMath::FloorToInt(255.f * (TriBounds.Min()[Idx] - NodeBounds.Min()[Idx]) / NodeBoundsRange), 0, 255));
+			QuantizedData[Idx + 3] = static_cast<uint8>(FMath::Clamp(FMath::CeilToInt(255.f * (TriBounds.Max()[Idx] - NodeBounds.Min()[Idx]) / NodeBoundsRange), 0, 255));
+		}
+	}
+}
+
 FReal FTriangleMeshImplicitObject::PhiWithNormal(const FVec3& x, FVec3& Normal) const
 {
 	FSphere TestSphere(x, 0.0f);
@@ -827,7 +933,7 @@ FReal FTriangleMeshImplicitObject::PhiWithNormal(const FVec3& x, FVec3& Normal)
 template <typename IdxType>
 bool FTriangleMeshImplicitObject::RaycastImp(const TArray<TVector<IdxType, 3>>& Elements, const FVec3& StartPoint, const FVec3& Dir, const FReal Length, const FReal Thickness, FReal& OutTime, FVec3& OutPosition, FVec3& OutNormal, int32& OutFaceIndex) const
 {
-	FTriangleMeshRaycastVisitor<IdxType> SQVisitor(StartPoint, Dir, Thickness, MParticles, Elements, bCullsBackFaceRaycast);
+	FTriangleMeshRaycastVisitor<IdxType> SQVisitor(StartPoint, Dir, Thickness, MParticles, Elements, bool(MeshSettingsFlags & ETriangleMeshSettingsFlags::CullsBackFaceRaycast));
 
 	if (Thickness > 0)
 	{
@@ -1040,12 +1146,19 @@ int32 FTriangleMeshImplicitObject::GetExternalFaceIndexFromInternal(int32 Intern
 
 bool FTriangleMeshImplicitObject::GetCullsBackFaceRaycast() const
 {
-	return bCullsBackFaceRaycast;
+	return bool(MeshSettingsFlags & ETriangleMeshSettingsFlags::CullsBackFaceRaycast);
 }
 
 void FTriangleMeshImplicitObject::SetCullsBackFaceRaycast(const bool bInCullsBackFace)
 {
-	bCullsBackFaceRaycast = bInCullsBackFace;
+	if (bInCullsBackFace)
+	{
+		MeshSettingsFlags |= ETriangleMeshSettingsFlags::CullsBackFaceRaycast;
+	}
+	else
+	{
+		MeshSettingsFlags &= ~ETriangleMeshSettingsFlags::CullsBackFaceRaycast;
+	}
 }
 
 template <typename IdxType>
@@ -1655,7 +1768,7 @@ bool FTriangleMeshImplicitObject::SweepGeomImp(const QueryGeomType& QueryGeom, c
 	{
 		// When length is zero, the direction of the ray is currently a hardcoded value (since it's computed from the start/end which would be equal). 
 		// We don't want the backface culling check to be arbitrary based upon the hardcoded direction, so we disable it with zero-length sweeps.
-		const FReal CullsBackFaceRaycastCode = (bCullsBackFaceRaycast && Length != 0) ? GetWindingOrder(TriMeshScale) : 0.f;
+		const FReal CullsBackFaceRaycastCode = (bool(MeshSettingsFlags & ETriangleMeshSettingsFlags::CullsBackFaceRaycast) && Length != 0) ? GetWindingOrder(TriMeshScale) : 0.f;
 		using VisitorType = FTriangleMeshSweepVisitor<QueryGeomType,decltype(Elements[0][0])>;
 		VisitorType SQVisitor(*this, Elements, QueryGeom, ScaledDirNormalized, LengthScale, ScaledStartTM, bComputeMTD, TriMeshScale, CullsBackFaceRaycastCode);
 
@@ -1743,7 +1856,7 @@ bool FTriangleMeshImplicitObject::SweepGeomCCDImp(const QueryGeomType& QueryGeom
 		const FReal IgnorePenetration = CVars::bCCDNewTargetDepthMode ? InIgnorePenetration : 0;
 		const FReal TargetPenetration = CVars::bCCDNewTargetDepthMode ? InTargetPenetration : 0;
 
-		const FReal CullsBackFaceRaycastCode = bCullsBackFaceRaycast ? GetWindingOrder(TriMeshScale) : 0.f;
+		const FReal CullsBackFaceRaycastCode = bool(MeshSettingsFlags & ETriangleMeshSettingsFlags::CullsBackFaceRaycast) ? GetWindingOrder(TriMeshScale) : 0.f;
 		using VisitorType = FTriangleMeshSweepVisitorCCD<QueryGeomType, decltype(Elements[0][0])>;
 		VisitorType SQVisitor(*this, Elements, QueryGeom, ScaledDirNormalized, LengthScale, Length, ScaledStartTM, TriMeshScale, CullsBackFaceRaycastCode, IgnorePenetration, TargetPenetration);
 
@@ -1939,7 +2052,7 @@ FImplicitObjectPtr FTriangleMeshImplicitObject::CopySlowImpl(const TArray<TVecto
 		ExternalVertexIndexMapCopy = MakeUnique<TArray<int32>>(*ExternalVertexIndexMap.Get());
 	}
 
-	return FImplicitObjectPtr(new FTriangleMeshImplicitObject(MoveTemp(ParticlesCopy), MoveTemp(ElementsCopy), MoveTemp(MaterialIndicesCopy), MoveTemp(ExternalFaceIndexMapCopy), MoveTemp(ExternalVertexIndexMapCopy), bCullsBackFaceRaycast));
+	return FImplicitObjectPtr(new FTriangleMeshImplicitObject(MoveTemp(ParticlesCopy), MoveTemp(ElementsCopy), MoveTemp(MaterialIndicesCopy), MoveTemp(ExternalFaceIndexMapCopy), MoveTemp(ExternalVertexIndexMapCopy), GetCullsBackFaceRaycast()));
 }
 
 FImplicitObjectPtr FTriangleMeshImplicitObject::DeepCopyGeometry() const
@@ -2081,6 +2194,21 @@ void FTriangleMeshImplicitObject::RebuildBVImp(const TArray<TVec3<IdxType>>& Ele
 FTriangleMeshImplicitObject::~FTriangleMeshImplicitObject() = default;
 
 
+FTriangleMeshImplicitObject::ETriangleMeshSettingsFlags FTriangleMeshImplicitObject::PackSettingsFlags(bool bInCullsBackFaceRaycast, FTriangleMeshImplicitObject::FBuildSettings BuildSettings)
+{
+	ETriangleMeshSettingsFlags Flags{ 0 };
+	if (bInCullsBackFaceRaycast)
+	{
+		Flags |= ETriangleMeshSettingsFlags::CullsBackFaceRaycast;
+	}
+	if (BuildSettings.PerFaceBounds == FTrimeshBVH::EPerFaceBoundsRepresentation::Byte)
+	{
+		Flags |= ETriangleMeshSettingsFlags::QuantizePerFaceBounds;
+	}
+	return Flags;
+}
+
+
 void FTriangleMeshImplicitObject::RebuildFastBVH()
 {
 	BVHType TreeBVH;
@@ -2148,7 +2276,8 @@ void FTriangleMeshImplicitObject::RebuildFastBVHFromTree(const BVHType& TreeBVH)
 	const TArray<LeafType>& Leaves = TreeBVH.GetLeaves();
 
 	FastBVH.Nodes.Reset();
-	FastBVH.FaceBounds.Reset();
+	const bool bUseQuantizedBounds = bool(MeshSettingsFlags & ETriangleMeshSettingsFlags::QuantizePerFaceBounds);
+	FastBVH.ResetFaceBounds(MElements.GetNumTriangles(), bUseQuantizedBounds);
 
 	int32 FaceNum = 0;
 	TArray<TVec3<FTrimeshIndexBuffer::LargeIdxType>> LargeIndices;
@@ -2189,7 +2318,14 @@ void FTriangleMeshImplicitObject::RebuildFastBVHFromTree(const BVHType& TreeBVH)
 		NewNode.Children.SetBounds(0, BoundingBox());
 		for (const TPayloadBoundsElement<int32, FRealSingle>& LeafPayload: Leaf.Elems)
 		{
-			FastBVH.FaceBounds.Add(FAABBVectorized(LeafPayload.Bounds));
+			if (bUseQuantizedBounds)
+			{
+				FastBVH.AddQuantizedFaceBound(BoundingBoxf(), LeafPayload.BoundingBox());
+			}
+			else
+			{
+				FastBVH.FaceBounds.Add(FAABBVectorized(LeafPayload.Bounds));
+			}
 			// Reorder triangle indices, triangles will be in the same order as the bounding volume in the BVH structure.
 			// And all triangles in a leaf will be contiguous in memory.
 			AddTriangles(LargeIndices, SmallIndices, MElements, LeafPayload.Payload);
@@ -2276,7 +2412,14 @@ void FTriangleMeshImplicitObject::RebuildFastBVHFromTree(const BVHType& TreeBVH)
 						// copy indices in the linear face array
 						for (const auto& LeafPayload: Leaf.Elems)
 						{
-							FastBVH.FaceBounds.Add(FAABBVectorized(LeafPayload.Bounds));
+							if (bUseQuantizedBounds)
+							{
+								FastBVH.AddQuantizedFaceBound(Node.ChildrenBounds[ChildIndex], LeafPayload.Bounds);
+							}
+							else
+							{
+								FastBVH.FaceBounds.Add(FAABBVectorized(LeafPayload.Bounds));
+							}
 							// Reorder triangle indices, triangles will be in the same order as the bounding volume in the BVH structure.
 							// And all triangles in a leaf will be contiguous in memory.
 							AddTriangles(LargeIndices, SmallIndices, MElements, LeafPayload.Payload);
diff --git a/Engine/Source/Runtime/Experimental/Chaos/Public/Chaos/TriangleMeshImplicitObject.h b/Engine/Source/Runtime/Experimental/Chaos/Public/Chaos/TriangleMeshImplicitObject.h
--- a/Engine/Source/Runtime/Experimental/Chaos/Public/Chaos/TriangleMeshImplicitObject.h
+++ b/Engine/Source/Runtime/Experimental/Chaos/Public/Chaos/TriangleMeshImplicitObject.h
@@ -36,10 +36,18 @@ namespace Chaos
 	class FTriangleMeshImplicitObject;
 	FArchive& operator<<(FArchive& Ar, FAABBVectorized& Bounds);
 
+	class FTrimeshIndexBuffer;
+
 	struct FTrimeshBVH
 	{
 		using FAABBType = TAABB<FRealSingle, 3>;
 
+		enum class EPerFaceBoundsRepresentation : uint8
+		{
+			VectorizedFloat = 0,
+			Byte
+		};
+
 		enum class EVisitorResult
 		{
 			Stop = 0,
@@ -89,7 +97,7 @@ namespace Chaos
 				Bounds[ChildIndex] = FAABBVectorized(AABB);
 			}
 
-			const FAABBVectorized& GetBounds(int ChildIndex) const { return Bounds[ChildIndex]; }
+			FORCEINLINE const FAABBVectorized& GetBounds(int ChildIndex) const { return Bounds[ChildIndex]; }
 
 			void Serialize(FArchive& Ar)
 			{
@@ -213,7 +221,8 @@ namespace Chaos
 		template <typename QueryGeomType>
 		bool FindAllIntersectionsNoMTD(const FAABB3& Intersection, const TRigidTransform<FReal, 3>& Transform, const QueryGeomType& QueryGeom, FReal Thickness, const FVec3& TriMeshScale, const FTriangleMeshImplicitObject* TriMesh) const;
 		CHAOS_API TArray<int32> FindAllIntersections(const FAABB3& Intersection) const;
-
+		
+		// original visit faces method, using original/full-float face bounds
 		template <typename BoundsFilterType, typename FaceVisitorType>
 		FORCEINLINE_DEBUGGABLE EVisitorResult VisitFaces(int32 StartIndex, int32 IndexCount, BoundsFilterType& BoundsFilter, FaceVisitorType& FaceVisitor) const
 		{
@@ -238,6 +247,19 @@ namespace Chaos
 			{
 				return;
 			}
+			if (!FaceBounds.IsEmpty())
+			{
+				VisitTreeImpl<BoundsFilterType, FaceVisitorType, false>(BoundsFilter, FaceVisitor);
+			}
+			else
+			{
+				VisitTreeImpl<BoundsFilterType, FaceVisitorType, true>(BoundsFilter, FaceVisitor);
+			}
+		}
+private:
+		template <typename BoundsFilterType, typename FaceVisitorType, bool bUseQuantizedBounds>
+		FORCEINLINE_DEBUGGABLE void VisitTreeImpl(BoundsFilterType& BoundsFilter, FaceVisitorType& FaceVisitor) const
+		{
 			TArray<int32, TInlineAllocator<16>> NodeIndexStack;
 			NodeIndexStack.Push(0);
 			while (NodeIndexStack.Num())
@@ -257,9 +279,19 @@ namespace Chaos
 							const uint32 FaceCount = ChildData.GetFaceCount(ChildIndex);
 							if (FaceCount > 0)
 							{
-								if (EVisitorResult::Stop == VisitFaces(FaceIndex, FaceCount, BoundsFilter, FaceVisitor))
+								if constexpr (!bUseQuantizedBounds)
 								{
-									return;
+									if (EVisitorResult::Stop == VisitFaces(FaceIndex, FaceCount, BoundsFilter, FaceVisitor))
+									{
+										return;
+									}
+								}
+								else
+								{
+									if (EVisitorResult::Stop == VisitFacesQuantized(ChildData.GetBounds(ChildIndex), FaceIndex, FaceCount, BoundsFilter, FaceVisitor))
+									{
+										return;
+									}
 								}
 							}
 							else
@@ -271,21 +303,65 @@ namespace Chaos
 				}
 			}
 		}
+public:
+
+		CHAOS_API void Serialize(FChaosArchive& Ar);
+		
+		TArray<FNode> Nodes;
+		TArray<FAABBVectorized> FaceBounds;
+
+		/// Quantized bounds data/methods:
+		/// Optionally stores the FaceBounds as uint8 fractions of their parent node bounds
+		/// to save storage costs, instead of using the FaceBounds array
+
+private:
+		// Array of quantized face bound data; with 6 elements per triangle + one buffer element at the end to allow safe use of vectorized loads
+		TArray<uint8> QuantizedFaceBoundData;
+
+		// @return true if the BVH should use quantized bounds -- i.e., if it is not empty and has no FaceBounds
+		bool NeedsQuantizedBounds();
+		
+		// Set quantized bounds from mesh data
+		void InitializeQuantizedBounds(const TParticles<FRealSingle, 3>& Particles, const FTrimeshIndexBuffer& IndexBuffer);
+		
+		// Reset face bounds data
+		void ResetFaceBounds(int32 NumTris, bool bUseQuantized);
+		
+		// Add a new quantized bound
+		void AddQuantizedFaceBound(const TAABB<FRealSingle, 3>& NodeBounds, const TAABB<FRealSingle, 3>& TriBounds);	
 
-		void Serialize(FChaosArchive& Ar)
+		// alternate visit faces method, using quantized face bounds
+		template <typename BoundsFilterType, typename FaceVisitorType>
+		FORCEINLINE_DEBUGGABLE EVisitorResult VisitFacesQuantized(const FAABBVectorized& NodeBounds, int32 StartIndex, int32 IndexCount, BoundsFilterType& BoundsFilter, FaceVisitorType& FaceVisitor) const
 		{
-			Ar.UsingCustomVersion(FUE5MainStreamObjectVersion::GUID);
-			Ar << Nodes;
-			Ar << FaceBounds;
-			if (Ar.CustomVer(FUE5MainStreamObjectVersion::GUID) < FUE5MainStreamObjectVersion::RemoveTriangleMeshBVHFaces)
+			// compute the scale factor to convert quantized uint8 values to relative positions within the node
+			constexpr VectorRegister4Float DeQuantFracScale = MakeVectorRegisterFloatConstant(1.f / 255.f, 1.f / 255.f, 1.f / 255.f, 1.f);
+			const VectorRegister4Float DeQuantScale = VectorMultiply(VectorSubtract(NodeBounds.GetMax(), NodeBounds.GetMin()), DeQuantFracScale);
+
+			const int32 EndIndex = (StartIndex + IndexCount);
+			int32 FaceIndex = StartIndex;
+			// Note EndIndex*6 is 1 element *past* the actual quantized bounds data, but should still be valid because we store an extra buffer byte to allow VectorLoadByte4 to be safely used
+			checkSlow(IndexCount > 0 && (QuantizedFaceBoundData.IsValidIndex(FaceIndex * 6) && QuantizedFaceBoundData.IsValidIndex(EndIndex * 6)));
+			for (const uint8* QBoundPtr = QuantizedFaceBoundData.GetData() + StartIndex * 6; FaceIndex < EndIndex; ++FaceIndex, QBoundPtr += 6)
 			{
-				TArray<int32> TmpFaces;
-				Ar << TmpFaces;
+				if (BoundsFilter(FAABBVectorized(
+					VectorMultiplyAdd(DeQuantScale, VectorLoadByte4(&QBoundPtr[0]), NodeBounds.GetMin()),
+					VectorMultiplyAdd(DeQuantScale, VectorLoadByte4(&QBoundPtr[3]), NodeBounds.GetMin()))) == EFilterResult::Keep)
+				{
+					if (FaceVisitor(FaceIndex) == EVisitorResult::Stop)
+					{
+						return EVisitorResult::Stop;
+					}
+				}
 			}
+			return EVisitorResult::Continue;
 		}
 		
-		TArray<FNode> Nodes;
-		TArray<FAABBVectorized> FaceBounds;
+		// Private helper to set quantized bound data
+		void SetQuantizedBounds(uint8* QuantizedData, const TAABB<FRealSingle, 3>& NodeBounds, const TAABB<FRealSingle, 3>& TriBounds);
+		
+		// FTriangleMeshImplicitObject has friend access to build the quantized bounds data
+		friend class FTriangleMeshImplicitObject;
 	};
 
 	FORCEINLINE_DEBUGGABLE FChaosArchive& operator<<(FChaosArchive& Ar, FTrimeshBVH::FChildData& ChildData)
@@ -494,8 +570,29 @@ namespace Chaos
 		using ParticlesType = TParticles<FRealSingle, 3>;
 		using ParticleVecType = TVec3<FRealSingle>;
 
+		struct FBuildSettings
+		{
+			// Whether to store "quantized" bounds per triangle, instead of full vectorized bounds
+			// This can save roughly 40% memory for the in-RAM representation of the triangle mesh, for some additional access cost (depending on the platform)
+			FTrimeshBVH::EPerFaceBoundsRepresentation PerFaceBounds = FTrimeshBVH::EPerFaceBoundsRepresentation::VectorizedFloat;
+			
+			FBuildSettings() {}
+		};
+
+		// Bitflag representation to pack triangle mesh settings
+		enum class ETriangleMeshSettingsFlags : uint8
+		{
+			CullsBackFaceRaycast = 1 << 0,
+			QuantizePerFaceBounds = 1 << 1
+		};
+		FRIEND_ENUM_CLASS_FLAGS(ETriangleMeshSettingsFlags);
+
 		template <typename IdxType>
-		FTriangleMeshImplicitObject(ParticlesType&& Particles, TArray<TVec3<IdxType>>&& Elements, TArray<uint16>&& InMaterialIndices, TUniquePtr<TArray<int32>>&& InExternalFaceIndexMap = nullptr, TUniquePtr<TArray<int32>>&& InExternalVertexIndexMap = nullptr, const bool bInCullsBackFaceRaycast = false)
+		FTriangleMeshImplicitObject(ParticlesType&& Particles, TArray<TVec3<IdxType>>&& Elements, TArray<uint16>&& InMaterialIndices,
+			TUniquePtr<TArray<int32>>&& InExternalFaceIndexMap = nullptr,
+			TUniquePtr<TArray<int32>>&& InExternalVertexIndexMap = nullptr,
+			const bool bInCullsBackFaceRaycast = false,
+			const FBuildSettings InBuildSettings = FBuildSettings())
 		: FImplicitObject(EImplicitObject::HasBoundingBox | EImplicitObject::DisableCollisions, ImplicitObjectType::TriangleMesh)
 		, MParticles(MoveTemp(Particles))
 		, MElements(MoveTemp(Elements))
@@ -503,7 +600,7 @@ namespace Chaos
 		, MaterialIndices(MoveTemp(InMaterialIndices))
 		, ExternalFaceIndexMap(MoveTemp(InExternalFaceIndexMap))
 		, ExternalVertexIndexMap(MoveTemp(InExternalVertexIndexMap))
-		, bCullsBackFaceRaycast(bInCullsBackFaceRaycast)
+		, MeshSettingsFlags(PackSettingsFlags(bInCullsBackFaceRaycast, InBuildSettings))
 		{
 			const int32 NumTriangles = MElements.GetNumTriangles();
 			if(NumTriangles > 0)
@@ -589,6 +686,7 @@ namespace Chaos
 		CHAOS_API bool GetCullsBackFaceRaycast() const;
 		CHAOS_API void SetCullsBackFaceRaycast(const bool bInCullsBackFace);
 
+		CHAOS_API static ETriangleMeshSettingsFlags PackSettingsFlags(bool bInCullsBackFaceRaycast, FBuildSettings BuildSettings);
 
 		CHAOS_API virtual int32 FindMostOpposingFace(const FVec3& Position, const FVec3& UnitDir, int32 HintFaceIndex, FReal SearchDistance) const override;
 		CHAOS_API virtual int32 FindMostOpposingFaceScaled(const FVec3& Position, const FVec3& UnitDir, int32 HintFaceIndex, FReal SearchDist, const FVec3& Scale) const override;
@@ -720,6 +818,13 @@ namespace Chaos
 					}
 				}
 			}
+
+			// Quantized bounds are not serialized; we recompute them on load
+			if (Ar.IsLoading() && FastBVH.NeedsQuantizedBounds())
+			{
+				MeshSettingsFlags |= ETriangleMeshSettingsFlags::QuantizePerFaceBounds;
+				FastBVH.InitializeQuantizedBounds(MParticles, MElements);
+			}
 		}
 
 		CHAOS_API virtual void Serialize(FChaosArchive& Ar) override;
@@ -813,11 +918,14 @@ namespace Chaos
 		TArray<uint16> MaterialIndices;
 		TUniquePtr<TArray<int32>> ExternalFaceIndexMap;
 		TUniquePtr<TArray<int32>> ExternalVertexIndexMap;
-		bool bCullsBackFaceRaycast;
+		ETriangleMeshSettingsFlags MeshSettingsFlags;
 
 		// Initialising constructor privately declared for use in CopySlow to copy the underlying BVH
 		template <typename IdxType>
-		FTriangleMeshImplicitObject(ParticlesType&& Particles, TArray<TVec3<IdxType>>&& Elements, TArray<uint16>&& InMaterialIndices, const BVHType& InBvhToCopy, TUniquePtr<TArray<int32>>&& InExternalFaceIndexMap = nullptr, TUniquePtr<TArray<int32>>&& InExternalVertexIndexMap = nullptr, const bool bInCullsBackFaceRaycast = false)
+		FTriangleMeshImplicitObject(ParticlesType&& Particles, TArray<TVec3<IdxType>>&& Elements, TArray<uint16>&& InMaterialIndices, const BVHType& InBvhToCopy,
+			TUniquePtr<TArray<int32>>&& InExternalFaceIndexMap = nullptr, TUniquePtr<TArray<int32>>&& InExternalVertexIndexMap = nullptr, 
+			const bool bInCullsBackFaceRaycast = false,
+			const FBuildSettings InBuildSettings = FBuildSettings())
 			: FImplicitObject(EImplicitObject::HasBoundingBox | EImplicitObject::DisableCollisions, ImplicitObjectType::TriangleMesh)
 			, MParticles(MoveTemp(Particles))
 			, MElements(MoveTemp(Elements))
@@ -825,7 +933,7 @@ namespace Chaos
 			, MaterialIndices(MoveTemp(InMaterialIndices))
 			, ExternalFaceIndexMap(MoveTemp(InExternalFaceIndexMap))
 			, ExternalVertexIndexMap(MoveTemp(InExternalVertexIndexMap))
-			, bCullsBackFaceRaycast(bInCullsBackFaceRaycast)
+			, MeshSettingsFlags(PackSettingsFlags(bInCullsBackFaceRaycast, InBuildSettings))
 		{
 			const int32 NumTriangles = MElements.GetNumTriangles();
 			if(NumTriangles > 0)
@@ -938,4 +1046,6 @@ namespace Chaos
 		FImplicitObjectPtr CopySlowImpl(const TArray < TVector<IdxType, 3>>& InElements) const;
 	};
 
+	ENUM_CLASS_FLAGS(FTriangleMeshImplicitObject::ETriangleMeshSettingsFlags);
+
 }
diff --git a/Engine/Source/Runtime/PhysicsCore/Public/Interface_CollisionDataProviderCore.h b/Engine/Source/Runtime/PhysicsCore/Public/Interface_CollisionDataProviderCore.h
--- a/Engine/Source/Runtime/PhysicsCore/Public/Interface_CollisionDataProviderCore.h
+++ b/Engine/Source/Runtime/PhysicsCore/Public/Interface_CollisionDataProviderCore.h
@@ -50,11 +50,15 @@ struct FTriMeshCollisionData
 	/** Turn off ActiveEdgePrecompute (This makes cooking faster, but will slow contact generation) */
 	uint32 bDisableActiveEdgePrecompute : 1;
 
+	/** Use a reduced memory representation for the triangle mesh data, which may have additional query cost */
+	uint32 bReducedMemoryRepresentation : 1;
+
 	FTriMeshCollisionData()
 		: bFlipNormals(false)
 		, bDeformableMesh(false)
 		, bFastCook(false)
 		, bDisableActiveEdgePrecompute(false)
+		, bReducedMemoryRepresentation(false)
 	{
 	}
 };
