Constraints & compact set: fix use of ConstraintTicks.FindOrAdd not beeing stable when adding/removing items
- ConstraintTicks now stores a unique pointer to the constraint's tick function to ensure that the pointer does not change during insertion/deletion
- FTickTaskManager uses raw pointers to FTickFunctions so they must remain stable

#rnx

#Synced-CL 48343954

diff --git a/Engine/Source/Runtime/Experimental/Animation/Constraints/Private/ConstraintSubsystem.cpp b/Engine/Source/Runtime/Experimental/Animation/Constraints/Private/ConstraintSubsystem.cpp
--- a/Engine/Source/Runtime/Experimental/Animation/Constraints/Private/ConstraintSubsystem.cpp
+++ b/Engine/Source/Runtime/Experimental/Animation/Constraints/Private/ConstraintSubsystem.cpp
@@ -175,7 +175,6 @@ void UConstraintSubsystem::RemoveConstraint(UWorld* InWorld, UTickableConstraint
 	// disable constraint
 	InConstraint->Modify();
 	InConstraint->TeardownConstraint(InWorld);
-	InConstraint->SetActive(false);
 
 	if (FConstraintsInWorld* Constraints = ConstraintsInWorldFind(InWorld))
 	{
@@ -357,7 +356,6 @@ void FConstraintsInWorld::RemoveConstraints(UWorld* InWorld)
 		if (Constraint.IsValid())
 		{
 			Constraint->TeardownConstraint(InWorld);
-			Constraint->SetActive(false);
 		}
 	}
 	Constraints.SetNum(0);
diff --git a/Engine/Source/Runtime/Experimental/Animation/Constraints/Private/ConstraintsManager.cpp b/Engine/Source/Runtime/Experimental/Animation/Constraints/Private/ConstraintsManager.cpp
--- a/Engine/Source/Runtime/Experimental/Animation/Constraints/Private/ConstraintsManager.cpp
+++ b/Engine/Source/Runtime/Experimental/Animation/Constraints/Private/ConstraintsManager.cpp
@@ -79,26 +79,57 @@ FString FConstraintTickFunction::DiagnosticMessage()
  * UTickableConstraint
  **/
 
-FConstraintTickFunction& UTickableConstraint::GetTickFunction(UWorld* InWorld) 
+FConstraintTickFunction& UTickableConstraint::GetTickFunction(const UWorld* InWorld)
 {
-	FConstraintTickFunction& ConstraintTick = ConstraintTicks.FindOrAdd(InWorld->GetCurrentLevel());
-	return ConstraintTick;
+	if (!ensure(InWorld && InWorld->GetCurrentLevel()))
+	{
+		static FConstraintTickFunction DummyFunction;
+		return DummyFunction;
+	}
+
+	if (TUniquePtr<FConstraintTickFunction>* Found = ConstraintTicks.Find(InWorld->GetCurrentLevel()))
+	{
+		if (!ensure(Found->Get()))
+		{
+			(*Found) = MakeUnique<FConstraintTickFunction>();
+		}
+		return *Found->Get();
+	}
+	
+	TUniquePtr<FConstraintTickFunction>& NewFunction = ConstraintTicks.Emplace(InWorld->GetCurrentLevel(), MakeUnique<FConstraintTickFunction>());
+	return *NewFunction.Get();
 }
 
-const FConstraintTickFunction& UTickableConstraint::GetTickFunction(UWorld* InWorld) const
+const FConstraintTickFunction& UTickableConstraint::GetTickFunction(const UWorld* InWorld) const
 {
-	const FConstraintTickFunction& ConstraintTick = ConstraintTicks.FindOrAdd(InWorld->GetCurrentLevel());
-	return ConstraintTick;
+	if (!ensure(InWorld && InWorld->GetCurrentLevel()))
+	{
+		static const FConstraintTickFunction DummyFunction;
+		return DummyFunction;
+	}
+	
+	if (TUniquePtr<FConstraintTickFunction>* Found = ConstraintTicks.Find(InWorld->GetCurrentLevel()))
+	{
+		if (!ensure(Found->Get()))
+		{
+			(*Found) = MakeUnique<FConstraintTickFunction>();
+		}
+		return *Found->Get();
+	}
+	
+	TUniquePtr<FConstraintTickFunction>& NewFunction = ConstraintTicks.Emplace(InWorld->GetCurrentLevel(), MakeUnique<FConstraintTickFunction>());
+	return *NewFunction.Get();
 }
 
 void UTickableConstraint::SetActive(const bool bIsActive)
 {
 	Active = bIsActive;
-	for (TPair<TWeakObjectPtr<ULevel>, FConstraintTickFunction>& Pair : ConstraintTicks)
+
+	for (auto& [WeakLevel, UniqueFunction] : ConstraintTicks)
 	{
-		if (Pair.Key.IsValid())
+		if (WeakLevel.IsValid() && UniqueFunction.IsValid())
 		{
-			Pair.Value.SetTickFunctionEnable(bIsActive);
+			UniqueFunction->SetTickFunctionEnable(bIsActive);
 		}
 	}
 }
@@ -110,11 +141,11 @@ bool UTickableConstraint::IsFullyActive() const
 
 void UTickableConstraint::Evaluate(bool bTickHandlesAlso) const
 {
-	for (const TPair<TWeakObjectPtr<ULevel>, FConstraintTickFunction>& Pair : ConstraintTicks)
+	for (const auto& [WeakLevel, UniqueFunction] : ConstraintTicks)
 	{
-		if (Pair.Key.IsValid())
+		if (WeakLevel.IsValid() && UniqueFunction.IsValid())
 		{
-			Pair.Value.EvaluateFunctions();
+			UniqueFunction->EvaluateFunctions();
 		}
 	}
 }
@@ -178,14 +209,14 @@ void UTickableConstraint::PostEditChangeProperty(FPropertyChangedEvent& Property
 	const FName PropertyName = PropertyChangedEvent.GetPropertyName();
 	if (PropertyName == GET_MEMBER_NAME_CHECKED(UTickableConstraint, Active))
 	{
-		for (TPair<TWeakObjectPtr<ULevel>, FConstraintTickFunction>& Pair : ConstraintTicks)
+		for (auto& [WeakLevel, UniqueFunction] : ConstraintTicks)
 		{
-			if (Pair.Key.IsValid())
+			if (WeakLevel.IsValid() && UniqueFunction.IsValid())
 			{
-				Pair.Value.SetTickFunctionEnable(Active);
+				UniqueFunction->SetTickFunctionEnable(Active);
 				if (Active)
 				{
-					Pair.Value.EvaluateFunctions();
+					UniqueFunction->EvaluateFunctions();
 				}
 			}	
 		}
@@ -198,12 +229,12 @@ void UTickableConstraint::PostEditUndo()
 
 	// make sure we update ticking when undone
 	const bool bActiveTick = ::IsValidChecked(this);
-	for (TPair<TWeakObjectPtr<ULevel>, FConstraintTickFunction>& Pair : ConstraintTicks)
+	for (auto& [WeakLevel, UniqueFunction] : ConstraintTicks)
 	{
-		if (Pair.Key.IsValid())
+		if (WeakLevel.IsValid() && UniqueFunction.IsValid())
 		{
-			Pair.Value.SetTickFunctionEnable(bActiveTick);
-			Pair.Value.bCanEverTick = bActiveTick;
+			UniqueFunction->SetTickFunctionEnable(bActiveTick);
+			UniqueFunction->bCanEverTick = bActiveTick;
 		}
 	}
 }
diff --git a/Engine/Source/Runtime/Experimental/Animation/Constraints/Private/Transform/TransformConstraint.cpp b/Engine/Source/Runtime/Experimental/Animation/Constraints/Private/Transform/TransformConstraint.cpp
--- a/Engine/Source/Runtime/Experimental/Animation/Constraints/Private/Transform/TransformConstraint.cpp
+++ b/Engine/Source/Runtime/Experimental/Animation/Constraints/Private/Transform/TransformConstraint.cpp
@@ -141,13 +141,6 @@ UTickableConstraint* UTickableTransformConstraint::Duplicate(UObject* NewOuter)
 		UTransformableHandle* HandleCopy = ParentTRSHandle->Duplicate(Dup);
 		Dup->ParentTRSHandle = HandleCopy;
 	}
-	for (TPair<TWeakObjectPtr<ULevel>, FConstraintTickFunction>& Pair : ConstraintTicks)
-	{
-		if (Pair.Key.IsValid())
-		{
-			UWorld* World = Pair.Key->GetTypedOuter<UWorld>();
-		}
-	}
 	return Dup;
 }
 
@@ -306,23 +299,23 @@ void UTickableTransformConstraint::SetupDependencies(const UWorld* InWorld)
 		}
 	}
 
-	if (!ensure(InWorld))
+	if (!ensure(InWorld && InWorld->GetCurrentLevel()))
 	{
 		return;
 	}
-	
+
 	FTickFunction* ParentTickFunction = GetParentHandleTickFunction();
 	FTickFunction* ChildTickFunction = GetChildHandleTickFunction();
-	
+
 	if (UE::Private::ShouldForceParentDependency() && ParentTickFunction && (ChildTickFunction != ParentTickFunction))
 	{
 		// manage dependencies
 		// force ConstraintTickFunction to tick after InParent does.
 		// Note that this might not register anything if the parent can't tick (static meshes for instance)
-		FConstraintTickFunction& ConstraintTick = ConstraintTicks.FindOrAdd(InWorld->GetCurrentLevel());
+		FConstraintTickFunction& ConstraintTick = GetTickFunction(InWorld);
 		ConstraintTick.AddPrerequisite(ParentTRSHandle->GetPrerequisiteObject(), *ParentTickFunction);
 	}
-	
+
 	// TODO also check for cycle dependencies here
 	if (ChildTickFunction)
 	{
@@ -332,15 +325,15 @@ void UTickableTransformConstraint::SetupDependencies(const UWorld* InWorld)
 			// force InChild to tick after ConstraintTickFunction does.
 			// Note that this might not register anything if the child can't tick (static meshes for instance)
 			// Also don't do if it's a skeletal mesh will cause cycle. (to be updated)
-			FConstraintTickFunction& ConstraintTick = ConstraintTicks.FindOrAdd(InWorld->GetCurrentLevel());
+			FConstraintTickFunction& ConstraintTick = GetTickFunction(InWorld);
 			ChildTickFunction->AddPrerequisite(this, ConstraintTick);
 		}
 	}
 }
 
-void UTickableTransformConstraint::EnsurePrimaryDependency(const UWorld* InWorld) const
+void UTickableTransformConstraint::EnsurePrimaryDependency(const UWorld* InWorld)
 {
-	if (!ensure(InWorld))
+	if (!ensure(InWorld && InWorld->GetCurrentLevel()))
 	{
 		return;
 	}
@@ -349,7 +342,7 @@ void UTickableTransformConstraint::EnsurePrimaryDependency(const UWorld* InWorld
 	const FTickFunction* ChildTickFunction = GetChildHandleTickFunction();
 	if (ParentTickFunction && (ChildTickFunction != ParentTickFunction))
 	{
-		FConstraintTickFunction& ConstraintTick = ConstraintTicks.FindOrAdd(InWorld->GetCurrentLevel());
+		FConstraintTickFunction& ConstraintTick = GetTickFunction(InWorld);
 		const TArray<FTickPrerequisite>& ParentPrerequisites = ConstraintTick.GetPrerequisites();
 		if (ParentPrerequisites.IsEmpty())
 		{
@@ -630,8 +623,9 @@ void UTickableTransformConstraint::InitConstraint(UWorld *InWorld)
 	{
 		return;
 	}
+
 	
-	FConstraintTickFunction& ConstraintTick = ConstraintTicks.FindOrAdd(Level);
+	FConstraintTickFunction& ConstraintTick = GetTickFunction(InWorld);
 
 	if (ConstraintTick.ConstraintFunctions.IsEmpty())
 	{
@@ -659,8 +653,13 @@ void UTickableTransformConstraint::TeardownConstraint(UWorld* InWorld)
 	}
 
 	const ULevel* Level = InWorld ? InWorld->GetCurrentLevel() : nullptr;
-	if (FConstraintTickFunction* ConstraintTick = Level ? ConstraintTicks.Find(Level) : nullptr)
+
+
+	TUniquePtr<FConstraintTickFunction>* Found = Level ? ConstraintTicks.Find(Level) : nullptr;
+	if (Found && Found->IsValid())
 	{
+		TUniquePtr<FConstraintTickFunction>& ConstraintTick = *Found;
+		
 		ConstraintTick->UnRegisterTickFunction();
 		ConstraintTick->SetTickFunctionEnable(false);
 
@@ -675,6 +674,11 @@ void UTickableTransformConstraint::TeardownConstraint(UWorld* InWorld)
 		}
 
 		ConstraintTicks.Remove(Level);
+		
+		if (ConstraintTicks.IsEmpty())
+		{
+			Active = false;
+		}
 	}
 
 	// unregister delegates. should handles delegates be unregistered as well?
diff --git a/Engine/Source/Runtime/Experimental/Animation/Constraints/Public/ConstraintsManager.h b/Engine/Source/Runtime/Experimental/Animation/Constraints/Public/ConstraintsManager.h
--- a/Engine/Source/Runtime/Experimental/Animation/Constraints/Public/ConstraintsManager.h
+++ b/Engine/Source/Runtime/Experimental/Animation/Constraints/Public/ConstraintsManager.h
@@ -110,8 +110,8 @@ public:
 	virtual void AddedToWorld(UWorld* InWorld) PURE_VIRTUAL(UTickableConstraint::AddedToWorld);
 	
 	/** Get Tick function in that world*/
-	CONSTRAINTS_API FConstraintTickFunction& GetTickFunction(UWorld* InWorld);
-	CONSTRAINTS_API const FConstraintTickFunction& GetTickFunction(UWorld* InWorld) const;
+	CONSTRAINTS_API FConstraintTickFunction& GetTickFunction(const UWorld* InWorld);
+	CONSTRAINTS_API const FConstraintTickFunction& GetTickFunction(const UWorld* InWorld) const;
 
 	// UObject interface
 	CONSTRAINTS_API virtual void PostLoad() override;
@@ -147,7 +147,7 @@ public:
 protected:
 
 	/** Tick function that will be registered and evaluated. */
-	mutable TMap<TWeakObjectPtr<ULevel>,FConstraintTickFunction> ConstraintTicks;
+	mutable TMap<TWeakObjectPtr<ULevel>, TUniquePtr<FConstraintTickFunction>, TInlineSetAllocator<4>> ConstraintTicks;
 };
 
 
diff --git a/Engine/Source/Runtime/Experimental/Animation/Constraints/Public/Transform/TransformConstraint.h b/Engine/Source/Runtime/Experimental/Animation/Constraints/Public/Transform/TransformConstraint.h
--- a/Engine/Source/Runtime/Experimental/Animation/Constraints/Public/Transform/TransformConstraint.h
+++ b/Engine/Source/Runtime/Experimental/Animation/Constraints/Public/Transform/TransformConstraint.h
@@ -118,7 +118,7 @@ public:
 	/**
 	* Sets up dependencies with the first primary prerequisite available if the parent does not tick.   
 	*/
-	CONSTRAINTS_API void EnsurePrimaryDependency(const UWorld* InWorld) const;
+	CONSTRAINTS_API void EnsurePrimaryDependency(const UWorld* InWorld);
 	
 protected:
 
