iax: Convenience API

#rnx

#Synced-CL Unknown

diff --git a/Engine/Source/Runtime/Experimental/IoStore/HttpClient/Internal/IO/Http/Client.h b/Engine/Source/Runtime/Experimental/IoStore/HttpClient/Internal/IO/Http/Client.h
--- a/Engine/Source/Runtime/Experimental/IoStore/HttpClient/Internal/IO/Http/Client.h
+++ b/Engine/Source/Runtime/Experimental/IoStore/HttpClient/Internal/IO/Http/Client.h
@@ -52,6 +52,14 @@ enum class EStatusCodeClass
 	Unknown,
 };
 
+////////////////////////////////////////////////////////////////////////////////
+enum class EContentStatus
+{
+	NoContent,		// response has no content; head, 100-class, 204, 304, etc.
+	KnownSize,		// "content-length" present
+	UnknownSize,	// no "content-length" field, streamed transfer (e.g. h1's chunked)
+};
+
 ////////////////////////////////////////////////////////////////////////////////
 using	FCertRootsRef	= UPTRINT;
 using	FTicket			= uint64;
@@ -154,6 +162,7 @@ public:
 	UE_API EStatusCodeClass	GetStatus() const;
 	UE_API uint32			GetStatusCode() const;
 	UE_API FAnsiStringView 	GetStatusMessage() const;
+	UE_API EContentStatus	GetContentStatus() const;
 	UE_API int64			GetContentLength() const;
 	UE_API EMimeType		GetContentType() const;
 	UE_API void				GetContentType(FAnsiStringView& Out) const;
@@ -206,6 +215,7 @@ public:
 	UE_API UPTRINT				GetParam() const;
 	UE_API FTicket				GetTicket() const;
 	UE_API uint32				GetIndex() const;
+	UE_API bool					IsComplete() const;
 	UE_API FResponse&			GetResponse() const;		// if GetId() == EId::Response
 	UE_API const FIoBuffer&		GetContent() const;			// if GetId() == EId::Content
 	UE_API uint32				GetContentLength() const;	//  |
diff --git a/Engine/Source/Runtime/Experimental/IoStore/HttpClient/Private/Activity.inl b/Engine/Source/Runtime/Experimental/IoStore/HttpClient/Private/Activity.inl
--- a/Engine/Source/Runtime/Experimental/IoStore/HttpClient/Private/Activity.inl
+++ b/Engine/Source/Runtime/Experimental/IoStore/HttpClient/Private/Activity.inl
@@ -368,7 +368,7 @@ UPTRINT FActivity::GetSinkParam() const
 ////////////////////////////////////////////////////////////////////////////////
 FIoBuffer& FActivity::GetContent() const
 {
-	check(State == EState::RecvContent || State == EState::RecvStream);
+	check(State == EState::RecvContent || State == EState::RecvStream || State == EState::RecvDone);
 	return *Dest;
 }
 
diff --git a/Engine/Source/Runtime/Experimental/IoStore/HttpClient/Private/Api.inl b/Engine/Source/Runtime/Experimental/IoStore/HttpClient/Private/Api.inl
--- a/Engine/Source/Runtime/Experimental/IoStore/HttpClient/Private/Api.inl
+++ b/Engine/Source/Runtime/Experimental/IoStore/HttpClient/Private/Api.inl
@@ -99,6 +99,17 @@ FAnsiStringView FResponse::GetStatusMessage() const
 	return Activity->GetTransaction()->GetStatusMessage();
 }
 
+////////////////////////////////////////////////////////////////////////////////
+EContentStatus FResponse::GetContentStatus() const
+{
+	switch (GetContentLength())
+	{
+	case -1:	return EContentStatus::UnknownSize;
+	case 0:		return EContentStatus::NoContent;
+	default:	return EContentStatus::KnownSize;
+	}
+}
+
 ////////////////////////////////////////////////////////////////////////////////
 int64 FResponse::GetContentLength() const
 {
@@ -205,6 +216,31 @@ uint32 FTicketStatus::GetIndex() const
 	return Activity->Slot;
 }
 
+////////////////////////////////////////////////////////////////////////////////
+bool FTicketStatus::IsComplete() const
+{
+	EId Id = GetId();
+	if (Id == EId::Response)
+	{
+		// If the response has no body there will be no further status updates
+		return GetResponse().GetContentStatus() == EContentStatus::NoContent;
+	}
+
+	if (Id == EId::Content)
+	{
+		const auto* Activity = (FActivity*)this;
+
+		// When the size of content is unknown then the end is signalled with
+		// an empty destination buffer. Otherwise a status update is only
+		// called when all the destination has been received.
+		int64 ContentLength = Activity->GetTransaction()->GetContentLength();
+		ContentLength = FMath::Max(ContentLength, 0ll);
+		return GetContent().GetSize() == ContentLength;
+	}
+
+	return true;
+}
+
 ////////////////////////////////////////////////////////////////////////////////
 FResponse& FTicketStatus::GetResponse() const
 {
diff --git a/Engine/Source/Runtime/Experimental/IoStore/HttpClient/Private/Test.inl b/Engine/Source/Runtime/Experimental/IoStore/HttpClient/Private/Test.inl
--- a/Engine/Source/Runtime/Experimental/IoStore/HttpClient/Private/Test.inl
+++ b/Engine/Source/Runtime/Experimental/IoStore/HttpClient/Private/Test.inl
@@ -436,6 +436,7 @@ static void ChunkedTest(const ANSICHAR* TestHost)
 			FResponse& Response = Status.GetResponse();
 			check(Response.GetStatus() == EStatusCodeClass::Successful);
 			check(Response.GetStatusCode() == 200);
+			check(Response.GetContentStatus() == EContentStatus::UnknownSize);
 
 			State.ExpectedHash = uint32(CrudeToInt(Response.GetHeader("X-TestServer-Hash")));
 			State.ExpectedSize = uint32(CrudeToInt(Response.GetHeader("X-TestServer-Size")));
@@ -449,6 +450,11 @@ static void ChunkedTest(const ANSICHAR* TestHost)
 
 		check(Status.GetId() == FTicketStatus::EId::Content);
 
+		check(Status.GetContent().GetSize() == Dest.GetSize());
+
+		bool IsTheEnd = (Status.GetContent().GetSize() == 0);
+		check(Status.IsComplete() == IsTheEnd);
+
 		FMemoryView View = Dest.GetView(); 
 		State.Size += uint32(View.GetSize());
 		for (uint32 i = 0, n = uint32(View.GetSize()); i < n; ++i)
@@ -571,10 +577,13 @@ static void HttpTest(const ANSICHAR* TestHost, FCertRootsRef VerifyCert)
 
 		if (Status.GetId() == FTicketStatus::EId::Response)
 		{
+			check(Status.IsComplete() == false);
+
 			FResponse& Response = Status.GetResponse();
 			check(Response.GetStatus() == EStatusCodeClass::Successful);
 			check(Response.GetStatusCode() == 200);
 			check(Response.GetContentLength() == Status.GetContentLength());
+			check(Response.GetContentStatus() == EContentStatus::KnownSize);
 
 			FAnsiStringView HashView = Response.GetHeader("X-TestServer-Hash");
 			Content[Index].Hash = CrudeToInt(HashView);
@@ -585,6 +594,8 @@ static void HttpTest(const ANSICHAR* TestHost, FCertRootsRef VerifyCert)
 			return nullptr;
 		}
 
+		check(Status.IsComplete());
+
 		uint32 ReceivedHash = 0x493;
 		FMemoryView ContentView = Content[Index].Dest.GetView();
 		check(ContentView.GetSize() == Status.GetContentLength());
@@ -688,6 +699,17 @@ static void HttpTest(const ANSICHAR* TestHost, FCertRootsRef VerifyCert)
 
 	// head barage
 	{
+		auto HeadSink = [] (const FTicketStatus& Status)
+		{
+			if (Status.GetId() != FTicketStatus::EId::Response)
+			{
+				return;
+			}
+
+			check(Status.GetResponse().GetContentStatus() == EContentStatus::NoContent);
+			check(Status.IsComplete());
+		};
+
 		FConnectionPool::FParams Params;
 		Params.SetHostFromUrl(BuildUrl());
 		Params.VerifyCert = VerifyCert;
@@ -698,7 +720,7 @@ static void HttpTest(const ANSICHAR* TestHost, FCertRootsRef VerifyCert)
 			TAnsiStringBuilder<16> Path;
 			Path << "/data?head";
 			FRequest Request = Loop.Request("HEAD", Path, Pool);
-			Loop.Send(MoveTemp(Request), NullSink);
+			Loop.Send(MoveTemp(Request), HeadSink);
 		}
 		WaitForLoopIdle();
 	}
@@ -731,6 +753,7 @@ static void HttpTest(const ANSICHAR* TestHost, FCertRootsRef VerifyCert)
 		FConnectionPool::FParams Params;
 		Params.SetHostFromUrl(BuildUrl());
 		Params.VerifyCert = VerifyCert;
+		Params.MaxInflight = 2;
 		FConnectionPool Pool(Params);
 
 		FEventLoop Loop2;
