Reset body when retrying HTTP request in on-demand HTTP client.

#rnx
#Synced-CL 48408831

diff --git a/Engine/Source/Runtime/Experimental/IoStore/OnDemand/Private/OnDemandHttpClient.cpp b/Engine/Source/Runtime/Experimental/IoStore/OnDemand/Private/OnDemandHttpClient.cpp
--- a/Engine/Source/Runtime/Experimental/IoStore/OnDemand/Private/OnDemandHttpClient.cpp
+++ b/Engine/Source/Runtime/Experimental/IoStore/OnDemand/Private/OnDemandHttpClient.cpp
@@ -613,6 +613,7 @@ void FMultiEndpointHttpClient::RetryRequest(FRequest&& Request)
 	Request.StatusCode = 0;
 	Request.RetryCount++;
 	Request.Host = Connection.CurrentHost;
+	Request.Body = FIoBuffer();
 	Retries.Emplace(MoveTemp(Request));
 }
 
