Changed log verbosity from Log -> Verbose when skipping on-demand install cache purge.

#rnx
#Synced-CL 48408831

diff --git a/Engine/Source/Runtime/Experimental/IoStore/OnDemand/Private/OnDemandInstallCache.cpp b/Engine/Source/Runtime/Experimental/IoStore/OnDemand/Private/OnDemandInstallCache.cpp
--- a/Engine/Source/Runtime/Experimental/IoStore/OnDemand/Private/OnDemandInstallCache.cpp
+++ b/Engine/Source/Runtime/Experimental/IoStore/OnDemand/Private/OnDemandInstallCache.cpp
@@ -2411,13 +2411,11 @@ FResult FOnDemandInstallCache::Purge(const uint64 BytesToInstall)
 	const uint64 TotalRequiredBytes = TotalCachedBytes + TotalUncachedBytes;
 	if (TotalRequiredBytes <= MaxCacheSize)
 	{
-		UE_LOG(LogIoStoreOnDemand, Log, TEXT("Skipping cache purge, MaxCacheSize=%.2lf MiB, CacheSize=%.2lf MiB, ReferencedBlockSize=%.2lf MiB, ReferencedSize=%.2lf MiB, FragmentedBytes=%.2lf MiB, UncachedSize=%.2lf MiB"),
+		UE_LOG(LogIoStoreOnDemand, Verbose, TEXT("Skipping cache purge, MaxCacheSize=%.2lf MiB, CacheSize=%.2lf MiB, ReferencedBlockSize=%.2lf MiB, ReferencedSize=%.2lf MiB, FragmentedBytes=%.2lf MiB, UncachedSize=%.2lf MiB"),
 			ToMiB(MaxCacheSize), ToMiB(TotalCachedBytes), ToMiB(TotalReferencedBlockBytes), ToMiB(ReferencedBytes), ToMiB(FragmentedBytes), ToMiB(TotalUncachedBytes));
 		return MakeValue();
 	}
 
-	//TODO: Compute fragmentation metric and redownload chunks when this number gets too high
-
 	UE_LOG(LogIoStoreOnDemand, Log, TEXT("Purging install cache, MaxCacheSize=%.2lf MiB, CacheSize=%.2lf MiB, ReferencedBlockSize=%.2lf MiB, ReferencedSize=%.2lf MiB, FragmentedBytes=%.2lf MiB, UncachedSize=%.2lf MiB"),
 		ToMiB(MaxCacheSize), ToMiB(TotalCachedBytes), ToMiB(TotalReferencedBlockBytes), ToMiB(ReferencedBytes), ToMiB(FragmentedBytes), ToMiB(TotalUncachedBytes));
 
