Dataflow : Add a derived UStruct type policy

#Synced-CL 48424252

diff --git a/Engine/Source/Runtime/Experimental/Dataflow/Core/Public/Dataflow/DataflowTypePolicy.h b/Engine/Source/Runtime/Experimental/Dataflow/Core/Public/Dataflow/DataflowTypePolicy.h
--- a/Engine/Source/Runtime/Experimental/Dataflow/Core/Public/Dataflow/DataflowTypePolicy.h
+++ b/Engine/Source/Runtime/Experimental/Dataflow/Core/Public/Dataflow/DataflowTypePolicy.h
@@ -307,7 +307,35 @@ struct FDataflowUStructTypePolicy : IDataflowTypePolicy
 	}
 };
 
+template <typename BaseStructType>
+struct FDataflowUDerivedStructTypePolicy : IDataflowTypePolicy
+{
+	virtual bool SupportsType(FName InType) const override
+	{
+		return SupportsTypeStatic(InType);
+	}
+
+	static bool SupportsTypeStatic(FName InType)
+	{
+		FString InnerTypeStr;
+		if (FDataflowUStructTypePolicy::GetStructType(InType.ToString(), InnerTypeStr))
+		{
+			if (UStruct* StructObject = Cast< UStruct>(StaticFindFirstObject(UStruct::StaticClass(), *InnerTypeStr, EFindFirstObjectOptions::NativeFirst)))
+			{
+				return StructObject->IsChildOf<BaseStructType>();
+			}
+		}
+
+		// not a proper Ustruct
+		return false;
+	}
 
+	static IDataflowTypePolicy* GetInterface()
+	{
+		static FDataflowUDerivedStructTypePolicy< BaseStructType> Instance;
+		return &Instance;
+	}
+};
 
 struct FDataflowNumericTypePolicy : 
 	public TDataflowMultiTypePolicy<double, float, int64, uint64, int32, uint32, int16, uint16, int8, uint8>
