[Geometry Collection] Update loose uniform buffer instead of recreating.
- also remove unused "EnableLooseParameter" code path.

#Synced-CL 48396324

diff --git a/Engine/Source/Runtime/Engine/Private/GeometryCollectionRendering.cpp b/Engine/Source/Runtime/Engine/Private/GeometryCollectionRendering.cpp
--- a/Engine/Source/Runtime/Engine/Private/GeometryCollectionRendering.cpp
+++ b/Engine/Source/Runtime/Engine/Private/GeometryCollectionRendering.cpp
@@ -150,7 +150,7 @@ void FGeometryCollectionVertexFactory::InitRHI(FRHICommandListBase& RHICmdList)
 	// then initialize PositionStream and PositionDeclaration.
 	if (Data.PositionComponent.VertexBuffer != Data.TangentBasisComponents[0].VertexBuffer)
 	{
-		auto AddDeclaration = [this, bCanUseGPUScene](EVertexInputStreamType InputStreamType, bool bAddNormal)
+		auto AddDeclaration = [this](EVertexInputStreamType InputStreamType, bool bAddNormal)
 		{
 			FVertexDeclarationElementList StreamElements;
 			StreamElements.Add(AccessStreamComponent(Data.PositionComponent, 0, InputStreamType));
@@ -271,9 +271,7 @@ void FGeometryCollectionVertexFactory::InitRHI(FRHICommandListBase& RHICmdList)
 
 		UniformParameters.VertexFetch_Parameters = { ColorIndexMask, NumTexCoords, LightMapCoordinateIndex, 0 };
 
-		EUniformBufferUsage UniformBufferUsage = EnableLooseParameter ? UniformBuffer_SingleFrame : UniformBuffer_MultiFrame;
-
-		UniformBuffer = TUniformBufferRef<FGeometryCollectionVertexFactoryUniformShaderParameters>::CreateUniformBufferImmediate(UniformParameters, UniformBufferUsage);
+		UniformBuffer = TUniformBufferRef<FGeometryCollectionVertexFactoryUniformShaderParameters>::CreateUniformBufferImmediate(UniformParameters, UniformBuffer_MultiFrame);
 		
 		FGCBoneLooseParameters LooseParameters;
 
@@ -281,7 +279,7 @@ void FGeometryCollectionVertexFactory::InitRHI(FRHICommandListBase& RHICmdList)
 		LooseParameters.VertexFetch_BonePrevTransformBuffer = GetBonePrevTransformSRV();
 		LooseParameters.VertexFetch_BoneMapBuffer = GetBoneMapSRV();
 
-		LooseParameterUniformBuffer = FGCBoneLooseParametersRef::CreateUniformBufferImmediate(LooseParameters, UniformBufferUsage);
+		LooseParameterUniformBuffer = FGCBoneLooseParametersRef::CreateUniformBufferImmediate(LooseParameters, UniformBuffer_MultiFrame);
 	}
 
 	check(IsValidRef(GetDeclaration()));
diff --git a/Engine/Source/Runtime/Engine/Public/GeometryCollectionRendering.h b/Engine/Source/Runtime/Engine/Public/GeometryCollectionRendering.h
--- a/Engine/Source/Runtime/Engine/Public/GeometryCollectionRendering.h
+++ b/Engine/Source/Runtime/Engine/Public/GeometryCollectionRendering.h
@@ -58,10 +58,9 @@ struct FGeometryCollectionVertexFactory : public FVertexFactory
 	DECLARE_VERTEX_FACTORY_TYPE_API(FGeometryCollectionVertexFactory, ENGINE_API);
 
 public:
-	FGeometryCollectionVertexFactory(ERHIFeatureLevel::Type InFeatureLevel, bool EnableLooseParameter = false)
+	FGeometryCollectionVertexFactory(ERHIFeatureLevel::Type InFeatureLevel)
 	: FVertexFactory(InFeatureLevel)
 	, LooseParameterUniformBuffer(nullptr)
-	, EnableLooseParameter(EnableLooseParameter)
 	{
 	}
 
@@ -193,7 +192,6 @@ public:
 	}
 
 	FUniformBufferRHIRef LooseParameterUniformBuffer;
-	bool EnableLooseParameter;
 
 	static ENGINE_API bool UseShaderBoneTransform(EShaderPlatform ShaderPlatform);
 
diff --git a/Engine/Source/Runtime/Experimental/GeometryCollectionEngine/Private/GeometryCollection/GeometryCollectionSceneProxy.cpp b/Engine/Source/Runtime/Experimental/GeometryCollectionEngine/Private/GeometryCollection/GeometryCollectionSceneProxy.cpp
--- a/Engine/Source/Runtime/Experimental/GeometryCollectionEngine/Private/GeometryCollection/GeometryCollectionSceneProxy.cpp
+++ b/Engine/Source/Runtime/Experimental/GeometryCollectionEngine/Private/GeometryCollection/GeometryCollectionSceneProxy.cpp
@@ -120,6 +120,7 @@ DEFINE_LOG_CATEGORY_STATIC(FGeometryCollectionSceneProxyLogging, Log, All);
 FGeometryCollectionDynamicDataPool GDynamicDataPool;
 
 static void UpdateLooseParameter(
+	FRHICommandListBase& RHICmdList,
 	FGeometryCollectionVertexFactory& VertexFactory,
 	FRHIShaderResourceView* BoneTransformSRV,
 	FRHIShaderResourceView* BonePrevTransformSRV,
@@ -131,9 +132,7 @@ static void UpdateLooseParameter(
 	LooseParameters.VertexFetch_BonePrevTransformBuffer = BonePrevTransformSRV;
 	LooseParameters.VertexFetch_BoneMapBuffer = BoneMapSRV;
 
-	EUniformBufferUsage UniformBufferUsage = VertexFactory.EnableLooseParameter ? UniformBuffer_SingleFrame : UniformBuffer_MultiFrame;
-
-	VertexFactory.LooseParameterUniformBuffer = FGCBoneLooseParametersRef::CreateUniformBufferImmediate(LooseParameters, UniformBufferUsage);
+	RHICmdList.UpdateUniformBuffer(VertexFactory.LooseParameterUniformBuffer, &LooseParameters);
 }
 
 FGeometryCollectionSceneProxyBase::FGeometryCollectionSceneProxyBase(UGeometryCollectionComponent* Component, bool bInIsNanite)
@@ -346,7 +345,7 @@ void FGeometryCollectionSceneProxyBase::SetDynamicData_RenderThread(FRHICommandL
 
 				TransformBuffer.UpdateDynamicData(RHICmdList, DynamicData->Transforms, LockMode);
 
-				UpdateLooseParameter(VertexFactory, TransformBuffer.VertexBufferSRV, PrevTransformBuffer.VertexBufferSRV, MeshResource.BoneMapVertexBuffer.GetSRV());
+				UpdateLooseParameter(RHICmdList, VertexFactory, TransformBuffer.VertexBufferSRV, PrevTransformBuffer.VertexBufferSRV, MeshResource.BoneMapVertexBuffer.GetSRV());
 			}
 		}
 		else
@@ -875,7 +874,7 @@ void FGeometryCollectionSceneProxy::SetDynamicData_RenderThread(FRHICommandListB
 		{
 			VertexFactoryDebugColor.SetBoneTransformSRV(TransformBuffer.VertexBufferSRV);
 			VertexFactoryDebugColor.SetBonePrevTransformSRV(PrevTransformBuffer.VertexBufferSRV);
-			UpdateLooseParameter(VertexFactoryDebugColor, TransformBuffer.VertexBufferSRV, PrevTransformBuffer.VertexBufferSRV, MeshResource.BoneMapVertexBuffer.GetSRV());
+			UpdateLooseParameter(RHICmdList, VertexFactoryDebugColor, TransformBuffer.VertexBufferSRV, PrevTransformBuffer.VertexBufferSRV, MeshResource.BoneMapVertexBuffer.GetSRV());
 		}
 #endif
 	}
