Clean up out of date code api documentation.

#rnx

#Synced-CL 48421160

diff --git a/Engine/Source/Runtime/Experimental/IoStore/OnDemandCore/Public/IO/IoStoreOnDemand.h b/Engine/Source/Runtime/Experimental/IoStore/OnDemandCore/Public/IO/IoStoreOnDemand.h
--- a/Engine/Source/Runtime/Experimental/IoStore/OnDemandCore/Public/IO/IoStoreOnDemand.h
+++ b/Engine/Source/Runtime/Experimental/IoStore/OnDemandCore/Public/IO/IoStoreOnDemand.h
@@ -634,7 +634,8 @@ public:
 	 * @param	Args				The install arguments.
 	 * @param	OnCompleted			Completion callback.
 	 * @param	OnProgress			Optional progress callback.
-	 * @param	CancellationToken	Optional cancellation token.
+	 * 
+	 * @return A handle representing the install request. @see FOnDemandInstallRequest
 	 */
 	virtual FOnDemandInstallRequest Install(
 			FOnDemandInstallArgs&& Args,
