Remove the default constructor for FIoHttpResponse.

#rnx

- The default constructor was not being used but would create a FIoHttpResponse object with unitialized variables.

#Synced-CL 48421160

diff --git a/Engine/Source/Runtime/Experimental/IoStore/OnDemandCore/Public/IO/HttpIoDispatcher.h b/Engine/Source/Runtime/Experimental/IoStore/OnDemandCore/Public/IO/HttpIoDispatcher.h
--- a/Engine/Source/Runtime/Experimental/IoStore/OnDemandCore/Public/IO/HttpIoDispatcher.h
+++ b/Engine/Source/Runtime/Experimental/IoStore/OnDemandCore/Public/IO/HttpIoDispatcher.h
@@ -268,7 +268,9 @@ ENUM_CLASS_FLAGS(EIoHttpResponseFlags);
 class FIoHttpResponse
 {
 public:
-							FIoHttpResponse() = default;
+							FIoHttpResponse() = delete;
+							~FIoHttpResponse() = default;
+
 	inline					FIoHttpResponse(const FIoHash& InCacheKey, const FIoBuffer& InBody, EIoErrorCode InErrorCode, uint32 InStatusCode, EIoHttpResponseFlags InFlags);
 	inline					FIoHttpResponse(const FIoHash& InCacheKey, FIoHttpHeaders&& InHeaders, const FIoBuffer& InBody, EIoErrorCode InErrorCode, uint32 InStatusCode, EIoHttpResponseFlags InFlags);
 							FIoHttpResponse(EIoErrorCode InErrorCode, uint32 InStatusCode)
