[Geometry Collection] Support cached MDCs.
- In test level with 300 geometry collections avoids 0.65ms spread over multiple GatherDynamicMeshElements tasks.

#Synced-CL 48477729

diff --git a/Engine/Source/Runtime/Experimental/GeometryCollectionEngine/Private/GeometryCollection/GeometryCollectionSceneProxy.cpp b/Engine/Source/Runtime/Experimental/GeometryCollectionEngine/Private/GeometryCollection/GeometryCollectionSceneProxy.cpp
--- a/Engine/Source/Runtime/Experimental/GeometryCollectionEngine/Private/GeometryCollection/GeometryCollectionSceneProxy.cpp
+++ b/Engine/Source/Runtime/Experimental/GeometryCollectionEngine/Private/GeometryCollection/GeometryCollectionSceneProxy.cpp
@@ -54,6 +54,13 @@ static_assert(sizeof(ispc::FMatrix44f) == sizeof(FMatrix44f), "sizeof(ispc::FMat
 static_assert(sizeof(ispc::FVector3f) == sizeof(FVector3f), "sizeof(ispc::FVector3f) != sizeof(FVector3f)");
 #endif
 
+static bool GGeometryCollectionUseCachedMDCs = true;
+static FAutoConsoleVariableRef CVarGeometryCollectionUseCachedMDCs(
+	TEXT("r.GeometryCollection.UseCachedMDCs"),
+	GGeometryCollectionUseCachedMDCs,
+	TEXT("Whether geometry collections will take the cached MDC path."),
+	ECVF_RenderThreadSafe);
+
 static int32 GParallelGeometryCollectionBatchSize = 1024;
 static TAutoConsoleVariable<int32> CVarParallelGeometryCollectionBatchSize(
 	TEXT("r.ParallelGeometryCollectionBatchSize"),
@@ -435,6 +442,61 @@ void FGeometryCollectionSceneProxyBase::UpdateSkinnedPositions(FRHICommandListBa
 	RHICmdList.UnlockBuffer(SkinnedPositionVertexBuffer.VertexBufferRHI);
 }
 
+TConstArrayView<FGeometryCollectionMeshElement> FGeometryCollectionSceneProxyBase::GetSectionArray(bool bUsesSubSections, bool bRemoveInternalFaces) const
+{
+	return bUsesSubSections
+		? MeshDescription.SubSections
+		: bRemoveInternalFaces ? MeshDescription.SectionsNoInternal : MeshDescription.Sections;
+}
+
+void FGeometryCollectionSceneProxyBase::GetMeshElement(
+	int32 LODIndex,
+	int32 SectionIndex,
+	const FGeometryCollectionMeshElement& Section,
+	const FVertexFactory* InVertexFactory,
+	FMaterialRenderProxy* MaterialProxy,
+	uint8 InDepthPriorityGroup,
+	FRHIUniformBuffer* UniformBuffer,
+	bool bWireframe,
+	bool bReverseCulling,
+	FMeshBatch& OutMeshBatch) const
+{
+	OutMeshBatch.bWireframe = bWireframe;
+	OutMeshBatch.SegmentIndex = SectionIndex;
+	OutMeshBatch.VertexFactory = InVertexFactory;
+	OutMeshBatch.MaterialRenderProxy = MaterialProxy;
+	OutMeshBatch.ReverseCulling = bReverseCulling;
+	OutMeshBatch.LODIndex = LODIndex;
+	OutMeshBatch.Type = PT_TriangleList;
+	OutMeshBatch.DepthPriorityGroup = InDepthPriorityGroup;
+
+	FMeshBatchElement& BatchElement = OutMeshBatch.Elements[0];
+	BatchElement.IndexBuffer = &MeshResource.IndexBuffer;
+	BatchElement.PrimitiveUniformBuffer = UniformBuffer;
+	BatchElement.FirstIndex = Section.TriangleStart * 3;
+	BatchElement.NumPrimitives = Section.TriangleCount;
+	BatchElement.MinVertexIndex = Section.VertexStart;
+	BatchElement.MaxVertexIndex = Section.VertexEnd;
+}
+
+void FGeometryCollectionSceneProxyBase::GetSectionMaterialProxies(TConstArrayView<FGeometryCollectionMeshElement> InSectionArray, FMaterialArrayType& OutMaterialProxies) const
+{
+	// Grab the material proxies we'll be using for each section
+	for (int32 SectionIndex = 0; SectionIndex < InSectionArray.Num(); ++SectionIndex)
+	{
+		const FGeometryCollectionMeshElement& Section = InSectionArray[SectionIndex];
+
+		FMaterialRenderProxy* MaterialProxy = Materials[Section.MaterialIndex]->GetRenderProxy();
+
+		if (MaterialProxy == nullptr)
+		{
+			MaterialProxy = UMaterial::GetDefaultMaterial(MD_Surface)->GetRenderProxy();
+		}
+
+		OutMaterialProxies.Add(MaterialProxy);
+	}
+}
+
 #if RHI_RAYTRACING
 void FGeometryCollectionSceneProxyBase::UpdatingRayTracingGeometry_RenderingThread(FRHICommandListBase& RHICmdList, TConstArrayView<FGeometryCollectionMeshElement> InSectionArray, bool bRayTracingGeometryPerSection)
 {
@@ -519,30 +581,18 @@ void FGeometryCollectionSceneProxyBase::GetDynamicRayTracingInstances(FRayTracin
 
 	const uint32 LODIndex = 0;
 	const bool bWireframe = false;
+	const bool bReverseCulling = false; // DXR doesn't need the flipping implied by the transform
 
 	const bool bUseSubSections = MeshDescription.SubSections.Num() && !bAnyMaterialHasWorldPositionOffset && CVarRayTracingGeometryCollectionCombinedBLAS.GetValueOnRenderThread() == 0;
 
-	TArray<FGeometryCollectionMeshElement> const& SectionArray = bUseSubSections ? MeshDescription.SubSections : MeshDescription.Sections;
+	TConstArrayView<FGeometryCollectionMeshElement> SectionArray = GetSectionArray(bUseSubSections, false);
 
 	UpdatingRayTracingGeometry_RenderingThread(Collector.GetRHICommandList(), SectionArray, bUseSubSections);
 
 	// Grab the material proxies we'll be using for each section
-	TArray<FMaterialRenderProxy*, TInlineAllocator<32>> MaterialProxies;
-
-	for (int32 SectionIndex = 0; SectionIndex < SectionArray.Num(); ++SectionIndex)
-	{
-		const FGeometryCollectionMeshElement& Section = SectionArray[SectionIndex];
-
-		//TODO: Add BoneColor support in Path/Ray tracing?
-		FMaterialRenderProxy* MaterialProxy = Materials[Section.MaterialIndex]->GetRenderProxy();
-
-		if (MaterialProxy == nullptr)
-		{
-			MaterialProxy = UMaterial::GetDefaultMaterial(MD_Surface)->GetRenderProxy();
-		}
-
-		MaterialProxies.Add(MaterialProxy);
-	}
+	// TODO: Add BoneColor support in Path/Ray tracing?
+	FMaterialArrayType MaterialProxies;
+	GetSectionMaterialProxies(SectionArray, MaterialProxies);
 
 	if (RayTracingGeometry.IsValid())
 	{
@@ -564,26 +614,9 @@ void FGeometryCollectionSceneProxyBase::GetDynamicRayTracingInstances(FRayTracin
 			const FGeometryCollectionMeshElement& Section = SectionArray[SectionIndex];
 
 			FMeshBatch& Mesh = RayTracingInstance.Materials.AddDefaulted_GetRef();
-			Mesh.bWireframe = bWireframe;
-			Mesh.SegmentIndex = SectionIndex;
-			Mesh.VertexFactory = &VertexFactory;
-			Mesh.MaterialRenderProxy = MaterialProxies[SectionIndex];
-			Mesh.LODIndex = LODIndex;
-			Mesh.bDisableBackfaceCulling = true;
-			Mesh.Type = PT_TriangleList;
-			Mesh.DepthPriorityGroup = SDPG_World;
-			Mesh.bCanApplyViewModeOverrides = true;
-
-			FMeshBatchElement& BatchElement = Mesh.Elements[0];
-			BatchElement.IndexBuffer = &MeshResource.IndexBuffer;
-			BatchElement.PrimitiveUniformBuffer = UniformBuffer;
-			BatchElement.FirstIndex = Section.TriangleStart * 3;
-			BatchElement.NumPrimitives = Section.TriangleCount;
-			BatchElement.MinVertexIndex = Section.VertexStart;
-			BatchElement.MaxVertexIndex = Section.VertexEnd;
-			BatchElement.NumInstances = 1;
-
-			MaxVertexIndex = std::max(Section.VertexEnd, MaxVertexIndex);
+			GetMeshElement(LODIndex, SectionIndex, Section, &VertexFactory, MaterialProxies[SectionIndex], SDPG_World, UniformBuffer, bWireframe, bReverseCulling, Mesh);
+
+			MaxVertexIndex = std::max(Mesh.Elements[0].MaxVertexIndex, MaxVertexIndex);
 
 			//#TODO: bone color, bone selection and render bound?
 		}
@@ -654,24 +687,7 @@ void FGeometryCollectionSceneProxyBase::GetDynamicRayTracingInstances(FRayTracin
 				const FGeometryCollectionMeshElement& Section = SectionArray[SectionIndex];
 
 				FMeshBatch& Mesh = CachedPartRayTracingMaterials.AddDefaulted_GetRef();
-				Mesh.bWireframe = bWireframe;
-				Mesh.SegmentIndex = 0;
-				Mesh.VertexFactory = &VertexFactory;
-				Mesh.MaterialRenderProxy = MaterialProxies[SectionIndex];
-				Mesh.LODIndex = LODIndex;
-				Mesh.bDisableBackfaceCulling = true;
-				Mesh.Type = PT_TriangleList;
-				Mesh.DepthPriorityGroup = SDPG_World;
-				Mesh.bCanApplyViewModeOverrides = true;
-
-				FMeshBatchElement& BatchElement = Mesh.Elements[0];
-				BatchElement.IndexBuffer = &MeshResource.IndexBuffer;
-				BatchElement.PrimitiveUniformBuffer = UniformBuffer;
-				BatchElement.FirstIndex = Section.TriangleStart * 3;
-				BatchElement.NumPrimitives = Section.TriangleCount;
-				BatchElement.MinVertexIndex = Section.VertexStart;
-				BatchElement.MaxVertexIndex = Section.VertexEnd;
-				BatchElement.NumInstances = 1;
+				GetMeshElement(LODIndex, 0, Section, &VertexFactory, MaterialProxies[SectionIndex], SDPG_World, UniformBuffer, bWireframe, bReverseCulling, Mesh);
 
 				// #TODO: bone color, bone selection and render bound?
 			}
@@ -954,6 +970,9 @@ void FGeometryCollectionSceneProxy::GetDynamicMeshElements(const TArray<const FS
 	const bool bDrawOnlyCollisionMeshes = EngineShowFlags.CollisionPawn || EngineShowFlags.CollisionVisibility;
 	const bool bDrawWireframeCollision = EngineShowFlags.Collision && IsCollisionEnabled();
 
+	const uint32 LODIndex = 0;
+	const bool bReverseCulling = IsLocalToWorldDeterminantNegative();
+
 	auto SetDebugMaterial = [this, &Collector, &EngineShowFlags, bProxyIsSelected](FMeshBatch& Mesh) -> void
 	{
 #if UE_ENABLE_DEBUG_DRAWING
@@ -980,6 +999,16 @@ void FGeometryCollectionSceneProxy::GetDynamicMeshElements(const TArray<const FS
 
 	if (bDrawGeometryCollectionMesh)
 	{
+		// If hiding geometry in editor then we don't remove hidden faces.
+		const bool bRemoveInternalFaces = false;
+
+#if GEOMETRYCOLLECTION_EDITOR_SELECTION
+		// If using subsections then use the subsection array. 
+		TConstArrayView<FGeometryCollectionMeshElement> SectionArray = GetSectionArray(bUsesSubSections, bRemoveInternalFaces);
+#else
+		TConstArrayView<FGeometryCollectionMeshElement> SectionArray = GetSectionArray(false, bRemoveInternalFaces);
+#endif
+
 		for (int32 ViewIndex = 0; ViewIndex < Views.Num(); ViewIndex++)
 		{
 			if ((VisibilityMap & (1 << ViewIndex)) == 0)
@@ -987,20 +1016,8 @@ void FGeometryCollectionSceneProxy::GetDynamicMeshElements(const TArray<const FS
 				continue;
 			}
 
-			// If hiding geometry in editor then we don't remove hidden faces.
-			const bool bRemoveInternalFaces = false;
-
-#if GEOMETRYCOLLECTION_EDITOR_SELECTION
-			// If using subsections then use the subsection array. 
-			TArray<FGeometryCollectionMeshElement> const& SectionArray = bUsesSubSections
-				? MeshDescription.SubSections
-				: bRemoveInternalFaces ? MeshDescription.SectionsNoInternal : MeshDescription.Sections;
-#else
-			TArray<FGeometryCollectionMeshElement> const& SectionArray = bRemoveInternalFaces ? MeshDescription.SectionsNoInternal : MeshDescription.Sections;
-#endif
-
 			// Grab the material proxies we'll be using for each section.
-			TArray<FMaterialRenderProxy*, TInlineAllocator<32>> MaterialProxies;
+			FMaterialArrayType MaterialProxies;
 			for (int32 SectionIndex = 0; SectionIndex < SectionArray.Num(); ++SectionIndex)
 			{
 				const FGeometryCollectionMeshElement& Section = SectionArray[SectionIndex];
@@ -1014,23 +1031,10 @@ void FGeometryCollectionSceneProxy::GetDynamicMeshElements(const TArray<const FS
 				const FGeometryCollectionMeshElement& Section = SectionArray[SectionIndex];
 
 				FMeshBatch& Mesh = Collector.AllocateMesh();
-				Mesh.bWireframe = bWireframe;
-				Mesh.VertexFactory = GetVertexFactory();
-				Mesh.MaterialRenderProxy = MaterialProxies[SectionIndex];
-				Mesh.ReverseCulling = IsLocalToWorldDeterminantNegative();
-				Mesh.Type = PT_TriangleList;
-				Mesh.DepthPriorityGroup = SDPG_World;
+				GetMeshElement(LODIndex, SectionIndex, Section, GetVertexFactory(), MaterialProxies[SectionIndex], SDPG_World, GetUniformBuffer(), bWireframe, bReverseCulling, Mesh);
 				Mesh.bCanApplyViewModeOverrides = true;
 				SetDebugMaterial(Mesh);
 
-				FMeshBatchElement& BatchElement = Mesh.Elements[0];
-				BatchElement.IndexBuffer = &MeshResource.IndexBuffer;
-				BatchElement.PrimitiveUniformBuffer = GetUniformBuffer();
-				BatchElement.FirstIndex = Section.TriangleStart * 3;
-				BatchElement.NumPrimitives = Section.TriangleCount;
-				BatchElement.MinVertexIndex = Section.VertexStart;
-				BatchElement.MaxVertexIndex = Section.VertexEnd;
-
 				Collector.AddMesh(ViewIndex, Mesh);
 			}
 
@@ -1095,10 +1099,37 @@ void FGeometryCollectionSceneProxy::GetDynamicMeshElements(const TArray<const FS
 
 FPrimitiveViewRelevance FGeometryCollectionSceneProxy::GetViewRelevance(const FSceneView* View) const
 {
+	const FEngineShowFlags& EngineShowFlags = View->Family->EngineShowFlags;
+	const bool bWireframe = AllowDebugViewmodes() && EngineShowFlags.Wireframe;
+	const bool bProxyIsSelected = IsSelected();
+	const bool bDrawOnlyCollisionMeshes = EngineShowFlags.CollisionPawn || EngineShowFlags.CollisionVisibility;
+	const bool bDrawWireframeCollision = EngineShowFlags.Collision && IsCollisionEnabled();
+
+	const bool bDynamic = !GGeometryCollectionUseCachedMDCs
+		|| HasViewDependentDPG()
+#if WITH_EDITOR
+		|| bShowBoneColors
+#endif
+#if GEOMETRYCOLLECTION_EDITOR_SELECTION
+		|| bEnableBoneSelection
+#endif
+		|| bWireframe
+		|| bProxyIsSelected
+		|| bDrawOnlyCollisionMeshes
+		|| bDrawWireframeCollision
+		|| View->Family->EngineShowFlags.Bounds;
+
 	FPrimitiveViewRelevance Result;
 	Result.bDrawRelevance = IsShown(View);
 	Result.bShadowRelevance = IsShadowCast(View);
-	Result.bDynamicRelevance = true;
+	if (bDynamic)
+	{
+		Result.bDynamicRelevance = true;
+	}
+	else
+	{
+		Result.bStaticRelevance = true;
+	}
 	Result.bRenderInMainPass = ShouldRenderInMainPass();
 	Result.bUsesLightingChannels = GetLightingChannelMask() != GetDefaultLightingChannelMask();
 	Result.bRenderCustomDepth = ShouldRenderCustomDepth();
@@ -1119,6 +1150,48 @@ HHitProxy* FGeometryCollectionSceneProxy::CreateHitProxies(UPrimitiveComponent*
 }
 #endif
 
+void FGeometryCollectionSceneProxy::DrawStaticElements(FStaticPrimitiveDrawInterface* PDI)
+{
+	checkSlow(IsInParallelRenderingThread());
+
+	if (HasViewDependentDPG())
+	{
+		return;
+	}
+
+	// Determine the DPG the primitive should be drawn in.
+	ESceneDepthPriorityGroup PrimitiveDPG = GetStaticDepthPriorityGroup();
+
+	const uint32 LODIndex = 0;
+	const bool bReverseCulling = IsLocalToWorldDeterminantNegative();
+	const bool bWireframe = false;
+
+	// If hiding geometry in editor then we don't remove hidden faces.
+	const bool bRemoveInternalFaces = false;
+
+#if GEOMETRYCOLLECTION_EDITOR_SELECTION
+	// If using subsections then use the subsection array. 
+	TConstArrayView<FGeometryCollectionMeshElement> SectionArray = GetSectionArray(bUsesSubSections, bRemoveInternalFaces);
+#else
+	TConstArrayView<FGeometryCollectionMeshElement> SectionArray = GetSectionArray(false, bRemoveInternalFaces);
+#endif
+
+	// Grab the material proxies we'll be using for each section.
+	FMaterialArrayType MaterialProxies;
+	GetSectionMaterialProxies(SectionArray, MaterialProxies);
+
+	// Draw the meshes.
+	for (int32 SectionIndex = 0; SectionIndex < SectionArray.Num(); ++SectionIndex)
+	{
+		const FGeometryCollectionMeshElement& Section = SectionArray[SectionIndex];
+
+		FMeshBatch Mesh;
+		GetMeshElement(LODIndex, SectionIndex, Section, GetVertexFactory(), MaterialProxies[SectionIndex], PrimitiveDPG, GetUniformBuffer(), bWireframe, bReverseCulling, Mesh);
+
+		PDI->DrawMesh(Mesh, 1.0f);
+	}
+}
+
 void FGeometryCollectionSceneProxy::GetPreSkinnedLocalBounds(FBoxSphereBounds& OutBounds) const
 {
 	OutBounds = PreSkinnedBounds; 
@@ -1375,8 +1448,6 @@ FPrimitiveViewRelevance FNaniteGeometryCollectionSceneProxy::GetViewRelevance(co
 	Result.bEditorStaticSelectionRelevance = (IsSelected() || IsHovered());
 #endif
 
-	bool bSetDynamicRelevance = false;
-
 	Result.bOpaque = true;
 
 	MaterialRelevance.SetPrimitiveViewRelevance(Result);
diff --git a/Engine/Source/Runtime/Experimental/GeometryCollectionEngine/Private/GeometryCollection/GeometryCollectionSceneProxy.h b/Engine/Source/Runtime/Experimental/GeometryCollectionEngine/Private/GeometryCollection/GeometryCollectionSceneProxy.h
--- a/Engine/Source/Runtime/Experimental/GeometryCollectionEngine/Private/GeometryCollection/GeometryCollectionSceneProxy.h
+++ b/Engine/Source/Runtime/Experimental/GeometryCollectionEngine/Private/GeometryCollection/GeometryCollectionSceneProxy.h
@@ -21,6 +21,7 @@
 class UGeometryCollection;
 class UGeometryCollectionComponent;
 struct FGeometryCollectionSection;
+struct FGeometryCollectionMeshElement;
 
 namespace Nanite
 {
@@ -172,6 +173,25 @@ protected:
 		}
 	}
 
+	TConstArrayView<FGeometryCollectionMeshElement> GetSectionArray(bool bUsesSubSections, bool bRemoveInternalFaces) const;
+
+	/** Sets up a FMeshBatch for a specific LOD and section. */
+	void GetMeshElement(
+		int32 LODIndex,
+		int32 SectionIndex,
+		const FGeometryCollectionMeshElement& Section,
+		const FVertexFactory* InVertexFactory,
+		FMaterialRenderProxy* MaterialProxy,
+		uint8 InDepthPriorityGroup,
+		FRHIUniformBuffer* UniformBuffer,
+		bool bWireframe,
+		bool bReverseCulling,
+		FMeshBatch& OutMeshBatch) const;
+
+	using FMaterialArrayType = TArray<FMaterialRenderProxy*, TInlineAllocator<32>>;
+
+	void GetSectionMaterialProxies(TConstArrayView<FGeometryCollectionMeshElement> InSectionArray, FMaterialArrayType& OutMaterialProxies) const;
+
 #if RHI_RAYTRACING
 	void UpdatingRayTracingGeometry_RenderingThread(FRHICommandListBase& RHICmdList, TConstArrayView<FGeometryCollectionMeshElement> InSectionArray, bool bRayTracingGeometryPerSection);
 
@@ -289,6 +309,7 @@ public:
 		return (bEnableBoneSelection || bUsesSubSections) ? &HitProxyIdBuffer : nullptr;
  	}
 #endif
+	virtual void DrawStaticElements(FStaticPrimitiveDrawInterface* PDI) override;
 
 #if RHI_RAYTRACING
 	bool IsRayTracingRelevant() const override { return FGeometryCollectionSceneProxyBase::IsRayTracingRelevant(); }
diff --git a/Engine/Source/Runtime/Renderer/Private/SceneVisibility.cpp b/Engine/Source/Runtime/Renderer/Private/SceneVisibility.cpp
--- a/Engine/Source/Runtime/Renderer/Private/SceneVisibility.cpp
+++ b/Engine/Source/Runtime/Renderer/Private/SceneVisibility.cpp
@@ -4103,6 +4103,8 @@ FGraphEventRef FDynamicMeshElementContext::LaunchRenderThreadTask(FDynamicPrimit
 {
 	return FFunctionGraphTask::CreateAndDispatchWhenReady([this, PrimitiveIndexList = MoveTemp(PrimitiveIndexList)]
 	{
+		SCOPED_NAMED_EVENT(GatherDynamicMeshElements, FColor::Magenta);
+
 		for (FDynamicPrimitiveIndex PrimitiveIndex : PrimitiveIndexList.Primitives)
 		{
 			GatherDynamicMeshElementsForPrimitive(Primitives[PrimitiveIndex.Index], PrimitiveIndex.ViewMask);
@@ -4122,6 +4124,9 @@ UE::Tasks::FTask FDynamicMeshElementContext::LaunchAsyncTask(FDynamicPrimitiveIn
 	return Pipe.Launch(UE_SOURCE_LOCATION, [this, PrimitiveIndexQueue]
 	{
 		FTaskTagScope Scope(ETaskTag::EParallelRenderingThread);
+
+		SCOPED_NAMED_EVENT(GatherDynamicMeshElements, FColor::Magenta);
+
 		FDynamicPrimitiveIndex PrimitiveIndex;
 
 		while (PrimitiveIndexQueue->Pop(PrimitiveIndex))
@@ -4463,7 +4468,7 @@ void FVisibilityTaskData::LaunchVisibilityTasks(const UE::Tasks::FTask& BeginIni
 		if (ViewPackets.Num() == 1)
 		{
 			// When using a single view, dynamic mesh elements are pushed into a pipe that is executed on the render thread which allows for some overlap with compute relevance work.
-			DynamicMeshElements.CommandPipe = Allocator.Create<TCommandPipe<FDynamicPrimitiveIndexList>>(TEXT("GatherDynamicMeshElements"));
+			DynamicMeshElements.CommandPipe = Allocator.Create<TCommandPipe<FDynamicPrimitiveIndexList>>(TEXT("GatherDynamicMeshElementsPipe"));
 
 			DynamicMeshElements.CommandPipe->SetCommandFunction([this](FDynamicPrimitiveIndexList&& DynamicPrimitiveIndexList)
 			{
