Collision Filter Refactor Part 5.3: Removing the seperate sim/query block words. This is needed to extend to 64 bit and stay within our memory budget.
Several spots were also changed from calling BuildSimBlockAll to BuildBlockAll. These spots were updated appropriately to disable query on the shapes.

#Synced-CL 48505955

diff --git a/Engine/Plugins/Experimental/ControlRigPhysics/Source/ControlRigPhysics/Private/RigPhysicsSimulation.cpp b/Engine/Plugins/Experimental/ControlRigPhysics/Source/ControlRigPhysics/Private/RigPhysicsSimulation.cpp
--- a/Engine/Plugins/Experimental/ControlRigPhysics/Source/ControlRigPhysics/Private/RigPhysicsSimulation.cpp
+++ b/Engine/Plugins/Experimental/ControlRigPhysics/Source/ControlRigPhysics/Private/RigPhysicsSimulation.cpp
@@ -405,7 +405,7 @@ static bool CreateGeometry(
 	// pairs that are explicitly set to collide)
 	FBodyCollisionData BodyCollisionData;
 
-	const Chaos::Filter::FShapeFilterData ShapeFilter = Chaos::Filter::FShapeFilterBuilder::BuildSimBlockAll(Chaos::EFilterFlags::All);
+	const Chaos::Filter::FShapeFilterData ShapeFilter = Chaos::Filter::FShapeFilterBuilder::BuildBlockAll(Chaos::EFilterFlags::All);
 	BodyCollisionData.CollisionFilterData.ComplexShapeFilterData = BodyCollisionData.CollisionFilterData.SimpleShapeFilterData = ShapeFilter;
 	BodyCollisionData.CollisionFilterData.FilterInstanceData = Chaos::Filter::FInstanceData();
 
diff --git a/Engine/Source/Programs/HeadlessChaos/Private/HeadlessChaosTestSimSweep.cpp b/Engine/Source/Programs/HeadlessChaos/Private/HeadlessChaosTestSimSweep.cpp
--- a/Engine/Source/Programs/HeadlessChaos/Private/HeadlessChaosTestSimSweep.cpp
+++ b/Engine/Source/Programs/HeadlessChaos/Private/HeadlessChaosTestSimSweep.cpp
@@ -133,11 +133,11 @@ namespace ChaosTest
 		Test.GetEvolution().SetParticleObjectState(Box, EObjectStateType::Kinematic);
 
 		// Floor -> CollisionChannel 0
-		Filter::FShapeFilterData FloorFilterData = Filter::FShapeFilterBuilder().SetCollisionChannelIndex(0).SetBlockChannelMasks(0xFFFFFFFF, 0).Build();
+		Filter::FShapeFilterData FloorFilterData = Filter::FShapeFilterBuilder().SetCollisionChannelIndex(0).SetBlockChannelMask(0xFFFFFFFF).Build();
 		Floor->ShapeInstances()[0]->SetShapeFilterData(FloorFilterData);
 
 		// Box -> CollisionChannel 1
-		Filter::FShapeFilterData BoxFilterData = Filter::FShapeFilterBuilder().SetCollisionChannelIndex(1).SetBlockChannelMasks(0xFFFFFFFF, 0).Build();
+		Filter::FShapeFilterData BoxFilterData = Filter::FShapeFilterBuilder().SetCollisionChannelIndex(1).SetBlockChannelMask(0xFFFFFFFF).Build();
 		Box->ShapeInstances()[0]->SetShapeFilterData(BoxFilterData);
 
 		Test.Tick();
diff --git a/Engine/Source/Runtime/Engine/Private/PhysicsEngine/ImmediatePhysics/ImmediatePhysicsAdapters.cpp b/Engine/Source/Runtime/Engine/Private/PhysicsEngine/ImmediatePhysics/ImmediatePhysicsAdapters.cpp
--- a/Engine/Source/Runtime/Engine/Private/PhysicsEngine/ImmediatePhysics/ImmediatePhysicsAdapters.cpp
+++ b/Engine/Source/Runtime/Engine/Private/PhysicsEngine/ImmediatePhysics/ImmediatePhysicsAdapters.cpp
@@ -124,10 +124,11 @@ namespace ImmediatePhysics
 		// Set the filter to collide with everything (we use a broad phase that only contains particle pairs that are explicitly set to collide)
 		FBodyCollisionData BodyCollisionData;
 
-		const Chaos::Filter::FShapeFilterData ShapeFilterData = Chaos::Filter::FShapeFilterBuilder::BuildSimBlockAll(EFilterFlags::All);
+		const Chaos::Filter::FShapeFilterData ShapeFilterData = Chaos::Filter::FShapeFilterBuilder::BuildBlockAll(EFilterFlags::All);
 		BodyCollisionData.CollisionFilterData.SimpleShapeFilterData = ShapeFilterData;
 		BodyCollisionData.CollisionFilterData.ComplexShapeFilterData = ShapeFilterData;
 		BodyCollisionData.CollisionFlags.bEnableSimCollisionSimple = true;
+		BodyCollisionData.CollisionFlags.bEnableQueryCollision = false;
 
 		BodyInstance->BuildBodyCollisionFlags(BodyCollisionData.CollisionFlags, BodyInstance->GetCollisionEnabled(), BodyInstance->BodySetup->GetCollisionTraceFlag() == CTF_UseComplexAsSimple);
 
diff --git a/Engine/Source/Runtime/Engine/Private/PhysicsEngine/ImmediatePhysics/ImmediatePhysicsChaos/ImmediatePhysicsActorHandle_Chaos.cpp b/Engine/Source/Runtime/Engine/Private/PhysicsEngine/ImmediatePhysics/ImmediatePhysicsChaos/ImmediatePhysicsActorHandle_Chaos.cpp
--- a/Engine/Source/Runtime/Engine/Private/PhysicsEngine/ImmediatePhysics/ImmediatePhysicsChaos/ImmediatePhysicsActorHandle_Chaos.cpp
+++ b/Engine/Source/Runtime/Engine/Private/PhysicsEngine/ImmediatePhysics/ImmediatePhysicsChaos/ImmediatePhysicsActorHandle_Chaos.cpp
@@ -88,11 +88,12 @@ namespace ImmediatePhysics_Chaos
 			// Set the collision filter data for the shapes to collide with everything.
 			// Even though we already tried to do this when we created the original shapes array, 
 			// that gets thrown away and we need to do it here. This is not a good API
-			const Chaos::Filter::FShapeFilterData ShapeFilter = Chaos::Filter::FShapeFilterBuilder::BuildSimBlockAll(Chaos::EFilterFlags::All);
+			const Chaos::Filter::FShapeFilterData ShapeFilter = Chaos::Filter::FShapeFilterBuilder::BuildBlockAll(Chaos::EFilterFlags::All);
 
 			FCollisionData CollisionData;
 			CollisionData.SetShapeFilterData(ShapeFilter);
 			CollisionData.bSimCollision = 1;
+			CollisionData.bQueryCollision = 0;
 			for (const TUniquePtr<FPerShapeData>& Shape : ParticleHandle->ShapesArray())
 			{
 				Shape->SetCollisionData(CollisionData);
diff --git a/Engine/Source/Runtime/Experimental/Chaos/Private/Chaos/Collision/SimRaycast.cpp b/Engine/Source/Runtime/Experimental/Chaos/Private/Chaos/Collision/SimRaycast.cpp
--- a/Engine/Source/Runtime/Experimental/Chaos/Private/Chaos/Collision/SimRaycast.cpp
+++ b/Engine/Source/Runtime/Experimental/Chaos/Private/Chaos/Collision/SimRaycast.cpp
@@ -108,7 +108,7 @@ namespace Chaos::Private
 
 		// Run the channel filter
 		// NOTE: We return true if we want to keep the result
-		const Filter::ENarrowFilterResult HitType = FilterData.NarrowFilter(Shape->GetShapeFilterData(), true);
+		const Filter::ENarrowFilterResult HitType = FilterData.NarrowFilter(Shape->GetShapeFilterData());
 		return HitType != Filter::ENarrowFilterResult::None;
 	}
 
diff --git a/Engine/Source/Runtime/Experimental/Chaos/Private/Chaos/CollisionFilterData.cpp b/Engine/Source/Runtime/Experimental/Chaos/Private/Chaos/CollisionFilterData.cpp
--- a/Engine/Source/Runtime/Experimental/Chaos/Private/Chaos/CollisionFilterData.cpp
+++ b/Engine/Source/Runtime/Experimental/Chaos/Private/Chaos/CollisionFilterData.cpp
@@ -167,12 +167,12 @@ namespace Chaos::Filter
 
 	bool FShapeFilterData::IsValid() const
 	{
-		return QueryBlockChannels != 0 || QueryOverlapChannels != 0 || SimBlockChannels != 0 || Word3 != 0;
+		return BlockChannels != 0 || OverlapChannels != 0 || Word3 != 0;
 	}
 
 	bool FShapeFilterData::IsSimValid() const
 	{
-		return SimBlockChannels != 0 || Word3 != 0;
+		return BlockChannels != 0 || Word3 != 0;
 	}
 
 	EFilterFlags FShapeFilterData::GetFlags() const
@@ -222,17 +222,17 @@ namespace Chaos::Filter
 
 	uint64 FShapeFilterData::GetQueryBlockChannels() const
 	{
-		return QueryBlockChannels;
+		return BlockChannels;
 	}
 
 	uint64 FShapeFilterData::GetQueryOverlapChannels() const
 	{
-		return QueryOverlapChannels;
+		return OverlapChannels;
 	}
 
 	uint64 FShapeFilterData::GetSimBlockChannels() const
 	{
-		return SimBlockChannels;
+		return BlockChannels;
 	}
 
 	ENarrowFilterResult FShapeFilterData::NarrowFilter(const FShapeFilterData& OtherShape) const
@@ -248,7 +248,7 @@ namespace Chaos::Filter
 
 	FString FShapeFilterData::ToString() const
 	{
-		return FString::Format(TEXT("SimBlock(%x) QueryBlock(%x) QueryOverlap(%x) CollisionChannel(%x) MaskFilter(%x) Flags(%x)"), { GetSimBlockChannels(), GetQueryBlockChannels(), GetQueryOverlapChannels(), GetCollisionChannelIndex(), GetMaskFilter(), (uint8)GetFlags() });
+		return FString::Format(TEXT("Block(%x) Overlap(%x) CollisionChannel(%x) MaskFilter(%x) Flags(%x)"), { GetQueryBlockChannels(), GetQueryOverlapChannels(), GetCollisionChannelIndex(), GetMaskFilter(), (uint8)GetFlags() });
 	}
 
 	uint32 FShapeFilterData::GetLegacyWord3() const
@@ -262,12 +262,12 @@ namespace Chaos::Filter
 		{
 			// Note: Query Channels are unused in broad filtering. 
 			// Including them would increase the size of this filter which would increase the acceleration handle size.
-			QueryBlockChannels |= ShapeFilter.QueryBlockChannels;
-			QueryOverlapChannels |= ShapeFilter.QueryOverlapChannels;
+			QueryBlockChannels |= ShapeFilter.BlockChannels;
+			QueryOverlapChannels |= ShapeFilter.OverlapChannels;
 		}
 		if (bSimEnabled)
 		{
-			SimBlockChannels |= ShapeFilter.SimBlockChannels;
+			SimBlockChannels |= ShapeFilter.BlockChannels;
 			SimChannelMask |= ShapeFilter.GetCollisionChannelMask();
 		}
 	}
@@ -367,11 +367,6 @@ namespace Chaos::Filter
 	}
 
 	ENarrowFilterResult FQueryFilterData::NarrowFilter(const FShapeFilterData& ShapeFilter) const
-	{
-		return NarrowFilter(ShapeFilter, false);
-	}
-
-	ENarrowFilterResult FQueryFilterData::NarrowFilter(const FShapeFilterData& ShapeFilter, bool bSimFilter) const
 	{
 		const FMaskFilter QuerierMaskFilter = GetIgnoreMask();
 		const FMaskFilter ShapeMaskFilter = ShapeFilter.GetMaskFilter();
@@ -388,7 +383,7 @@ namespace Chaos::Filter
 		}
 		else
 		{
-			return ChannelTypeNarrowFilter(ShapeFilter, bSimFilter);
+			return ChannelTypeNarrowFilter(ShapeFilter);
 		}
 	}
 
@@ -404,14 +399,14 @@ namespace Chaos::Filter
 		}
 	}
 
-	ENarrowFilterResult FQueryFilterData::ChannelTypeNarrowFilter(const FShapeFilterData& ShapeFilter, bool bSimBlock) const
+	ENarrowFilterResult FQueryFilterData::ChannelTypeNarrowFilter(const FShapeFilterData& ShapeFilter) const
 	{
 		const uint64 QueryChannelMask = GetCollisionChannelMask();
 		const uint64 QueryBlockChannels = GetBlockChannels();
 		const uint64 QueryOverlapChannels = GetOverlapChannels();
 
 		const uint64 ShapeChannelMask = ShapeFilter.GetCollisionChannelMask();
-		const uint64 ShapeBlockChannels = bSimBlock ? ShapeFilter.GetSimBlockChannels() : ShapeFilter.GetQueryBlockChannels();
+		const uint64 ShapeBlockChannels = ShapeFilter.GetQueryBlockChannels();
 		const uint64 ShapeOverlapChannels = ShapeFilter.GetQueryOverlapChannels();
 
 		ENarrowFilterResult QuerierHitType = ENarrowFilterResult::None;
@@ -503,21 +498,13 @@ namespace Chaos::Filter
 
 	FShapeFilterBuilder& FShapeFilterBuilder::SetBlockChannelMask(const uint64 BlockMask)
 	{
-		ShapeFilterData.SimBlockChannels = BlockMask;
-		ShapeFilterData.QueryBlockChannels = BlockMask;
-		return *this;
-	}
-
-	FShapeFilterBuilder& FShapeFilterBuilder::SetBlockChannelMasks(const uint64 SimBlockMask, const uint64 QueryBlockMask)
-	{
-		ShapeFilterData.SimBlockChannels = SimBlockMask;
-		ShapeFilterData.QueryBlockChannels = QueryBlockMask;
+		ShapeFilterData.BlockChannels = BlockMask;
 		return *this;
 	}
 
 	FShapeFilterBuilder& FShapeFilterBuilder::SetOverlapChannelMask(const uint64 OverlapMask)
 	{
-		ShapeFilterData.QueryOverlapChannels = OverlapMask;
+		ShapeFilterData.OverlapChannels = OverlapMask;
 		return *this;
 	}
 
@@ -553,18 +540,10 @@ namespace Chaos::Filter
 		return Builder.Build();
 	}
 
-	FShapeFilterData FShapeFilterBuilder::BuildSimBlockAll(const EFilterFlags InFilterFlags)
-	{
-		FShapeFilterBuilder Builder;
-		Builder.SetBlockChannelMasks(FullChannelMask, 0);
-		Builder.SetFilterFlags(InFilterFlags);
-		return Builder.Build();
-	}
-
 	FShapeFilterData FShapeFilterBuilder::BuildFromLegacySimFilter(const FCollisionFilterData& SimFilter)
 	{
 		FShapeFilterData Result;
-		Result.SimBlockChannels = SimFilter.Word1;
+		Result.BlockChannels = SimFilter.Word1;
 		Result.Word3 = LegacyToNewWord3(SimFilter.Word3);
 		return Result;
 	}
@@ -591,10 +570,10 @@ namespace Chaos::Filter
 		CombinedShapeFilterData.InstanceData.OwnerId = QueryFilterData.Word0;
 		CombinedShapeFilterData.InstanceData.ComponentId = SimFilterData.Word2;
 
-		CombinedShapeFilterData.ShapeFilterData.QueryBlockChannels = QueryFilterData.Word1;
-		CombinedShapeFilterData.ShapeFilterData.QueryOverlapChannels = QueryFilterData.Word2;
+		ensureMsgf(QueryFilterData.Word1 == SimFilterData.Word1, TEXT("Legacy filters had differing block channels. Query(0x%X) Sim(0x%x). Using query block channels."), QueryFilterData.Word1, SimFilterData.Word1);
+		CombinedShapeFilterData.ShapeFilterData.BlockChannels = QueryFilterData.Word1;
+		CombinedShapeFilterData.ShapeFilterData.OverlapChannels = QueryFilterData.Word2;
 		CombinedShapeFilterData.ShapeFilterData.Word3 = LegacyToNewWord3(Word3);
-		CombinedShapeFilterData.ShapeFilterData.SimBlockChannels = SimFilterData.Word1;
 
 		return CombinedShapeFilterData;
 	}
@@ -602,15 +581,17 @@ namespace Chaos::Filter
 	void FShapeFilterBuilder::SetLegacyShapeQueryFilter(FCombinedShapeFilterData& CombinedShapeFilterData, const FCollisionFilterData& QueryFilterData)
 	{
 		FCollisionFilterData SimFilterData = GetLegacyShapeSimFilter(CombinedShapeFilterData);
+		const uint32 BlockChannels = QueryFilterData.Word1;
 		const uint32 Word3 = GetWord3FromNewQueryFilter(QueryFilterData.Word3, SimFilterData.Word3);
-		BuildFromLegacyShapeFilter(QueryFilterData, SimFilterData, Word3, CombinedShapeFilterData);
+		BuildFromLegacyShapeFilter(QueryFilterData, SimFilterData, BlockChannels, Word3, CombinedShapeFilterData);
 	}
 
 	void FShapeFilterBuilder::SetLegacyShapeSimFilter(FCombinedShapeFilterData& CombinedShapeFilterData, const FCollisionFilterData& SimFilterData)
 	{
 		const FCollisionFilterData QueryFilterData = GetLegacyShapeQueryFilter(CombinedShapeFilterData);
+		const uint32 BlockChannels = SimFilterData.Word1;
 		const uint32 Word3 = GetWord3FromNewSimFilter(SimFilterData.Word3, QueryFilterData.Word3);
-		BuildFromLegacyShapeFilter(QueryFilterData, SimFilterData, Word3, CombinedShapeFilterData);
+		BuildFromLegacyShapeFilter(QueryFilterData, SimFilterData, BlockChannels, Word3, CombinedShapeFilterData);
 	}
 
 	void FShapeFilterBuilder::GetLegacyShapeFilter(const FCombinedShapeFilterData& CombinedShapeFilterData, FCollisionFilterData& OutQueryFilterData, FCollisionFilterData& OutSimFilterData)
@@ -621,10 +602,10 @@ namespace Chaos::Filter
 		OutQueryFilterData.Word0 = CombinedShapeFilterData.InstanceData.OwnerId;
 		OutSimFilterData.Word2 = CombinedShapeFilterData.InstanceData.ComponentId;
 
-		OutQueryFilterData.Word1 = (uint32)CombinedShapeFilterData.ShapeFilterData.QueryBlockChannels;
-		OutQueryFilterData.Word2 = (uint32)CombinedShapeFilterData.ShapeFilterData.QueryOverlapChannels;
+		OutQueryFilterData.Word1 = (uint32)CombinedShapeFilterData.ShapeFilterData.BlockChannels;
+		OutQueryFilterData.Word2 = (uint32)CombinedShapeFilterData.ShapeFilterData.OverlapChannels;
 		OutQueryFilterData.Word3 = Word3;
-		OutSimFilterData.Word1 = (uint32)CombinedShapeFilterData.ShapeFilterData.SimBlockChannels;
+		OutSimFilterData.Word1 = (uint32)CombinedShapeFilterData.ShapeFilterData.BlockChannels;
 		OutSimFilterData.Word3 = Word3;
 	}
 
@@ -668,15 +649,14 @@ namespace Chaos::Filter
 		return SimWord3;
 	}
 
-	void FShapeFilterBuilder::BuildFromLegacyShapeFilter(const FCollisionFilterData& QueryFilterData, const FCollisionFilterData& SimFilterData, const uint32 Word3, FCombinedShapeFilterData& OutCombinedShapeFilterData)
+	void FShapeFilterBuilder::BuildFromLegacyShapeFilter(const FCollisionFilterData& QueryFilterData, const FCollisionFilterData& SimFilterData, const uint32 BlockChannels, const uint32 Word3, FCombinedShapeFilterData& OutCombinedShapeFilterData)
 	{
 		OutCombinedShapeFilterData.InstanceData.OwnerId = QueryFilterData.Word0;
 		OutCombinedShapeFilterData.InstanceData.ComponentId = SimFilterData.Word2;
 
-		OutCombinedShapeFilterData.ShapeFilterData.QueryBlockChannels = QueryFilterData.Word1;
-		OutCombinedShapeFilterData.ShapeFilterData.QueryOverlapChannels = QueryFilterData.Word2;
+		OutCombinedShapeFilterData.ShapeFilterData.BlockChannels = BlockChannels;
+		OutCombinedShapeFilterData.ShapeFilterData.OverlapChannels = QueryFilterData.Word2;
 		OutCombinedShapeFilterData.ShapeFilterData.Word3 = LegacyToNewWord3(Word3);
-		OutCombinedShapeFilterData.ShapeFilterData.SimBlockChannels = SimFilterData.Word1;
 	}
 
 	FQueryFilterData FQueryFilterBuilder::BuildFromLegacyQueryFilter(const FCollisionFilterData& QueryFilterData)
diff --git a/Engine/Source/Runtime/Experimental/Chaos/Private/PhysicsProxy/GeometryCollectionPhysicsProxy.cpp b/Engine/Source/Runtime/Experimental/Chaos/Private/PhysicsProxy/GeometryCollectionPhysicsProxy.cpp
--- a/Engine/Source/Runtime/Experimental/Chaos/Private/PhysicsProxy/GeometryCollectionPhysicsProxy.cpp
+++ b/Engine/Source/Runtime/Experimental/Chaos/Private/PhysicsProxy/GeometryCollectionPhysicsProxy.cpp
@@ -385,7 +385,7 @@ void SetImplicitToPTParticles(Chaos::TPBDRigidParticleHandle<Chaos::FReal, 3>* H
 		// Override collision filters and make this body collide with everything.
 		int32 CurrShape = 0;
 		
-		const Chaos::Filter::FShapeFilterData ShapeFilterData = Chaos::Filter::FShapeFilterBuilder::BuildSimBlockAll(Chaos::EFilterFlags::All);
+		const Chaos::Filter::FShapeFilterData ShapeFilterData = Chaos::Filter::FShapeFilterBuilder::BuildBlockAll(Chaos::EFilterFlags::All);
 		for (const TUniquePtr<Chaos::FPerShapeData>& Shape : Handle->ShapesArray())
 		{
 			Shape->SetSimEnabled(true);
diff --git a/Engine/Source/Runtime/Experimental/Chaos/Public/Chaos/CollisionFilterData.h b/Engine/Source/Runtime/Experimental/Chaos/Public/Chaos/CollisionFilterData.h
--- a/Engine/Source/Runtime/Experimental/Chaos/Public/Chaos/CollisionFilterData.h
+++ b/Engine/Source/Runtime/Experimental/Chaos/Public/Chaos/CollisionFilterData.h
@@ -157,9 +157,8 @@ namespace Chaos::Filter
 		friend FCombinedShapeFilterData;
 		friend struct FShapeUnionFilterData;
 
-		uint64 QueryBlockChannels = 0;
-		uint64 QueryOverlapChannels = 0;
-		uint64 SimBlockChannels = 0;
+		uint64 BlockChannels = 0;
+		uint64 OverlapChannels = 0;
 		uint32 Word3 = 0;
 	};
 
@@ -220,7 +219,6 @@ namespace Chaos::Filter
 
 		// Filter API
 		CHAOS_API ENarrowFilterResult NarrowFilter(const FShapeFilterData& ShapeFilter) const;
-		UE_INTERNAL CHAOS_API ENarrowFilterResult NarrowFilter(const FShapeFilterData& ShapeFilter, bool bSimFilter) const;
 
 		CHAOS_API FString ToString() const;
 
@@ -229,7 +227,7 @@ namespace Chaos::Filter
 		friend class FQueryObjectFilterBuilder;
 		friend class FQueryTraceFilterBuilder;
 
-		ENarrowFilterResult ChannelTypeNarrowFilter(const FShapeFilterData& ShapeFilter, bool bSimBlock) const;
+		ENarrowFilterResult ChannelTypeNarrowFilter(const FShapeFilterData& ShapeFilter) const;
 		ENarrowFilterResult ObjectTypeNarrowFilter(const FShapeFilterData& ShapeFilter) const;
 
 		uint32 Word0 = 0;
@@ -271,7 +269,6 @@ namespace Chaos::Filter
 		// TODO @ JoshD: Eventually deprecated and replace with SetChannelMask
 		CHAOS_API FShapeFilterBuilder& SetCollisionChannelIndex(const uint8 ChannelIndex);
 		CHAOS_API FShapeFilterBuilder& SetBlockChannelMask(const uint64 BlockMask);
-		UE_INTERNAL CHAOS_API FShapeFilterBuilder& SetBlockChannelMasks(const uint64 SimBlockMask, const uint64 QueryBlockMask);
 		CHAOS_API FShapeFilterBuilder& SetOverlapChannelMask(const uint64 OverlapMask);
 		CHAOS_API FShapeFilterBuilder& SetMaskFilter(const uint8 InMaskFilter);
 		CHAOS_API FShapeFilterBuilder& SetFilterFlags(const EFilterFlags InFilterFlags);
@@ -279,7 +276,6 @@ namespace Chaos::Filter
 
 		CHAOS_API FShapeFilterData Build() const;
 		CHAOS_API static FShapeFilterData BuildBlockAll(const EFilterFlags InFilterFlags = EFilterFlags::None);
-		CHAOS_API static FShapeFilterData BuildSimBlockAll(const EFilterFlags InFilterFlags = EFilterFlags::None);
 
 		UE_INTERNAL CHAOS_API static FShapeFilterData BuildFromLegacySimFilter(const FCollisionFilterData& SimFilter);
 		UE_INTERNAL CHAOS_API static FCombinedShapeFilterData BuildFromLegacyShapeFilter(const FCollisionFilterData& QueryFilterData, const FCollisionFilterData& SimFilterData);
@@ -295,7 +291,7 @@ namespace Chaos::Filter
 		UE_INTERNAL static uint32 GetWord3FromNewQueryFilter(const uint32 QueryWord3, const uint32 SimWord3);
 
 	private:
-		static void BuildFromLegacyShapeFilter(const FCollisionFilterData& QueryFilterData, const FCollisionFilterData& SimFilterData, const uint32 Word3, FCombinedShapeFilterData& OutCombinedShapeFilterData);
+		static void BuildFromLegacyShapeFilter(const FCollisionFilterData& QueryFilterData, const FCollisionFilterData& SimFilterData, const uint32 BlockChannels, const uint32 Word3, FCombinedShapeFilterData& OutCombinedShapeFilterData);
 
 		FShapeFilterData ShapeFilterData;
 	};
diff --git a/Engine/Source/Runtime/Experimental/ChaosSolverEngine/Private/Chaos/ChaosSolverActor.cpp b/Engine/Source/Runtime/Experimental/ChaosSolverEngine/Private/Chaos/ChaosSolverActor.cpp
--- a/Engine/Source/Runtime/Experimental/ChaosSolverEngine/Private/Chaos/ChaosSolverActor.cpp
+++ b/Engine/Source/Runtime/Experimental/ChaosSolverEngine/Private/Chaos/ChaosSolverActor.cpp
@@ -458,8 +458,9 @@ void AChaosSolverActor::MakeFloor()
 		FloorParticle->SetGeometry(MakeImplicitObjectPtr<Chaos::TBox<Chaos::FReal, 3>>(BoxMin, BoxMax));
 		FloorParticle->SetX(Chaos::FVec3(SolverLocation.X, SolverLocation.Y, FloorHeight));
 
-		const Chaos::Filter::FShapeFilterData ShapeFilterData = Chaos::Filter::FShapeFilterBuilder::BuildSimBlockAll(Chaos::EFilterFlags::All);
+		const Chaos::Filter::FShapeFilterData ShapeFilterData = Chaos::Filter::FShapeFilterBuilder::BuildBlockAll(Chaos::EFilterFlags::All);
 		FloorParticle->SetShapeFilterData(0, ShapeFilterData);
+		FloorParticle->SetShapeQueryCollisionEnabled(0, false);
 		Proxy = Chaos::FSingleParticlePhysicsProxy::Create(MoveTemp(FloorParticle));
 		RigidSolverProxy.Solver->RegisterObject(Proxy);
 	}
