Collision Filter Refactor Part 5.4: Deprecated the old Get Sim/Query Block/Overlap Masks functions in favor of unified functions now that the words are merged.
#Synced-CL 48505955

diff --git a/Engine/Source/Runtime/Engine/Private/Collision/PhysicsFiltering.cpp b/Engine/Source/Runtime/Engine/Private/Collision/PhysicsFiltering.cpp
--- a/Engine/Source/Runtime/Engine/Private/Collision/PhysicsFiltering.cpp
+++ b/Engine/Source/Runtime/Engine/Private/Collision/PhysicsFiltering.cpp
@@ -198,8 +198,8 @@ bool HasCollisionChannel(const Chaos::Filter::FQueryFilterData& QueryFilterData,
 FCollisionResponseContainer ExtractCollisionResponseContainer(const Chaos::Filter::FShapeFilterData& ShapeFilterData)
 {
 	FCollisionResponseContainer CollisionResponseContainer;
-	const uint64 BlockChannels = ShapeFilterData.GetQueryBlockChannels();
-	const uint64 OverlapChannels = ShapeFilterData.GetQueryOverlapChannels();
+	const uint64 BlockChannels = ShapeFilterData.GetBlockChannels();
+	const uint64 OverlapChannels = ShapeFilterData.GetOverlapChannels();
 
 	for (int32 ChannelIndex = 0; ChannelIndex < UE_ARRAY_COUNT(CollisionResponseContainer.EnumArray); ++ChannelIndex)
 	{
diff --git a/Engine/Source/Runtime/Engine/Public/Physics/PhysicsFiltering.h b/Engine/Source/Runtime/Engine/Public/Physics/PhysicsFiltering.h
--- a/Engine/Source/Runtime/Engine/Public/Physics/PhysicsFiltering.h
+++ b/Engine/Source/Runtime/Engine/Public/Physics/PhysicsFiltering.h
@@ -81,8 +81,8 @@ struct FPhysicsFilterBuilder
 		PRAGMA_DISABLE_INTERNAL_WARNINGS
 		const Chaos::Filter::FShapeFilterData& FilterData = BuildShapeFilterData();
 		OutWord0 = SourceObjectID;
-		OutWord1 = (uint32)FilterData.GetQueryBlockChannels();
-		OutWord2 = (uint32)FilterData.GetQueryOverlapChannels();
+		OutWord1 = (uint32)FilterData.GetBlockChannels();
+		OutWord2 = (uint32)FilterData.GetOverlapChannels();
 		OutWord3 = FilterData.GetLegacyWord3();
 		PRAGMA_ENABLE_INTERNAL_WARNINGS
 	}
@@ -100,7 +100,7 @@ struct FPhysicsFilterBuilder
 		PRAGMA_DISABLE_INTERNAL_WARNINGS
 		const Chaos::Filter::FShapeFilterData& FilterData = BuildShapeFilterData();
 		OutWord0 = BodyIndex;
-		OutWord1 = (uint32)FilterData.GetSimBlockChannels();
+		OutWord1 = (uint32)FilterData.GetBlockChannels();
 		OutWord2 = ComponentID;
 		OutWord3 = FilterData.GetLegacyWord3();
 		PRAGMA_ENABLE_INTERNAL_WARNINGS
@@ -111,8 +111,8 @@ struct FPhysicsFilterBuilder
 	{
 		PRAGMA_DISABLE_INTERNAL_WARNINGS
 		const Chaos::Filter::FShapeFilterData& FilterData = BuildShapeFilterData();
-		OutBlockingBits = (uint32)FilterData.GetQueryBlockChannels();
-		OutTouchingBits = (uint32)FilterData.GetQueryOverlapChannels();
+		OutBlockingBits = (uint32)FilterData.GetBlockChannels();
+		OutTouchingBits = (uint32)FilterData.GetOverlapChannels();
 		OutObjectTypeAndFlags = FilterData.GetLegacyWord3();
 		PRAGMA_ENABLE_INTERNAL_WARNINGS
 	}
diff --git a/Engine/Source/Runtime/Experimental/Chaos/Private/Chaos/CollisionFilterData.cpp b/Engine/Source/Runtime/Experimental/Chaos/Private/Chaos/CollisionFilterData.cpp
--- a/Engine/Source/Runtime/Experimental/Chaos/Private/Chaos/CollisionFilterData.cpp
+++ b/Engine/Source/Runtime/Experimental/Chaos/Private/Chaos/CollisionFilterData.cpp
@@ -220,6 +220,16 @@ namespace Chaos::Filter
 		return false;
 	}
 
+	uint64 FShapeFilterData::GetBlockChannels() const
+	{
+		return BlockChannels;
+	}
+
+	uint64 FShapeFilterData::GetOverlapChannels() const
+	{
+		return OverlapChannels;
+	}
+
 	uint64 FShapeFilterData::GetQueryBlockChannels() const
 	{
 		return BlockChannels;
@@ -239,8 +249,8 @@ namespace Chaos::Filter
 	{
 		const uint64 ChannelMask0 = GetCollisionChannelMask();
 		const uint64 ChannelMask1 = OtherShape.GetCollisionChannelMask();
-		const uint64 BlockMask0 = GetSimBlockChannels();
-		const uint64 BlockMask1 = OtherShape.GetSimBlockChannels();
+		const uint64 BlockMask0 = GetBlockChannels();
+		const uint64 BlockMask1 = OtherShape.GetBlockChannels();
 
 		const bool bBlocking = (ChannelMask0 & BlockMask1) && (ChannelMask1 & BlockMask0);
 		return bBlocking ? ENarrowFilterResult::Block : ENarrowFilterResult::None;
@@ -248,7 +258,7 @@ namespace Chaos::Filter
 
 	FString FShapeFilterData::ToString() const
 	{
-		return FString::Format(TEXT("Block(%x) Overlap(%x) CollisionChannel(%x) MaskFilter(%x) Flags(%x)"), { GetQueryBlockChannels(), GetQueryOverlapChannels(), GetCollisionChannelIndex(), GetMaskFilter(), (uint8)GetFlags() });
+		return FString::Format(TEXT("Block(%x) Overlap(%x) CollisionChannel(%x) MaskFilter(%x) Flags(%x)"), { GetBlockChannels(), GetOverlapChannels(), GetCollisionChannelIndex(), GetMaskFilter(), (uint8)GetFlags() });
 	}
 
 	uint32 FShapeFilterData::GetLegacyWord3() const
@@ -406,8 +416,8 @@ namespace Chaos::Filter
 		const uint64 QueryOverlapChannels = GetOverlapChannels();
 
 		const uint64 ShapeChannelMask = ShapeFilter.GetCollisionChannelMask();
-		const uint64 ShapeBlockChannels = ShapeFilter.GetQueryBlockChannels();
-		const uint64 ShapeOverlapChannels = ShapeFilter.GetQueryOverlapChannels();
+		const uint64 ShapeBlockChannels = ShapeFilter.GetBlockChannels();
+		const uint64 ShapeOverlapChannels = ShapeFilter.GetOverlapChannels();
 
 		ENarrowFilterResult QuerierHitType = ENarrowFilterResult::None;
 		ENarrowFilterResult ShapeHitType = ENarrowFilterResult::None;
diff --git a/Engine/Source/Runtime/Experimental/Chaos/Public/Chaos/CollisionFilterData.h b/Engine/Source/Runtime/Experimental/Chaos/Public/Chaos/CollisionFilterData.h
--- a/Engine/Source/Runtime/Experimental/Chaos/Public/Chaos/CollisionFilterData.h
+++ b/Engine/Source/Runtime/Experimental/Chaos/Public/Chaos/CollisionFilterData.h
@@ -140,8 +140,13 @@ namespace Chaos::Filter
 		CHAOS_API uint64 GetCollisionChannelMask() const;
 		CHAOS_API bool IsCollisionChannelSet(const uint32 ChannelIndex) const;
 
+		CHAOS_API uint64 GetBlockChannels() const;
+		CHAOS_API uint64 GetOverlapChannels() const;
+		UE_DEPRECATED(5.8, "Use GetBlockChannels")
 		CHAOS_API uint64 GetQueryBlockChannels() const;
+		UE_DEPRECATED(5.8, "Use GetOverlapChannels")
 		CHAOS_API uint64 GetQueryOverlapChannels() const;
+		UE_DEPRECATED(5.8, "Use GetBlockChannels")
 		CHAOS_API uint64 GetSimBlockChannels() const;
 
 		// Filter API
