Add IadTotalInstallCacheMissCount to IAD installer analytics - This tracks the number of requests that actually downloaded data.
Add analytics for DeraggedBlock fragmentation and total fragmentation.
Fix some ration calculations to return 1 instead of 0 when the denominator is 0.

#Synced-CL 48498858

diff --git a/Engine/Source/Runtime/Experimental/IoStore/OnDemand/Private/OnDemandContentInstaller.cpp b/Engine/Source/Runtime/Experimental/IoStore/OnDemand/Private/OnDemandContentInstaller.cpp
--- a/Engine/Source/Runtime/Experimental/IoStore/OnDemand/Private/OnDemandContentInstaller.cpp
+++ b/Engine/Source/Runtime/Experimental/IoStore/OnDemand/Private/OnDemandContentInstaller.cpp
@@ -1415,7 +1415,7 @@ void FOnDemandContentInstaller::CompleteInstallRequest(FRequest& Request)
 	const double DurationInSeconds	= FPlatformTime::ToSeconds64(DurationCycles);
 	const double CacheHitRatio		= InstallRequest.Progress.TotalContentSize > 0
 		? double(InstallRequest.Progress.TotalContentSize - InstallRequest.Progress.TotalInstallSize) / double(InstallRequest.Progress.TotalContentSize)
-		: 0.0;
+		: 1.0;
 
 #if !UE_BUILD_SHIPPING
 	FString DebugInfoString = FString::Join(DebugInfo, TEXT(","));
diff --git a/Engine/Source/Runtime/Experimental/IoStore/OnDemand/Private/OnDemandInstallCache.cpp b/Engine/Source/Runtime/Experimental/IoStore/OnDemand/Private/OnDemandInstallCache.cpp
--- a/Engine/Source/Runtime/Experimental/IoStore/OnDemand/Private/OnDemandInstallCache.cpp
+++ b/Engine/Source/Runtime/Experimental/IoStore/OnDemand/Private/OnDemandInstallCache.cpp
@@ -3003,47 +3003,20 @@ FResult FOnDemandInstallCache::Defrag(
 
 	if (TotalBytesToFree && *TotalBytesToFree == 0)
 	{
-		return MakeValue();
-	}
-
-	FResult Result			= MakeValue();
-	uint64 FragmentedBytes	= 0;
-
-	ON_SCOPE_EXIT
-	{
-		if (OutTotalFreedBytes != nullptr)
+		if (OutTotalFreedBytes)
 		{
-			*OutTotalFreedBytes = Result.HasError() ? 0 : FragmentedBytes;
+			*OutTotalFreedBytes = 0;
 		}
-		FOnDemandInstallCacheStats::OnDefrag(Result, FragmentedBytes);
-	};
-
-	const uint64 TotalCachedBytes = Algo::TransformAccumulate(BlockInfo,
-		[](const TPair<FCasBlockId, FCasBlockInfo>& Kv) { return Kv.Value.FileSize; },
-		uint64(0));
 
-	if (TotalCachedBytes > MaxCacheSize)
-	{
-		// Ruh-Roh! There's not enough of the disk quota left to run a defrag!
-		const FString ErrorMsg = FString::Printf(TEXT("Cache size is greater than disk quota - Cannot Defragment!, MaxCacheSize=%.2lf MiB"), 
-			ToMiB(MaxCacheSize));
-		UE_LOG(LogIoStoreOnDemand, Error, TEXT("%s"), *ErrorMsg);
-
-		FInstallCacheErrorContext Ctx = MakeInstallCacheErrorContext(TotalCachedBytes);
-		Ctx.IoError = FIoStatus(EIoErrorCode::OutOfDiskSpace, ErrorMsg),
-		Result = MakeError(InstallCacheDefragError::MakeError(MoveTemp(Ctx), UE::UnifiedError::EDetailFilter::All));
-
-		// Append a critical error entry to clear the cache at next startup
-		FCasJournal::FTransaction Transaction = FCasJournal::Begin(GetJournalFilename());
-		Transaction.CriticalError(FCasJournal::EErrorCode::DefragOutOfDiskSpace);
-		if (FResult CommitResult = FCasJournal::Commit(MoveTemp(Transaction)); CommitResult.HasError())
-		{
-			UE_LOGFMT(LogIoStoreOnDemand, Error, "Failed to commit critical errors to CAS journal ({Error})", CommitResult.GetError());
-		}
-
-		return Result;
+		return MakeValue();
 	}
 
+	FResult Result						= MakeValue();
+	uint64 TotalCachedBytes				= 0;
+	uint64 TotalFragmentedBytes			= 0;
+	uint64 BlocksToDefragFragmentedSize	= 0;
+	uint64 BlocksToDefragTotalSize		= 0;
+
 	struct FDefragBlockReferencedChunk
 	{
 		FIoHash		Hash;
@@ -3068,8 +3041,6 @@ FResult FOnDemandInstallCache::Defrag(
 		return LHS.RefSize < RHS.RefSize;
 	});
 
-	uint64 TotalBlockSize = 0;
-
 	if (TotalBytesToFree)
 	{
 		// Partial defrag
@@ -3083,11 +3054,15 @@ FResult FOnDemandInstallCache::Defrag(
 			const FCasBlockId BlockId = Kv.Key;
 			const FCasBlockInfo& Info = Kv.Value;
 
+			TotalCachedBytes += Info.FileSize;
+			TotalFragmentedBytes += ensure(Info.RefSize <= Info.FileSize) ?
+				(Info.FileSize - Info.RefSize) : Info.FileSize;
+
 			if (!bPossibleToFreeBytes && Info.RefSize < Info.FileSize)
 			{
 				// Block is fragmented
-				FragmentedBytes += (Info.FileSize - Info.RefSize);
-				TotalBlockSize += Info.FileSize;
+				BlocksToDefragFragmentedSize += (Info.FileSize - Info.RefSize);
+				BlocksToDefragTotalSize += Info.FileSize;
 
 				FreedBlockBytes += Info.FileSize;
 				NewBlockBytes += Info.RefSize; // For now, assume that nothing will be moved to the current block
@@ -3104,10 +3079,10 @@ FResult FOnDemandInstallCache::Defrag(
 				// Block is too small whether or not its fragmented
 				if (ensure(Info.RefSize <= Info.FileSize))
 				{
-					FragmentedBytes += (Info.FileSize - Info.RefSize);
+					BlocksToDefragFragmentedSize += (Info.FileSize - Info.RefSize);
 				}
 
-				TotalBlockSize += Info.FileSize;
+				BlocksToDefragTotalSize += Info.FileSize;
 
 				BlocksToDefrag.Add(FDefragBlock{ .BlockId = BlockId, .LastAccess = Info.LastAccess });
 			}
@@ -3115,8 +3090,8 @@ FResult FOnDemandInstallCache::Defrag(
 
 		if (!bPossibleToFreeBytes)
 		{
-			FString ErrorMessage = "Failed to defrag the install cache due to all data being referenced by the game";
-			UE_LOG(LogIoStoreOnDemand, Error, TEXT("%s"), *ErrorMessage);
+			FStringView ErrorMessage(TEXTVIEW("Failed to defrag the install cache due to too much data being referenced by the game"));
+			UE_LOGFMT(LogIoStoreOnDemand, Error, "{ErrorMessage}", ErrorMessage);
 			FInstallCacheErrorContext Ctx = MakeInstallCacheErrorContext(TotalCachedBytes);
 			Ctx.IoError = FIoStatus(EIoErrorCode::OutOfDiskSpace, ErrorMessage);
 			return Result = MakeError(InstallCacheDefragError::MakeError(MoveTemp(Ctx), UE::UnifiedError::EDetailFilter::All));
@@ -3130,11 +3105,15 @@ FResult FOnDemandInstallCache::Defrag(
 			const FCasBlockId BlockId = Kv.Key;
 			const FCasBlockInfo& Info = Kv.Value;
 
+			TotalCachedBytes += Info.FileSize;
+			TotalFragmentedBytes += ensure(Info.RefSize <= Info.FileSize) ?
+				(Info.FileSize - Info.RefSize) : Info.FileSize;
+
 			if (Info.RefSize < Info.FileSize)
 			{
 				// Block is fragmented
-				FragmentedBytes += (Info.FileSize - Info.RefSize);
-				TotalBlockSize += Info.FileSize;
+				BlocksToDefragFragmentedSize += (Info.FileSize - Info.RefSize);
+				BlocksToDefragTotalSize += Info.FileSize;
 
 				BlocksToDefrag.Add(FDefragBlock{ .BlockId = BlockId, .LastAccess = Info.LastAccess });
 			}
@@ -3143,10 +3122,10 @@ FResult FOnDemandInstallCache::Defrag(
 				// Block is too small whether or not its fragmented
 				if (ensure(Info.RefSize <= Info.FileSize))
 				{
-					FragmentedBytes += (Info.FileSize - Info.RefSize);
+					BlocksToDefragFragmentedSize += (Info.FileSize - Info.RefSize);
 				}
 
-				TotalBlockSize += Info.FileSize;
+				BlocksToDefragTotalSize += Info.FileSize;
 
 				BlocksToDefrag.Add(FDefragBlock{ .BlockId = BlockId, .LastAccess = Info.LastAccess });
 			}
@@ -3160,7 +3139,39 @@ FResult FOnDemandInstallCache::Defrag(
 		}
 	}
 
-	UE_LOG(LogIoStoreOnDemand, Display, TEXT("Defrag found %" UINT64_FMT " fragmented bytes of %" UINT64_FMT " total bytes in %i blocks."), FragmentedBytes, TotalBlockSize, BlocksToDefrag.Num());
+	UE_LOG(LogIoStoreOnDemand, Display, TEXT("Defrag found %" UINT64_FMT " fragmented bytes of %" UINT64_FMT " total bytes in %i blocks."), 
+		BlocksToDefragFragmentedSize, BlocksToDefragTotalSize, BlocksToDefrag.Num()); 
+
+	ON_SCOPE_EXIT
+	{
+		if (OutTotalFreedBytes != nullptr)
+		{
+			*OutTotalFreedBytes = Result.HasError() ? 0 : BlocksToDefragFragmentedSize;
+		}
+		FOnDemandInstallCacheStats::OnDefrag(Result, TotalFragmentedBytes, TotalCachedBytes, BlocksToDefragFragmentedSize, BlocksToDefragTotalSize);
+	};
+
+	if (TotalCachedBytes > MaxCacheSize)
+	{
+		// Ruh-Roh! There's not enough of the disk quota left to run a defrag!
+		const FString ErrorMsg = FString::Printf(TEXT("Cache size is greater than disk quota - Cannot Defragment!, MaxCacheSize=%.2lf MiB"),
+			ToMiB(MaxCacheSize));
+		UE_LOG(LogIoStoreOnDemand, Error, TEXT("%s"), *ErrorMsg);
+
+		FInstallCacheErrorContext Ctx = MakeInstallCacheErrorContext(TotalCachedBytes);
+		Ctx.IoError = FIoStatus(EIoErrorCode::OutOfDiskSpace, ErrorMsg),
+			Result = MakeError(InstallCacheDefragError::MakeError(MoveTemp(Ctx), UE::UnifiedError::EDetailFilter::All));
+
+		// Append a critical error entry to clear the cache at next startup
+		FCasJournal::FTransaction Transaction = FCasJournal::Begin(GetJournalFilename());
+		Transaction.CriticalError(FCasJournal::EErrorCode::DefragOutOfDiskSpace);
+		if (FResult CommitResult = FCasJournal::Commit(MoveTemp(Transaction)); CommitResult.HasError())
+		{
+			UE_LOGFMT(LogIoStoreOnDemand, Error, "Failed to commit critical errors to CAS journal ({Error})", CommitResult.GetError());
+		}
+
+		return Result;
+	}
 
 	// Right now, don't allow moving chunks to the current block for defrag. Its somewhat dangerous and hard to reason about.
 	// - Currently, the slack in the current block cannot be determined without opening a write handle to the block.
@@ -3204,9 +3215,10 @@ FResult FOnDemandInstallCache::Defrag(
 		}
 	}
 
-	if (TotalBlockSize - FragmentedBytes != TotalRefBytes)
+	if (BlocksToDefragTotalSize - BlocksToDefragFragmentedSize != TotalRefBytes)
 	{
-		UE_LOGFMT(LogIoStoreOnDemand, Error, "Possibly corrupt CAS blocks - TotalBlockSize {TotalBlockSize}, FragmentedBytes {FragmentedBytes}, TotalRefBytes {TotalRefBytes}", TotalBlockSize,  FragmentedBytes, TotalRefBytes);
+		UE_LOGFMT(LogIoStoreOnDemand, Error, "Possibly corrupt CAS blocks - TotalBlockSize {TotalBlockSize}, FragmentedBytes {FragmentedBytes}, TotalRefBytes {TotalRefBytes}", 
+			BlocksToDefragTotalSize, BlocksToDefragFragmentedSize, TotalRefBytes);
 		ensure(false);
 	}
 
@@ -3337,7 +3349,8 @@ FResult FOnDemandInstallCache::Defrag(
 		}
 	}
 
-	UE_LOG(LogIoStoreOnDemand, Display, TEXT("Defrag removed %" UINT64_FMT " fragmented bytes of %" UINT64_FMT " total bytes in %i blocks."), FragmentedBytes, TotalBlockSize, BlocksToDefrag.Num());
+	UE_LOG(LogIoStoreOnDemand, Display, TEXT("Defrag removed %" UINT64_FMT " fragmented bytes of %" UINT64_FMT " total bytes in %i blocks."),
+		BlocksToDefragTotalSize, BlocksToDefragTotalSize, BlocksToDefrag.Num());
 
 	return Result;
 }
diff --git a/Engine/Source/Runtime/Experimental/IoStore/OnDemand/Private/Statistics.cpp b/Engine/Source/Runtime/Experimental/IoStore/OnDemand/Private/Statistics.cpp
--- a/Engine/Source/Runtime/Experimental/IoStore/OnDemand/Private/Statistics.cpp
+++ b/Engine/Source/Runtime/Experimental/IoStore/OnDemand/Private/Statistics.cpp
@@ -110,10 +110,15 @@ static void SendImmediateAnalytic(FOnDemandImmediateAnalytic&& Event)
 
 FOnDemandImmediateAnalytic MakeAnalyticsEventFromResult(const TCHAR* EventName, const FResult& Result)
 {
-	FString ErrorCode = "Success";
-	FString ErrorJson = "{\"Root\":{\"ErrorCodeString\":\"Success\", \"ModuleIdString\":\"UnifiedError\"}}";
+	FString ErrorCode;
+	FString ErrorJson;
 
-	if (Result.HasError())
+	if (!Result.HasError())
+	{
+		ErrorCode = TEXTVIEW("Success");
+		ErrorJson = TEXTVIEW("{\"Root\":{\"ErrorCodeString\":\"Success\", \"ModuleIdString\":\"UnifiedError\"}}");
+	}
+	else
 	{
 		ErrorCode = FString(Result.GetError().GetModuleIdAndErrorCodeString());
 		ErrorJson = Result.GetError().SerializeToJsonForAnalytics();
@@ -522,9 +527,12 @@ struct FInstallerAnalytics
 
 	uint64 InstallCount = 0;
 	uint64 InstallErrorCount = 0;
+	uint64 InstallCacheMissCount = 0;
 	uint64 DownloadedBytes = 0;
 	uint64 TotalInstallDurationMs = 0;
-	double TotalCacheHitRatio = 0;
+	double TotalCacheHitRatio = 0.0;
+	double TotalFragmentionRatio = 0.0;
+	double DefraggedBlockFragmentionRatio = 0.0;
 };
 FMutex FInstallerAnalytics::Mutex;
 static FInstallerAnalytics InstallerAnalytics;
@@ -543,7 +551,6 @@ struct FInstallCacheAnalytics
 	uint64 PurgedBytes = 0;
 	uint64 PurgeCount = 0;
 	uint64 PurgeErrorCount = 0;
-	uint64 FragmentedBytes = 0;
 	uint64 DefragCount = 0;
 	uint64 DefragErrorCount = 0;
 	uint64 JournalCommitCount = 0;
@@ -1476,11 +1483,16 @@ void FOnDemandContentInstallerStats::OnRequestCompleted(
 
 	{
 		TUniqueLock Lock(InstallerAnalytics);
+		
+		InstallerAnalytics.InstallCount++;
 		if (Result.HasError() && !UE::UnifiedError::IsCancellationError(Result.GetError()))
 		{
 			InstallerAnalytics.InstallErrorCount++;
 		}
-		InstallerAnalytics.InstallCount++;
+		if (CacheHitRatio != 1.0)
+		{
+			InstallerAnalytics.InstallCacheMissCount++;
+		}
 		InstallerAnalytics.DownloadedBytes += DownloadedBytes;
 		InstallerAnalytics.TotalInstallDurationMs += uint64(FPlatformTime::ToMilliseconds64(DurationCycles));
 		InstallerAnalytics.TotalCacheHitRatio += CacheHitRatio;
@@ -1513,15 +1525,20 @@ void FOnDemandContentInstallerStats::ReportAnalytics(TArray<FAnalyticsEventAttri
 	const FTimespan OldestBlockAge = (Now >= OldestBlockAccess) ? (Now - OldestBlockAccess) : FTimespan::MaxValue();
 
 	const double AvgInstallDurationMs = CurrentInstallerAnalytics.InstallCount > 0 ? (double)CurrentInstallerAnalytics.TotalInstallDurationMs / (double)CurrentInstallerAnalytics.InstallCount : 0.0;
-	const double AvgCacheHitRatio = CurrentInstallerAnalytics.InstallCount > 0 ? (double)CurrentInstallerAnalytics.TotalCacheHitRatio / (double)CurrentInstallerAnalytics.InstallCount : 0.0;
+	const double AvgCacheHitRatio = CurrentInstallerAnalytics.InstallCount > 0 ? (double)CurrentInstallerAnalytics.TotalCacheHitRatio / (double)CurrentInstallerAnalytics.InstallCount : 1.0;
+	const double AvgTotalFragmentationRatio = CurrentInstallerAnalytics.InstallCount > 0 ? (double)CurrentInstallerAnalytics.TotalFragmentionRatio / (double)CurrentInstallerAnalytics.InstallCount : 0.0;
+	const double AvgDefraggedBlockFragmentationRatio = CurrentInstallerAnalytics.InstallCount > 0 ? (double)CurrentInstallerAnalytics.DefraggedBlockFragmentionRatio / (double)CurrentInstallerAnalytics.InstallCount : 0.0;
 
 	AppendAnalyticsEventAttributeArray(OutAnalyticsArray
 		, TEXT("IadTotalInstallCount"), CurrentInstallerAnalytics.InstallCount
 		, TEXT("IadTotalInstallErrorCount"), CurrentInstallerAnalytics.InstallErrorCount 
+		, TEXT("IadTotalInstallCacheMissCount"), CurrentInstallerAnalytics.InstallCacheMissCount
 		, TEXT("IadTotalDownloadedBytes"), CurrentInstallerAnalytics.DownloadedBytes 
 		, TEXT("IadTotalInstallDurationMs"), CurrentInstallerAnalytics.TotalInstallDurationMs 
 		, TEXT("IadAvgInstallDurationMs"), AvgInstallDurationMs
 		, TEXT("IadAvgCacheHitRatio"), AvgCacheHitRatio
+		, TEXT("IadAvgTotalFragmentationRatio"), AvgTotalFragmentationRatio
+		, TEXT("IadAvgDefraggedBlockFragmentationRatio"), AvgDefraggedBlockFragmentationRatio
 		, TEXT("IadInstallCacheStartupErrorCode"), CurrentInstallCacheAnalytics.StartupErrorCode
 		, TEXT("IadInstallCacheVerificationRemovedBlockCount"), CurrentInstallCacheAnalytics.VerificationRemovedBlockCount 
 		, TEXT("IadInstallCacheFlushCount"), CurrentInstallCacheAnalytics.FlushCount 
@@ -1618,13 +1635,26 @@ void FOnDemandInstallCacheStats::OnPurge(
 	SendImmediateAnalytic(MoveTemp(Event));
 }
 
-void FOnDemandInstallCacheStats::OnDefrag(const FResult& Result, uint64 FragmentedBytes)
+void FOnDemandInstallCacheStats::OnDefrag(
+	const FResult& Result, 
+	uint64 TotalFragmentedBytes,
+	uint64 TotalCachedBytes,
+	uint64 BlocksToDefragFragmentedSize, 
+	uint64 BlocksToDefragTotalSize)
 {
 	if (Result.HasError() && UE::UnifiedError::IsCancellationError(Result.GetError()))
 	{
 		return;
 	}
 
+	const double TotalFragmentionRatio = TotalCachedBytes > 0
+		? double(TotalFragmentedBytes) / double(TotalCachedBytes)
+		: 0.0;
+
+	const double BlocksToDefragFragmentionRatio = BlocksToDefragTotalSize > 0
+		? double(BlocksToDefragFragmentedSize) / double(BlocksToDefragTotalSize)
+		: 0.0;
+
 	{
 		UE::TUniqueLock Lock(InstallCacheAnalytics);
 		if (Result.HasError())
@@ -1632,11 +1662,22 @@ void FOnDemandInstallCacheStats::OnDefrag(const FResult& Result, uint64 Fragment
 			InstallCacheAnalytics.DefragErrorCount++;
 		}
 		InstallCacheAnalytics.DefragCount++;
-		InstallCacheAnalytics.FragmentedBytes += FragmentedBytes;
+	}
+
+	{
+		TUniqueLock Lock(InstallerAnalytics);
+		InstallerAnalytics.TotalFragmentionRatio += TotalFragmentionRatio;
+		InstallerAnalytics.DefraggedBlockFragmentionRatio += BlocksToDefragFragmentionRatio;
 	}
 
 	FOnDemandImmediateAnalytic Event = MakeAnalyticsEventFromResult(TEXT("IadCacheDefrag"), Result);
-	AppendAnalyticsEventAttributeArray(Event.AnalyticsArray, TEXT("FragmentedBytes"), FragmentedBytes);
+	AppendAnalyticsEventAttributeArray(Event.AnalyticsArray, 
+		TEXT("TotalFragmentedBytes"), TotalFragmentedBytes,
+		TEXT("TotalCachedBytes"), TotalCachedBytes,
+		TEXT("TotalFragmentionRatio"), TotalFragmentionRatio,
+		TEXT("DefraggedBlocksFragmentedSize"), BlocksToDefragFragmentedSize,
+		TEXT("DefraggedBlocksTotalSize"), BlocksToDefragTotalSize,
+		TEXT("DefraggedBlocksFragmentionRatio"), BlocksToDefragFragmentionRatio);
 	SendImmediateAnalytic(MoveTemp(Event));
 }
 
diff --git a/Engine/Source/Runtime/Experimental/IoStore/OnDemand/Private/Statistics.h b/Engine/Source/Runtime/Experimental/IoStore/OnDemand/Private/Statistics.h
--- a/Engine/Source/Runtime/Experimental/IoStore/OnDemand/Private/Statistics.h
+++ b/Engine/Source/Runtime/Experimental/IoStore/OnDemand/Private/Statistics.h
@@ -169,7 +169,12 @@ public:
 		uint64 NewCacheSize,
 		uint64 BytesToPurge,
 		uint64 PurgedBytes) IAS_STATISTICS_IMPL()
-	static void OnDefrag(const FResult& Result, uint64 FragmentedBytes) IAS_STATISTICS_IMPL()
+	static void OnDefrag(
+		const FResult& Result,
+		uint64 TotalFragmentedBytes,
+		uint64 TotalCachedBytes,
+		uint64 BlocksToDefragFragmentedSize, 
+		uint64 BlocksToDefragTotalSize) IAS_STATISTICS_IMPL()
 	static void OnBlockDeleted(int64 LastAccessTicks, bool bFromDefrag) IAS_STATISTICS_IMPL()
 	static void OnCacheUsage(
 		uint64 MaxCacheSize,
