Fix uninitialized struct member warnings

#Synced-CL 48490803

diff --git a/Engine/Source/Runtime/Experimental/ChaosVisualDebugger/Public/DataWrappers/ChaosVDCameraDataWrapper.h b/Engine/Source/Runtime/Experimental/ChaosVisualDebugger/Public/DataWrappers/ChaosVDCameraDataWrapper.h
--- a/Engine/Source/Runtime/Experimental/ChaosVisualDebugger/Public/DataWrappers/ChaosVDCameraDataWrapper.h
+++ b/Engine/Source/Runtime/Experimental/ChaosVisualDebugger/Public/DataWrappers/ChaosVDCameraDataWrapper.h
@@ -51,13 +51,13 @@ struct FChaosVDCameraDataWrapper
 	FChaosVDCameraIdentifier Camera;
 
 	UPROPERTY()
-	FVector Position;
+	FVector Position = FVector::ZeroVector;
 
 	UPROPERTY()
-	FQuat Rotation;
+	FQuat Rotation = FQuat::Identity;
 
 	UPROPERTY()
-	float FOV;
+	float FOV = 0;
 
 	CHAOSVDRUNTIME_API bool Serialize(FArchive& Ar);
 
