[HWRT] Enable caching of ray tracing instance mask and flags for FGeometryCollectionSceneProxyBase dynamic ray tracing instances.
- all instances use the same mesh batches and that array is also already cached in the proxy (CachedPartRayTracingMaterials).

- RayTracingScene_AddDynamicInstances cost on test level with 300 GCs each consisting of ~70 parts: 12.9ms -> 8.52ms.

#Synced-CL 48477729

diff --git a/Engine/Source/Runtime/Experimental/GeometryCollectionEngine/Private/GeometryCollection/GeometryCollectionSceneProxy.cpp b/Engine/Source/Runtime/Experimental/GeometryCollectionEngine/Private/GeometryCollection/GeometryCollectionSceneProxy.cpp
--- a/Engine/Source/Runtime/Experimental/GeometryCollectionEngine/Private/GeometryCollection/GeometryCollectionSceneProxy.cpp
+++ b/Engine/Source/Runtime/Experimental/GeometryCollectionEngine/Private/GeometryCollection/GeometryCollectionSceneProxy.cpp
@@ -677,8 +677,12 @@ void FGeometryCollectionSceneProxyBase::GetDynamicRayTracingInstances(FRayTracin
 			PartInstanceTransforms[TransformIndex] = FMatrix(DynamicData->Transforms[TransformIndex]) * LocalToWorld;
 		}
 
+		bool bInstanceMaskAndFlagsDirty = false;
+
 		if (SectionArray.Num() != CachedPartRayTracingMaterials.Num())
 		{
+			bInstanceMaskAndFlagsDirty = true;
+
 			CachedPartRayTracingMaterials.Reset();
 			CachedPartRayTracingMaterials.Reserve(SectionArray.Num());
 
@@ -703,6 +707,7 @@ void FGeometryCollectionSceneProxyBase::GetDynamicRayTracingInstances(FRayTracin
 			RayTracingInstance.Geometry = &PartRayTracingGeometries[SectionIndex];
 			RayTracingInstance.InstanceTransformsView = MakeConstArrayView(&PartInstanceTransforms[Section.TransformIndex], 1);
 			RayTracingInstance.MaterialsView = MakeConstArrayView(&CachedPartRayTracingMaterials[SectionIndex], 1);
+			RayTracingInstance.bInstanceMaskAndFlagsDirty = bInstanceMaskAndFlagsDirty;
 
 			for (int32 ViewIndex = 0; ViewIndex < Views.Num(); ViewIndex++)
 			{
