Prevent the global on-demand TOC from being mounted since this TOC is auto mounted at startup.

#rnx
#tests Manual testing on Win64
#Synced-CL 48408831

diff --git a/Engine/Source/Runtime/Experimental/IoStore/OnDemand/Private/OnDemandIoStore.cpp b/Engine/Source/Runtime/Experimental/IoStore/OnDemand/Private/OnDemandIoStore.cpp
--- a/Engine/Source/Runtime/Experimental/IoStore/OnDemand/Private/OnDemandIoStore.cpp
+++ b/Engine/Source/Runtime/Experimental/IoStore/OnDemand/Private/OnDemandIoStore.cpp
@@ -450,7 +450,7 @@ FIoStatus FOnDemandIoStore::InitializePostHotfix()
 		{
 			MountArgs.Emplace(FOnDemandMountArgs
 			{
-				.MountId = Config.TocFilePath,
+				.MountId = TEXT("global"),
 				.FilePath = Config.TocFilePath,
 				.HostGroup = HostGroup,
 				.Options = EOnDemandMountOptions::StreamOnDemand
@@ -466,7 +466,7 @@ FIoStatus FOnDemandIoStore::InitializePostHotfix()
 		}
 		MountArgs.Emplace(FOnDemandMountArgs
 		{
-			.MountId = Config.ServiceUrls[0] / Config.TocPath,
+			.MountId = TEXT("global"),
 			.TocRelativeUrl = Config.TocPath,
 			.HostGroup = HostGroup.ConsumeValueOrDie(),
 			.Options = EOnDemandMountOptions::StreamOnDemand
@@ -496,7 +496,8 @@ FIoStatus FOnDemandIoStore::InitializePostHotfix()
 
 	if (MountArgs)
 	{
-		Mount(MoveTemp(MountArgs.GetValue()), [](FOnDemandMountResult MountResult)
+		// Directly enqueue to skip the global TOC validation
+		EnqueueMountRequest(MoveTemp(MountArgs.GetValue()), [](FOnDemandMountResult MountResult)
 		{
 			MountResult.LogResult();
 		});
@@ -578,17 +579,18 @@ FOnDemandRegisterHostGroupResult FOnDemandIoStore::RegisterHostGroup(FOnDemandRe
 
 void FOnDemandIoStore::Mount(FOnDemandMountArgs&& Args, FOnDemandMountCompleted&& OnCompleted)
 {
-	FSharedMountRequest MountRequest = MakeShared<FMountRequest>();
-	MountRequest->Args = MoveTemp(Args);
-	MountRequest->OnCompleted = MoveTemp(OnCompleted);
-
-	UE_LOG(LogIoStoreOnDemand, Log, TEXT("Enqueing mount request, MountId='%s'"), *MountRequest->Args.MountId);
+	// The global.uondemandtoc is currently a special TOC but will go away once we start with per container TOCs
+	if (FPathViews::GetBaseFilename(Args.FilePath) == TEXTVIEW("global"))
 	{
-		UE::TUniqueLock Lock(RequestMutex);
-		MountRequests.Add(MoveTemp(MountRequest));
+		UE_LOG(LogIoStoreOnDemand, Warning, TEXT("Trying to mount the global container, please update install bundle configuration"));
+		return OnCompleted(FOnDemandMountResult
+		{
+			.MountId	= MoveTemp(Args.MountId),
+			.Status		= FIoStatus::Ok
+		});
 	}
 
-	TryEnterTickLoop();
+	EnqueueMountRequest(MoveTemp(Args), MoveTemp(OnCompleted));
 }
 
 FOnDemandInstallRequest FOnDemandIoStore::Install(
@@ -1162,10 +1164,10 @@ bool FOnDemandIoStore::Tick()
 
 	bool bTicked = LocalMountRequests.IsEmpty() == false;
 
-	// Tick mount request(s)
+	// Process mount request(s)
 	for (FSharedMountRequest& Request : LocalMountRequests)
 	{
-		FIoStatus MountStatus = TickMountRequest(*Request);
+		FIoStatus MountStatus = ProcessMountRequest(*Request);
 
 		{
 			UE::TUniqueLock Lock(RequestMutex);
@@ -1184,9 +1186,23 @@ bool FOnDemandIoStore::Tick()
 	return bTicked;
 }
 
-FIoStatus FOnDemandIoStore::TickMountRequest(FMountRequest& MountRequest)
+void FOnDemandIoStore::EnqueueMountRequest(FOnDemandMountArgs&& Args, FOnDemandMountCompleted&& OnCompleted)
+{
+	FSharedMountRequest MountRequest = MakeShared<FMountRequest>();
+	MountRequest->Args = MoveTemp(Args);
+	MountRequest->OnCompleted = MoveTemp(OnCompleted);
+
+	{
+		UE::TUniqueLock Lock(RequestMutex);
+		MountRequests.Add(MoveTemp(MountRequest));
+	}
+
+	TryEnterTickLoop();
+}
+
+FIoStatus FOnDemandIoStore::ProcessMountRequest(FMountRequest& MountRequest)
 {
-	UE_LOG(LogIoStoreOnDemand, Verbose, TEXT("Ticking mount request, MountId='%s'"), *MountRequest.Args.MountId);
+	UE_LOG(LogIoStoreOnDemand, Verbose, TEXT("Processing mount request, MountId='%s'"), *MountRequest.Args.MountId);
 
 	bool bWasLoaded = false;
 	if (FIoStatus CertStatus = LoadDefaultHttpCertificates(bWasLoaded); CertStatus.IsOk() == false) 
diff --git a/Engine/Source/Runtime/Experimental/IoStore/OnDemand/Private/OnDemandIoStore.h b/Engine/Source/Runtime/Experimental/IoStore/OnDemand/Private/OnDemandIoStore.h
--- a/Engine/Source/Runtime/Experimental/IoStore/OnDemand/Private/OnDemandIoStore.h
+++ b/Engine/Source/Runtime/Experimental/IoStore/OnDemand/Private/OnDemandIoStore.h
@@ -415,7 +415,8 @@ private:
 	void								TryEnterTickLoop();
 	void								TickLoop();
 	bool								Tick();
-	FIoStatus							TickMountRequest(FMountRequest& MountRequest);
+	void								EnqueueMountRequest(FOnDemandMountArgs&& Args, FOnDemandMountCompleted&& OnCompleted);
+	FIoStatus							ProcessMountRequest(FMountRequest& MountRequest);
 	void								CompleteMountRequest(FMountRequest& MountRequest, FOnDemandMountResult&& MountResult);
 	void								OnEncryptionKeyAdded(const FGuid& Id, const FAES::FAESKey& Key);
 	void								OnHostGroupRegistered(const FName& HostGroup);
