Prevent 'iostore.MountFile' from crashing if it attempts to mount more than 1 container file.

#rnx

- Previously, HostGroup was being consumed in every iteration of the loop, after the first iteration it would be invalid and then crash (Die) in the second iteration.
- Now we consume it once, outside of the loop and apply it to all mount args as needed.

#Synced-CL 48498703

diff --git a/Engine/Source/Runtime/Experimental/IoStore/OnDemandCore/Private/IoStoreOnDemand.cpp b/Engine/Source/Runtime/Experimental/IoStore/OnDemandCore/Private/IoStoreOnDemand.cpp
--- a/Engine/Source/Runtime/Experimental/IoStore/OnDemandCore/Private/IoStoreOnDemand.cpp
+++ b/Engine/Source/Runtime/Experimental/IoStore/OnDemandCore/Private/IoStoreOnDemand.cpp
@@ -270,19 +270,21 @@ static void MountFile(const TArray<FString>& Args)
 	FString HostPath = Args[2];
 	HostPath.TrimQuotesInline();
 	FStringView TocRelativeUrl;
-	TIoStatusOr<FOnDemandHostGroup> HostGroup = SplitHostAndTocUrl(HostPath, TocRelativeUrl);
-	if (HostGroup.IsOk() == false)
+	TIoStatusOr<FOnDemandHostGroup> HostResult = SplitHostAndTocUrl(HostPath, TocRelativeUrl);
+	if (HostResult.IsOk() == false)
 	{
-		UE_LOG(LogIoStoreOnDemand, Error, TEXT("%s"), *HostGroup.Status().ToString());
+		UE_LOG(LogIoStoreOnDemand, Error, TEXT("%s"), *HostResult.Status().ToString());
 		return;
 	}
 
+	FOnDemandHostGroup HostGroup = HostResult.ConsumeValueOrDie();
+
 	const FString& MountId = Args.Num() > 3 ? Args[3] : FilePaths[0];
 	for (const FString& FilePath : FilePaths)
 	{
 		FOnDemandMountArgs& MountArgs	= AllMountArgs.AddDefaulted_GetRef();
 		MountArgs.FilePath				= FilePath;
-		MountArgs.HostGroup				= HostGroup.ConsumeValueOrDie();
+		MountArgs.HostGroup				= HostGroup;
 		MountArgs.TocRelativeUrl		= TocRelativeUrl;
 		MountArgs.MountId				= MountId;
 	}
