Remove unused cvar 'ias.DistributedEndpointFallbackUrl'

#rnx

- This cvar was only used in the old backend which has been removed.

#Synced-CL 48498703

diff --git a/Engine/Source/Runtime/Experimental/IoStore/OnDemand/Private/OnDemandConfig.cpp b/Engine/Source/Runtime/Experimental/IoStore/OnDemand/Private/OnDemandConfig.cpp
--- a/Engine/Source/Runtime/Experimental/IoStore/OnDemand/Private/OnDemandConfig.cpp
+++ b/Engine/Source/Runtime/Experimental/IoStore/OnDemand/Private/OnDemandConfig.cpp
@@ -38,14 +38,6 @@ static FAutoConsoleVariableRef CVar_IoStoreOnDemandInstallCacheEnabled(
 	ECVF_ReadOnly
 );
 
-/** Temp cvar to allow the fallback url to be hotfixed in case of problems */
-static FString GDistributedEndpointFallbackUrl;
-static FAutoConsoleVariableRef CVar_DistributedEndpointFallbackUrl(
-	TEXT("ias.DistributedEndpointFallbackUrl"),
-	GDistributedEndpointFallbackUrl,
-	TEXT("CDN url to be used if a distributed endpoint cannot be reached (overrides IoStoreOnDemand.ini)")
-);
-
 ////////////////////////////////////////////////////////////////////////////////
 static bool ParseEncryptionKeyParam(const FString& Param, FGuid& OutKeyGuid, FAES::FAESKey& OutKey)
 {
