Revert exposing FIoDispatcher::GetFilename as public method and use FIoDispatcherInternal instead.

#rnx
#Synced-CL 48546254

diff --git a/Engine/Source/Runtime/Core/Internal/IO/IoDispatcherInternal.h b/Engine/Source/Runtime/Core/Internal/IO/IoDispatcherInternal.h
--- a/Engine/Source/Runtime/Core/Internal/IO/IoDispatcherInternal.h
+++ b/Engine/Source/Runtime/Core/Internal/IO/IoDispatcherInternal.h
@@ -24,5 +24,8 @@ public:
 		FIoDispatcher::Get().MountContainerMeta(FilePath);
 	}
 
-	FUtf8StringView GetFilename(const FIoChunkId& ChunkId, FUtf8StringBuilderBase& OutFilename, FUtf8StringView& OutContainerName);
+	static FUtf8StringView GetFilename(const FIoChunkId& ChunkId, FUtf8StringBuilderBase& OutFilename, FUtf8StringView& OutContainerName)
+	{
+		return FIoDispatcher::Get().GetFilename(ChunkId, OutFilename, OutContainerName);
+	}
 };
diff --git a/Engine/Source/Runtime/Experimental/IoStore/OnDemandCore/Private/OnDemandError.cpp b/Engine/Source/Runtime/Experimental/IoStore/OnDemandCore/Private/OnDemandError.cpp
--- a/Engine/Source/Runtime/Experimental/IoStore/OnDemandCore/Private/OnDemandError.cpp
+++ b/Engine/Source/Runtime/Experimental/IoStore/OnDemandCore/Private/OnDemandError.cpp
@@ -1,7 +1,7 @@
 // Copyright Epic Games, Inc. All Rights Reserved.
 
 #include "IO/OnDemandError.h"
-#include "IO/IoStore.h"
+#include "IO/IoDispatcherInternal.h"
 
 UE_DEFINE_ERROR_MODULE(IoStoreOnDemand);
 UE_DEFINE_ERROR(HttpError, IoStoreOnDemand);
@@ -20,31 +20,30 @@ namespace UE::UnifiedError::IoStore
 {
 } // namespace UE::UnifiedError::IoStore
 
-
-
 ////////////////////////////////////////////////////////////////////////////////
-
-
 void SerializeForLog(FCbWriter& Writer, const FIoChunkId& ChunkId)
 {
+	FUtf8StringBuilderBase Sb;
+	FUtf8StringView ContainerName;
+	FUtf8StringView Filename = FIoDispatcherInternal::GetFilename(ChunkId, Sb, ContainerName);
+
 	Writer.BeginObject();
 	Writer.AddString(ANSITEXTVIEW("$type"), ANSITEXTVIEW("FIoChunkId"));
-	FUtf8StringBuilderBase Filename;
-	FUtf8StringView ContainerName;
-	if (!FIoDispatcher::Get().GetFilename(ChunkId, Filename, ContainerName).IsEmpty())
+	if (Filename.IsEmpty())
 	{
-		Writer.AddString(ANSITEXTVIEW("$format"), TEXT("({Filename}:{ContainerName}:{ChunkId})"));
-		Writer.AddString(ANSITEXTVIEW("Filename"), Filename);
-		Writer.AddString(ANSITEXTVIEW("ContainerName"), ContainerName);
+		Writer.AddString(ANSITEXTVIEW("$format"), TEXT("{ChunkId}"));
 	}
 	else
 	{
-		Writer.AddString(ANSITEXTVIEW("$format"), TEXT("{ChunkId}"));
+		Writer.AddString(ANSITEXTVIEW("$format"), TEXT("({Filename}:{ContainerName}:{ChunkId})"));
+		Writer.AddString(ANSITEXTVIEW("Filename"), Filename);
+		Writer.AddString(ANSITEXTVIEW("ContainerName"), ContainerName);
 	}
 	Writer.AddString(ANSITEXTVIEW("ChunkId"), LexToString(ChunkId));
 	Writer.EndObject();
 }
 
+////////////////////////////////////////////////////////////////////////////////
 void SerializeForLog(FCbWriter& Writer, const UE::UnifiedError::IoStoreOnDemand::FChunkMissingErrorContext& Ctx)
 {
 	Writer.BeginObject();
