Fix number of solver iterations on GCs being incorrectly set to 127, this was affecting vehicle simulation and general destruction performance

#Synced-CL 48355415

diff --git a/Engine/Source/Runtime/Experimental/Chaos/Public/GeometryCollection/GeometryCollectionSimulationCoreTypes.h b/Engine/Source/Runtime/Experimental/Chaos/Public/GeometryCollection/GeometryCollectionSimulationCoreTypes.h
--- a/Engine/Source/Runtime/Experimental/Chaos/Public/GeometryCollection/GeometryCollectionSimulationCoreTypes.h
+++ b/Engine/Source/Runtime/Experimental/Chaos/Public/GeometryCollection/GeometryCollectionSimulationCoreTypes.h
@@ -421,9 +421,9 @@ PRAGMA_ENABLE_DEPRECATION_WARNINGS
 	bool bUseDamagePropagation : 1;
 	bool bOptimizeForRuntimeMemory : 1;
 
-	uint8 PositionSolverIterations;
-	uint8 VelocitySolverIterations;
-	uint8 ProjectionSolverIterations;
+	int32 PositionSolverIterations;
+	int32 VelocitySolverIterations;
+	int32 ProjectionSolverIterations;
 
 	float BreakDamagePropagationFactor;
 	float ShockDamagePropagationFactor;
