Replaced named placeholders with positional in logging to match how the args are being provided.

#rnx

#Synced-CL 48498703

diff --git a/Engine/Source/Runtime/Experimental/IoStore/OnDemand/Private/OnDemandHttpIoDispacher.cpp b/Engine/Source/Runtime/Experimental/IoStore/OnDemand/Private/OnDemandHttpIoDispacher.cpp
--- a/Engine/Source/Runtime/Experimental/IoStore/OnDemand/Private/OnDemandHttpIoDispacher.cpp
+++ b/Engine/Source/Runtime/Experimental/IoStore/OnDemand/Private/OnDemandHttpIoDispacher.cpp
@@ -93,7 +93,7 @@ void OnHttpClientCVarChanged(IConsoleVariable* CVar)
 	if (ensure(CVar != nullptr))
 	{
 		const FString CVarName = IConsoleManager::Get().FindConsoleObjectName(CVar);
-		UE_LOGFMT(LogIas, Log, "Existing http client config invalidated as cvar '{CvarName}' has changed", CVarName);
+		UE_LOGFMT(LogIas, Log, "Existing http client config invalidated as cvar '{0}' has changed", CVarName);
 
 		OnRecreateHttpClient.Broadcast();
 	}
@@ -1470,7 +1470,7 @@ void FHttpDispatcher::UpdateThreadPriorityIfNeeded()
 	EThreadPriority DesiredThreadPriority = (EThreadPriority)GThreadPriorities[ThreadPriorityIndex];
 	if (DesiredThreadPriority != ThreadPriority)
 	{
-		UE_LOGFMT(LogHttpIoDispatcher, Log, "Updated IoService.Http thread priority to '{Priority}'", LexToString(DesiredThreadPriority));
+		UE_LOGFMT(LogHttpIoDispatcher, Log, "Updated IoService.Http thread priority to '{0}'", LexToString(DesiredThreadPriority));
 
 		FPlatformProcess::SetThreadPriority(DesiredThreadPriority);
 		ThreadPriority = DesiredThreadPriority;
