Dataflow : deprecate unsused methods and clean up some already deprecated ones

#Synced-CL 48458398

diff --git a/Engine/Source/Runtime/Experimental/Dataflow/Core/Private/Dataflow/DataflowConnection.cpp b/Engine/Source/Runtime/Experimental/Dataflow/Core/Private/Dataflow/DataflowConnection.cpp
--- a/Engine/Source/Runtime/Experimental/Dataflow/Core/Private/Dataflow/DataflowConnection.cpp
+++ b/Engine/Source/Runtime/Experimental/Dataflow/Core/Private/Dataflow/DataflowConnection.cpp
@@ -11,22 +11,6 @@
 const UE::Dataflow::FConnectionKey UE::Dataflow::FConnectionKey::Invalid = { (uint32)INDEX_NONE, INDEX_NONE, (uint32)INDEX_NONE };
 const UE::Dataflow::FPin UE::Dataflow::FPin::InvalidPin = { UE::Dataflow::FPin::EDirection::NONE, NAME_None, NAME_None };
 
-FDataflowConnection::FDataflowConnection(UE::Dataflow::FPin::EDirection InDirection, FName InType, FName InName, FDataflowNode* InOwningNode, const FProperty* InProperty, FGuid InGuid)
-	: OwningNode(InOwningNode)
-	, Property(InProperty)
-	, Guid(InGuid)
-	, Direction(InDirection)
-	, OriginalType(InType)
-	, Name(InName)
-{
-	ensure(!OriginalType.ToString().Contains(TEXT(" ")));
-	ensure(!Type.ToString().Contains(TEXT(" ")));
-	// this cannot be in the initializer because of  warning V670 about order of initialization
-	Offset = OwningNode ? OwningNode->GetPropertyOffset(Name) : INDEX_NONE;
-	SetTypeInternal(InType);
-	InitFromType();
-}
-
 FDataflowConnection::FDataflowConnection(UE::Dataflow::FPin::EDirection InDirection, const UE::Dataflow::FConnectionParameters& Params)
 	: OwningNode(Params.Owner)
 	, Property(Params.Property)
@@ -284,11 +268,10 @@ FString FDataflowConnection::GetPropertyTypeNameTooltip() const
 {
 #if WITH_EDITORONLY_DATA
 	FString TypeNameStr = Type.ToString();
-	if (bIsAnyType)
+	if (bIsAnyType && ensure(Property))
 	{
 		if (!HasConcreteType())
 		{
-			check(Property);
 			TypeNameStr = TEXT("Wildcard");
 		}
 		if (Property->GetClass()->IsChildOf(FStructProperty::StaticClass()))
diff --git a/Engine/Source/Runtime/Experimental/Dataflow/Core/Private/Dataflow/DataflowInputOutput.cpp b/Engine/Source/Runtime/Experimental/Dataflow/Core/Private/Dataflow/DataflowInputOutput.cpp
--- a/Engine/Source/Runtime/Experimental/Dataflow/Core/Private/Dataflow/DataflowInputOutput.cpp
+++ b/Engine/Source/Runtime/Experimental/Dataflow/Core/Private/Dataflow/DataflowInputOutput.cpp
@@ -18,14 +18,6 @@ namespace UE::Dataflow::Private
 FDataflowInput FDataflowInput::NoOpInput = FDataflowInput();
 FDataflowOutput FDataflowOutput::NoOpOutput = FDataflowOutput();
 
-
-FDataflowInput::FDataflowInput(const UE::Dataflow::FInputParameters& Param, FGuid InGuid)
-	: FDataflowConnection(UE::Dataflow::FPin::EDirection::INPUT, Param)
-	, Connection(nullptr)
-{
-	Guid = InGuid;
-}
-
 FDataflowInput::FDataflowInput(const UE::Dataflow::FInputParameters& Param)
 	: FDataflowConnection(UE::Dataflow::FPin::EDirection::INPUT, Param)
 	, Connection(nullptr)
@@ -163,14 +155,6 @@ void* FDataflowArrayInput::RealAddress() const
 //
 //
 
-
-FDataflowOutput::FDataflowOutput(const UE::Dataflow::FOutputParameters& Param, FGuid InGuid)
-	: FDataflowConnection(UE::Dataflow::FPin::EDirection::OUTPUT, Param)
-{
-	Guid = InGuid;
-	OutputLock = MakeShared<FCriticalSection>();
-}
-
 FDataflowOutput::FDataflowOutput(const UE::Dataflow::FOutputParameters& Param)
 	: FDataflowConnection(UE::Dataflow::FPin::EDirection::OUTPUT, Param)
 {
@@ -327,6 +311,7 @@ bool FDataflowOutput::EvaluateImpl(UE::Dataflow::FContext& Context) const
 	return true;
 }
 
+// Deprecated 5.8
 TFuture<bool> FDataflowOutput::EvaluateParallel(UE::Dataflow::FContext& Context) const
 {
 	return Async(EAsyncExecution::TaskGraph, [&]() -> bool { return this->Evaluate(Context); });
@@ -387,7 +372,7 @@ void FDataflowOutput::ForwardInput(const FDataflowInput* Input, UE::Dataflow::FC
 {
 	if (Property && OwningNode)
 	{
-		if (Input->GetConnectedOutputs().Num())
+		if (Input && Input->GetConnectedOutputs().Num())
 		{
 			ensure(Input->GetType() == GetType());
 			ensure(Input->GetConnectedOutputs().Num() == 1);
diff --git a/Engine/Source/Runtime/Experimental/Dataflow/Core/Private/Dataflow/DataflowNode.cpp b/Engine/Source/Runtime/Experimental/Dataflow/Core/Private/Dataflow/DataflowNode.cpp
--- a/Engine/Source/Runtime/Experimental/Dataflow/Core/Private/Dataflow/DataflowNode.cpp
+++ b/Engine/Source/Runtime/Experimental/Dataflow/Core/Private/Dataflow/DataflowNode.cpp
@@ -1424,17 +1424,6 @@ TArray<FString> FDataflowNode::GetPinMetaData(const FName& PropertyFullName, con
 	return TArray<FString>();
 }
 
-void FDataflowNode::CopyNodeProperties(const TSharedPtr<FDataflowNode> CopyFromDataflowNode)
-{
-	TArray<uint8> NodeData;
-
-	FObjectWriter ArWriter(NodeData);
-	CopyFromDataflowNode->SerializeInternal(ArWriter);
-
-	FObjectReader ArReader(NodeData);
-	this->SerializeInternal(ArReader);
-}
-
 TSharedPtr<UE::Dataflow::FGraph> FDataflowNode::GetDataflowGraph() const
 {
 	return WeakDataflowGraph.Pin();
@@ -1526,17 +1515,6 @@ void FDataflowNode::NotifyConnectionTypeChanged(FDataflowConnection* Connection)
 	if (Connection && Connection->IsAnyType())
 	{
 		OnConnectionTypeChanged(*Connection);
-
-PRAGMA_DISABLE_DEPRECATION_WARNINGS
-		if (Connection->GetDirection() == UE::Dataflow::FPin::EDirection::INPUT)
-		{
-			OnInputTypeChanged((FDataflowInput*)Connection);
-		}
-		if (Connection->GetDirection() == UE::Dataflow::FPin::EDirection::OUTPUT)
-		{
-			OnOutputTypeChanged((FDataflowOutput*)Connection);
-		}
-PRAGMA_ENABLE_DEPRECATION_WARNINGS
 	}
 }
 
diff --git a/Engine/Source/Runtime/Experimental/Dataflow/Core/Public/Dataflow/DataflowConnection.h b/Engine/Source/Runtime/Experimental/Dataflow/Core/Public/Dataflow/DataflowConnection.h
--- a/Engine/Source/Runtime/Experimental/Dataflow/Core/Public/Dataflow/DataflowConnection.h
+++ b/Engine/Source/Runtime/Experimental/Dataflow/Core/Public/Dataflow/DataflowConnection.h
@@ -193,8 +193,6 @@ protected:
 
 public:
 	FDataflowConnection() {};
-	UE_DEPRECATED(5.5, "Deprecated constructor : use FConnectionParameters to pass parameters")
-	DATAFLOWCORE_API FDataflowConnection(UE::Dataflow::FPin::EDirection Direction, FName InType, FName InName, FDataflowNode* OwningNode = nullptr, const FProperty* InProperty = nullptr, FGuid InGuid = FGuid::NewGuid());
 	DATAFLOWCORE_API FDataflowConnection(UE::Dataflow::FPin::EDirection Direction, const UE::Dataflow::FConnectionParameters& Params);
 	virtual ~FDataflowConnection() = default;
 
diff --git a/Engine/Source/Runtime/Experimental/Dataflow/Core/Public/Dataflow/DataflowInputOutput.h b/Engine/Source/Runtime/Experimental/Dataflow/Core/Public/Dataflow/DataflowInputOutput.h
--- a/Engine/Source/Runtime/Experimental/Dataflow/Core/Public/Dataflow/DataflowInputOutput.h
+++ b/Engine/Source/Runtime/Experimental/Dataflow/Core/Public/Dataflow/DataflowInputOutput.h
@@ -57,9 +57,6 @@ private:
 	bool bIsRequired = false;
 
 public:
-	UE_DEPRECATED(5.5, "Deprecated constructor : Guid is now passed through FInputParameters")
-	DATAFLOWCORE_API FDataflowInput(const UE::Dataflow::FInputParameters& Param, FGuid InGuid);
-
 	DATAFLOWCORE_API FDataflowInput(const UE::Dataflow::FInputParameters& Param = {});
 
 	virtual bool IsConnected() const override;
@@ -95,6 +92,7 @@ public:
 	DATAFLOWCORE_API void PullValue(UE::Dataflow::FContext& Context) const;
 
 	template<class T>
+	UE_DEPRECATED(5.8, "This code path is not supported, use context to evaluate nodes and connections asynchronously")
 	TFuture<const T&> GetValueParallel(UE::Dataflow::FContext& Context, const T& Default) const;
 
 	virtual void Invalidate(const UE::Dataflow::FTimestamp& ModifiedTimestamp = UE::Dataflow::FTimestamp::Current()) override;
@@ -153,9 +151,6 @@ struct FDataflowOutput : public FDataflowConnection
 	
 	TArray< FDataflowInput* > Connections;
 
-	UE_DEPRECATED(5.5, "Use PassthroughKey instead")
-	uint32 PassthroughOffset = INDEX_NONE;
-
 	UE::Dataflow::FConnectionKey PassthroughKey;
 
 protected:
@@ -167,9 +162,6 @@ public:
 	
 	mutable TSharedPtr<FCriticalSection> OutputLock;
 	
-	UE_DEPRECATED(5.5, "Deprecated constructor : Guid is now passed through FOutputParameters")
-	DATAFLOWCORE_API FDataflowOutput(const UE::Dataflow::FOutputParameters& Param, FGuid InGuid);
-
 	DATAFLOWCORE_API FDataflowOutput(const UE::Dataflow::FOutputParameters& Param = {});
 
 	PRAGMA_DISABLE_DEPRECATION_WARNINGS
@@ -194,12 +186,6 @@ public:
 
 	virtual void GetConnections(TArray<FDataflowConnection*>& OutConnections) const override;
 
-	UE_DEPRECATED(5.5, "Use SetPassthroughInput instead")
-	virtual void SetPassthroughOffset(const uint32 InPassthroughOffset)
-	{
-		SetPassthroughInput(UE::Dataflow::FConnectionKey(InPassthroughOffset, INDEX_NONE, INDEX_NONE));
-	}
-
 	DATAFLOWCORE_API FDataflowOutput& SetPassthroughInput(const UE::Dataflow::FConnectionReference& Reference);
 	DATAFLOWCORE_API FDataflowOutput& SetPassthroughInput(const UE::Dataflow::FConnectionKey& Key);
 
@@ -354,6 +340,7 @@ public:
 	
 	DATAFLOWCORE_API bool Evaluate(UE::Dataflow::FContext& Context) const;
 
+	UE_DEPRECATED(5.8, "This code path is not supported, use context to evaluate nodes and connections asynchronously")
 	DATAFLOWCORE_API TFuture<bool> EvaluateParallel(UE::Dataflow::FContext& Context) const;
 
 	DATAFLOWCORE_API virtual void Invalidate(const UE::Dataflow::FTimestamp& ModifiedTimestamp = UE::Dataflow::FTimestamp::Current()) override;
@@ -474,6 +461,7 @@ typename UE::Dataflow::FContextValue FDataflowInput::GetValueFromAnyType(UE::Dat
 }
 
 template<class T>
+UE_DEPRECATED(5.8, "This code path is not supported, use context to evaluate nodes and connections asynchronously")
 TFuture<const T&> FDataflowInput::GetValueParallel(UE::Dataflow::FContext& Context, const T& Default) const
 {
 	return Async(EAsyncExecution::TaskGraph, [&]() -> const T& { return this->GetValue<T>(Context, Default); });
diff --git a/Engine/Source/Runtime/Experimental/Dataflow/Core/Public/Dataflow/DataflowNode.h b/Engine/Source/Runtime/Experimental/Dataflow/Core/Public/Dataflow/DataflowNode.h
--- a/Engine/Source/Runtime/Experimental/Dataflow/Core/Public/Dataflow/DataflowNode.h
+++ b/Engine/Source/Runtime/Experimental/Dataflow/Core/Public/Dataflow/DataflowNode.h
@@ -56,9 +56,6 @@ PRAGMA_ENABLE_DEPRECATION_WARNINGS
 	friend class UE::Dataflow::FGraph;
 	friend struct FDataflowConnection;
 
-	UE_DEPRECATED(5.6, "Outputs type has changed and has been made private (see ExpandedOutputs)")
-	TMap< int, FDataflowOutput* > Outputs;
-
 	UE_DEPRECATED(5.7, "Will be made private in 5.9, use IsColorOverriden() instead.")
 	UPROPERTY(EditAnywhere, Category = "Dataflow", meta = (InlineEditConditionToggle, DisplayAfter = "bActive"))
 	bool bOverrideColor = false;
@@ -98,14 +95,6 @@ PRAGMA_ENABLE_DEPRECATION_WARNINGS
 	DATAFLOWCORE_API FText GetPinDisplayName(const FName& PropertyName, const UE::Dataflow::FPin::EDirection Direction = UE::Dataflow::FPin::EDirection::NONE) const;
 	DATAFLOWCORE_API TArray<FString> GetPinMetaData(const FName& PropertyName, const UE::Dataflow::FPin::EDirection Direction = UE::Dataflow::FPin::EDirection::NONE) const;
 	virtual TArray<UE::Dataflow::FRenderingParameter> GetRenderParameters() const { return GetRenderParametersImpl(); }
-	// Copy node property values from another node
-	UE_DEPRECATED(5.4, "FDataflowNode::CopyNodeProperties is deprecated.")
-	DATAFLOWCORE_API void CopyNodeProperties(const TSharedPtr<FDataflowNode> CopyFromDataflowNode);
-
-	UE_DEPRECATED(5.4, "FDataflowNode::IsDeprecated is deprecated.")
-	virtual bool IsDeprecated() { return false; }
-	UE_DEPRECATED(5.4, "FDataflowNode::IsExperimental is deprecated.")
-	virtual bool IsExperimental() { return false; }
 
 	/** Override if node will generate primitives from the RenderCollection.*/
 	virtual bool HasRenderCollectionPrimitives() const
@@ -290,13 +279,6 @@ PRAGMA_ENABLE_DEPRECATION_WARNINGS
 		return RegisterOutputConnection(UE::Dataflow::TConnectionReference<T>(Reference), UE::Dataflow::TConnectionReference<T>(Passthrough), PropertyName);
 	}
 
-	template <typename T>
-	UE_DEPRECATED(5.5, "PassthroughName is no longer needed to register output connections")
-	FDataflowOutput& RegisterOutputConnection(const T* Reference, const T* Passthrough, const FName& PropertyName, const FName& PassthroughName)
-	{
-		return RegisterOutputConnection(UE::Dataflow::TConnectionReference<T>(Reference), UE::Dataflow::TConnectionReference<T>(Passthrough), PropertyName);
-	}
-
 	template<typename T>
 	FDataflowInput& RegisterInputArrayConnection(const UE::Dataflow::TConnectionReference<T>& Reference, const FName& ElementPropertyName = NAME_None,
 		const FName& ArrayPropertyName = NAME_None)
@@ -470,16 +452,22 @@ PRAGMA_ENABLE_DEPRECATION_WARNINGS
 	}
 	
 	template<class T>
+	UE_DEPRECATED(5.8, "This code path is not supported, use context to evaluate nodes and connections asynchronously")
 	TFuture<const T&> GetValueParallel(UE::Dataflow::FContext& Context, const UE::Dataflow::TConnectionReference<T>& Reference) const
 	{
+		PRAGMA_DISABLE_DEPRECATION_WARNINGS
 		checkSlow(FindInput(Reference));
 		return FindInput(Reference)->template GetValueParallel<T>(Context, *static_cast<const T*>(Reference.Reference));
+		PRAGMA_ENABLE_DEPRECATION_WARNINGS
 	}
 
 	template<class T>
+	UE_DEPRECATED(5.8, "This code path is not supported, use context to evaluate nodes and connections asynchronously")
 	TFuture<const T&> GetValueParallel(UE::Dataflow::FContext& Context, const T* Reference) const
 	{
+		PRAGMA_DISABLE_DEPRECATION_WARNINGS
 		return GetValueParallel(Context, UE::Dataflow::TConnectionReference<T>(Reference));
+		PRAGMA_ENABLE_DEPRECATION_WARNINGS
 	}
 
 	/**
@@ -787,12 +775,6 @@ protected:
 	*/
 	DATAFLOWCORE_API virtual FName GetDependentConnectionType(const FDataflowConnection& SourceConnection, const FDataflowConnection& DependentConnection) const;
 
-	UE_DEPRECATED(5.6, "Use inputs and outputs SetTypeDependencyGroup() to declare dependent types")
-	virtual bool OnInputTypeChanged(const FDataflowInput* Input) { return false; }
-
-	UE_DEPRECATED(5.6, "Use inputs and outputs SetTypeDependencyGroup() to declare dependent types")
-	virtual bool OnOutputTypeChanged(const FDataflowOutput* Output) { return false; }
-
 	template <typename T>
 	bool SetInputConcreteType(const UE::Dataflow::FConnectionReference& InputReference)
 	{
