Add Request Size to install error/analytics.

#Synced-CL 48539628

diff --git a/Engine/Source/Runtime/Experimental/IoStore/OnDemand/Private/OnDemandInstallCache.cpp b/Engine/Source/Runtime/Experimental/IoStore/OnDemand/Private/OnDemandInstallCache.cpp
--- a/Engine/Source/Runtime/Experimental/IoStore/OnDemand/Private/OnDemandInstallCache.cpp
+++ b/Engine/Source/Runtime/Experimental/IoStore/OnDemand/Private/OnDemandInstallCache.cpp
@@ -1770,7 +1770,7 @@ private:
 	FResult										FlushPendingChunksImpl(const FPendingChunks& Chunks, int64 UtcAccessTicks = 0);
 	FString										GetJournalFilename() const { return CacheDirectory / TEXT("cas.jrn"); }
 	FString										GetSnapshotFilename() const { return CacheDirectory / TEXT("cas.snp"); }
-	UE::UnifiedError::IoStoreOnDemand::FInstallCacheErrorContext MakeInstallCacheErrorContext(uint64 TotalCachedBytes = 0, uint32 LineNo = __builtin_LINE());
+	UE::UnifiedError::IoStoreOnDemand::FInstallCacheErrorContext MakeInstallCacheErrorContext(uint64 TotalCachedBytes = 0, uint64 RequestedBytes = 0, uint32 LineNo = __builtin_LINE());
 
 	FOnDemandIoStore&			IoStore;
 	FDiskCacheGovernor&			Governor;
@@ -2445,7 +2445,7 @@ FResult FOnDemandInstallCache::Purge(const uint64 BytesToInstall)
 		}
 		else
 		{
-			FInstallCacheErrorContext Ctx = MakeInstallCacheErrorContext(TotalCachedBytes);
+			FInstallCacheErrorContext Ctx = MakeInstallCacheErrorContext(TotalCachedBytes, BytesToInstall);
 			Ctx.IoError = FIoStatus(EIoErrorCode::InvalidCode, TEXT("Failed to purge required size from install cache"));
 			Result = MakeError(InstallCachePurgeError::MakeError(MoveTemp(Ctx), UE::UnifiedError::EDetailFilter::All));
 		}
@@ -2480,6 +2480,7 @@ FResult FOnDemandInstallCache::PurgeAllUnreferenced(bool bDefrag, const uint64*
 	UE_LOG(LogIoStoreOnDemand, Log, TEXT("Purging install cache, MaxCacheSize=%.2lf MiB, CacheSize=%.2lf MiB, ReferencedBytes=%.2lf MiB"),
 		ToMiB(MaxCacheSize), ToMiB(TotalCachedBytes), ToMiB(TotalReferencedBytes));
 
+	// TODO: should probably default to TotalCachedBytes and not MaxCacheSize
 	const uint64	TotalBytesToPurge	= BytesToPurge ? *BytesToPurge : MaxCacheSize;
 	uint64			TotalPurgedBytes	= 0;
 	FResult			Result				= PurgeInternal(BlockInfo, TotalBytesToPurge, TotalPurgedBytes);
@@ -2505,7 +2506,7 @@ FResult FOnDemandInstallCache::PurgeAllUnreferenced(bool bDefrag, const uint64*
 		}
 		else
 		{
-			FInstallCacheErrorContext Ctx = MakeInstallCacheErrorContext();
+			FInstallCacheErrorContext Ctx = MakeInstallCacheErrorContext(TotalCachedBytes, TotalBytesToPurge);
 			Result = MakeError(InstallCachePurgeError::MakeError(MoveTemp(Ctx), UE::UnifiedError::EDetailFilter::All));
 		}
 	}
@@ -3092,7 +3093,7 @@ FResult FOnDemandInstallCache::Defrag(
 		{
 			FStringView ErrorMessage(TEXTVIEW("Failed to defrag the install cache due to too much data being referenced by the game"));
 			UE_LOGFMT(LogIoStoreOnDemand, Error, "{ErrorMessage}", ErrorMessage);
-			FInstallCacheErrorContext Ctx = MakeInstallCacheErrorContext(TotalCachedBytes);
+			FInstallCacheErrorContext Ctx = MakeInstallCacheErrorContext(TotalCachedBytes, *TotalBytesToFree);
 			Ctx.IoError = FIoStatus(EIoErrorCode::OutOfDiskSpace, ErrorMessage);
 			return Result = MakeError(InstallCacheDefragError::MakeError(MoveTemp(Ctx), UE::UnifiedError::EDetailFilter::All));
 		}
@@ -3158,7 +3159,7 @@ FResult FOnDemandInstallCache::Defrag(
 			ToMiB(MaxCacheSize));
 		UE_LOG(LogIoStoreOnDemand, Error, TEXT("%s"), *ErrorMsg);
 
-		FInstallCacheErrorContext Ctx = MakeInstallCacheErrorContext(TotalCachedBytes);
+		FInstallCacheErrorContext Ctx = MakeInstallCacheErrorContext(TotalCachedBytes, TotalBytesToFree ? *TotalBytesToFree : TotalFragmentedBytes);
 		Ctx.IoError = FIoStatus(EIoErrorCode::OutOfDiskSpace, ErrorMsg),
 			Result = MakeError(InstallCacheDefragError::MakeError(MoveTemp(Ctx), UE::UnifiedError::EDetailFilter::All));
 
@@ -3233,7 +3234,7 @@ FResult FOnDemandInstallCache::Defrag(
 			{
 				UE_LOGFMT(LogIoStoreOnDemand, Error, "Defrag failed to open CAS block for reading {Error}", FileOpenResult.GetError());
 
-				FInstallCacheErrorContext Ctx = MakeInstallCacheErrorContext(TotalCachedBytes);
+				FInstallCacheErrorContext Ctx = MakeInstallCacheErrorContext(TotalCachedBytes, TotalBytesToFree ? *TotalBytesToFree : TotalFragmentedBytes);
 				FileOpenResult.GetError().PushErrorContext(MoveTemp(Ctx));
 				return Result = MakeError(FileOpenResult.StealError());
 			}
@@ -3270,7 +3271,7 @@ FResult FOnDemandInstallCache::Defrag(
 
 				if (ReadResult.HasError())
 				{
-					FInstallCacheErrorContext Ctx = MakeInstallCacheErrorContext(TotalCachedBytes);
+					FInstallCacheErrorContext Ctx = MakeInstallCacheErrorContext(TotalCachedBytes, TotalBytesToFree ? *TotalBytesToFree : TotalFragmentedBytes);
 					ReadResult.GetError().PushErrorContext(MoveTemp(Ctx));
 					return Result = MakeError(ReadResult.StealError());
 				}
@@ -3699,7 +3700,7 @@ void FOnDemandInstallCache::CompleteRequest(FIoRequestImpl* Request, EIoErrorCod
 	BackendContext->WakeUpDispatcherThreadDelegate.Execute();
 }
 
-UE::UnifiedError::IoStoreOnDemand::FInstallCacheErrorContext FOnDemandInstallCache::MakeInstallCacheErrorContext(uint64 TotalCachedBytes, uint32 LineNo)
+UE::UnifiedError::IoStoreOnDemand::FInstallCacheErrorContext FOnDemandInstallCache::MakeInstallCacheErrorContext(uint64 TotalCachedBytes, uint64 RequestedBytes, uint32 LineNo)
 {
 	UE::UnifiedError::IoStoreOnDemand::FInstallCacheErrorContext OutCtx;
 
@@ -3707,6 +3708,7 @@ UE::UnifiedError::IoStoreOnDemand::FInstallCacheErrorContext FOnDemandInstallCac
 	OutCtx.bDiskQuerySucceeded		= FPlatformMisc::GetDiskTotalAndFreeSpace(CachRootDirectory, OutCtx.DiskTotalBytes, OutCtx.DiskFreeBytes);
 	OutCtx.MaxCacheSize				= MaxCacheSize;
 	OutCtx.CacheSize				= TotalCachedBytes;
+	OutCtx.RequestedSize			= RequestedBytes;
 	OutCtx.LineNo					= LineNo;
 
 	return OutCtx;
diff --git a/Engine/Source/Runtime/Experimental/IoStore/OnDemandCore/Public/IO/OnDemandError.h b/Engine/Source/Runtime/Experimental/IoStore/OnDemandCore/Public/IO/OnDemandError.h
--- a/Engine/Source/Runtime/Experimental/IoStore/OnDemandCore/Public/IO/OnDemandError.h
+++ b/Engine/Source/Runtime/Experimental/IoStore/OnDemandCore/Public/IO/OnDemandError.h
@@ -165,6 +165,7 @@ struct FInstallCacheErrorContext
 	FIoStatus		IoError = FIoStatus(EIoErrorCode::Unknown);
 	uint64			MaxCacheSize = 0;
 	uint64			CacheSize = 0;
+	uint64			RequestedSize = 0;
 	uint64			DiskTotalBytes = 0;
 	uint64			DiskFreeBytes = 0;
 	uint32			LineNo = uint32(__builtin_LINE());
@@ -227,6 +228,7 @@ inline void SerializeForLog(FCbWriter& Writer, const UE::UnifiedError::IoStoreOn
 	Writer.AddInteger(ANSITEXTVIEW("SystemErrorCode"), int32(Ctx.IoError.GetSystemErrorCode()));
 	Writer.AddInteger(ANSITEXTVIEW("MaxCacheSize"), Ctx.MaxCacheSize);
 	Writer.AddInteger(ANSITEXTVIEW("CacheSize"), Ctx.CacheSize);
+	Writer.AddInteger(ANSITEXTVIEW("RequestedSize"), Ctx.RequestedSize);
 	Writer.AddInteger(ANSITEXTVIEW("DiskTotalBytes"), Ctx.DiskTotalBytes);
 	Writer.AddInteger(ANSITEXTVIEW("DiskFreeBytes"), Ctx.DiskFreeBytes);
 	Writer.AddInteger(ANSITEXTVIEW("LineNo"), uint32(Ctx.LineNo));
