Add two new options to 'IasTool.exe ListTocs' to improve filtering when outputting toc info to json files.

#rnx

- These commands work when using the following command 'ListTocs -Json="..."' as they only affect the json output.
- A new filter category '-ContainerHeader' has been added which will output the 'FIoContainerHeader' for each container.
-- Currently only outputs the packageids in the header.
- A new filter category '-AllFilters' has been added which will cause all filters for json output to be enabled, which can simplify command lines.
- Refactored the code handling the filters a little. We now create the EOnDemandTocJsonOptions from the cmdline when generating the FListTocsParams structure which is passed to the main command.
-- This would let us run validation etc on the input before the command is actually run.
-- Maybe in the future we can consider adding proper enum bit field support to FContext but probably not worth the effort unless this method of cmdline arg parsing is re-used in other tools.

#Synced-CL 48498703

diff --git a/Engine/Source/Runtime/Experimental/IoStore/OnDemand/Private/Tool/ListTocs.cpp b/Engine/Source/Runtime/Experimental/IoStore/OnDemand/Private/Tool/ListTocs.cpp
--- a/Engine/Source/Runtime/Experimental/IoStore/OnDemand/Private/Tool/ListTocs.cpp
+++ b/Engine/Source/Runtime/Experimental/IoStore/OnDemand/Private/Tool/ListTocs.cpp
@@ -8,14 +8,15 @@
 #include "Async/ParallelFor.h"
 #include "Async/UniqueLock.h"
 #include "HAL/FileManager.h"
+#include "IO/IoContainerHeader.h"
 #include "IO/IoStoreOnDemand.h"
 #include "Misc/EnumClassFlags.h"
 #include "Misc/FileHelper.h"
 #include "Misc/Paths.h"
+#include "S3/S3Client.h"
 #include "Serialization/JsonSerializer.h"
 #include "Serialization/JsonWriter.h"
 #include "Serialization/MemoryReader.h"
-#include "S3/S3Client.h"
 
 namespace UE::IoStore::Tool
 {
@@ -54,8 +55,57 @@ using FJsonWriterFactory = TJsonWriterFactory<TCHAR, TPrettyJsonPrintPolicy<TCHA
  * aws_secret_access_key="<key>
  * aws_session_token="<token>"
  *
- * Note: All values must be surounded with "".
+ * Note: All values must be surrounded with "".
  */
+
+////////////////////////////////////////////////////////////////////////////////
+
+enum class EOnDemandTocJsonOptions : uint32
+{
+	Header			= (1 << 0),
+	TocEntries		= (1 << 1),
+	BlockSizes		= (1 << 2),
+	BlockHashes		= (1 << 3),
+	ContainerHeader	= (1 << 4),
+
+	All				= Header | TocEntries | BlockSizes | BlockHashes | ContainerHeader
+};
+ENUM_CLASS_FLAGS(EOnDemandTocJsonOptions);
+
+EOnDemandTocJsonOptions LexFromContext(const FContext& Context)
+{
+	if (Context.Get<bool>(TEXT("-AllFilters"), false))
+	{
+		return EOnDemandTocJsonOptions::All;
+	}
+
+	EOnDemandTocJsonOptions JsonOptions = EOnDemandTocJsonOptions::Header;
+
+	if (Context.Get<bool>(TEXT("-TocEntries"), false))
+	{
+		JsonOptions |= EOnDemandTocJsonOptions::TocEntries;
+	}
+
+	if (Context.Get<bool>(TEXT("-BlockSizes"), false))
+	{
+		JsonOptions |= EOnDemandTocJsonOptions::BlockSizes;
+	}
+
+	if (Context.Get<bool>(TEXT("-BlockSizes"), false))
+	{
+		JsonOptions |= EOnDemandTocJsonOptions::BlockHashes;
+	}
+
+	if (Context.Get<bool>(TEXT("-ContainerHeader"), false))
+	{
+		JsonOptions |= EOnDemandTocJsonOptions::ContainerHeader;
+	}
+
+	return JsonOptions;
+}
+
+////////////////////////////////////////////////////////////////////////////////
+
 struct FListTocsParams
 {
 	FString OutFile;
@@ -74,9 +124,7 @@ struct FListTocsParams
 	FString BuildVersion;
 	FString TargetPlatform;
 	FString ChunkKeys;
-	bool bTocEntries = false;
-	bool bBlockSizes = false;
-	bool bBlockHashes = false;
+	EOnDemandTocJsonOptions JsonOptions = EOnDemandTocJsonOptions::Header;
 
 	FIoStatus Validate() const;
 };
@@ -157,9 +205,8 @@ static FListTocsParams BuildListTocsParams(const FContext& Context)
 	Ret.BuildVersion			= Context.Get<FStringView>(TEXT("-BuildVersion"),			Ret.BuildVersion);
 	Ret.TargetPlatform			= Context.Get<FStringView>(TEXT("-TargetPlatform"),			Ret.TargetPlatform);
 	Ret.ChunkKeys				= Context.Get<FStringView>(TEXT("-ChunkKeys"),				Ret.ChunkKeys);
-	Ret.bTocEntries				= Context.Get<bool>(TEXT("-TocEntries"),					Ret.bTocEntries);
-	Ret.bBlockSizes				= Context.Get<bool>(TEXT("-BlockSizes"),					Ret.bBlockSizes);
-	Ret.bBlockHashes			= Context.Get<bool>(TEXT("-BlockHashes"),					Ret.bBlockHashes);
+
+	Ret.JsonOptions = LexFromContext(Context);
 
 	if (!Ret.TocUrl.IsEmpty())
 	{
@@ -205,16 +252,6 @@ static void GetChunkObjectKeys(const FOnDemandToc& Toc, FString Prefix, TArray<F
 	}
 }
 
-////////////////////////////////////////////////////////////////////////////////
-enum class EOnDemandTocJsonOptions
-{
-	Header			= (1 << 0),
-	TocEntries		= (1 << 1),
-	BlockSizes		= (1 << 2),
-	BlockHashes		= (1 << 3),
-	All				= Header | TocEntries | BlockSizes | BlockHashes
-};
-ENUM_CLASS_FLAGS(EOnDemandTocJsonOptions);
 
 ////////////////////////////////////////////////////////////////////////////////
 static void ToJson(FJsonWriter JsonWriter, const FOnDemandToc& Toc, EOnDemandTocJsonOptions Options)
@@ -245,6 +282,33 @@ static void ToJson(FJsonWriter JsonWriter, const FOnDemandToc& Toc, EOnDemandToc
 			JsonWriter->WriteValue(TEXT("EncryptionKeyGuid"), Container.EncryptionKeyGuid);
 			JsonWriter->WriteValue(TEXT("UTocHash"), LexToString(Container.UTocHash));
 
+			if (EnumHasAnyFlags(Options, EOnDemandTocJsonOptions::ContainerHeader))
+			{
+				JsonWriter->WriteObjectStart(TEXT("Header"));
+				{
+					FIoContainerHeader ContainerHeader;
+					FMemoryReaderView Ar(Container.Header);
+					Ar << ContainerHeader;
+
+					JsonWriter->WriteValue(TEXT("ContainerId"), LexToString(ContainerHeader.ContainerId));
+					
+					JsonWriter->WriteArrayStart("PackageIds");
+					for (const FPackageId& PackageId : ContainerHeader.PackageIds)
+					{
+						JsonWriter->WriteValue(LexToString(PackageId));
+					}
+					JsonWriter->WriteArrayEnd();
+
+					JsonWriter->WriteArrayStart("OptionalSegmentPackageIds");
+					for (const FPackageId& PackageId : ContainerHeader.OptionalSegmentPackageIds)
+					{
+						JsonWriter->WriteValue(LexToString(PackageId));
+					}
+					JsonWriter->WriteArrayEnd();
+				}
+				JsonWriter->WriteObjectEnd();
+			}
+
 			if (EnumHasAnyFlags(Options, EOnDemandTocJsonOptions::TocEntries))
 			{
 				JsonWriter->WriteArrayStart("Entries");
@@ -594,20 +658,7 @@ FIoStatus ListTocs(const FListTocsParams& Params)
 			JsonWriter->WriteValue(TEXT("ChunkCount"), Desc.ChunkCount);
 			JsonWriter->WriteValue(TEXT("TotalChunkSize"), Desc.TotalChunkSize);
 
-			EOnDemandTocJsonOptions JsonOptions = EOnDemandTocJsonOptions::Header;
-			if (Params.bTocEntries)
-			{
-				JsonOptions |= EOnDemandTocJsonOptions::TocEntries;
-			}
-			if (Params.bBlockSizes)
-			{
-				JsonOptions |= EOnDemandTocJsonOptions::BlockSizes;
-			}
-			if (Params.bBlockHashes)
-			{
-				JsonOptions |= EOnDemandTocJsonOptions::BlockHashes;
-			}
-			ToJson(JsonWriter, Desc.Toc, JsonOptions);
+			ToJson(JsonWriter, Desc.Toc, Params.JsonOptions);
 			JsonWriter->WriteObjectEnd();
 		}
 		JsonWriter->WriteArrayEnd();
@@ -695,9 +746,11 @@ static FCommand ListTocsCommand(
 		TArgument<FStringView>(TEXT("-BuildVersion"),	TEXT("Filter results a particular build version")),
 		TArgument<FStringView>(TEXT("-TargetPlatform"),	TEXT("Only show results for a given platform")),
 		TArgument<FStringView>(TEXT("-ChunkKeys"),		TEXT("Path to write JSON file of chunk keys to")),
+		TArgument<bool>(TEXT("-AllFilters"),			TEXT("Enabled all filters in JSON output")),
+		TArgument<bool>(TEXT("-ContainerHeader"),		TEXT("Add the container header info in JSON")),
 		TArgument<bool>(TEXT("-TocEntries"),			TEXT("Include TOC entries in JSON output")),
-		TArgument<bool>(TEXT("-BlockSizes"),			TEXT("Write block sizes to JSON file")),
-		TArgument<bool>(TEXT("-BlockHashes"),			TEXT("Add block hash values in JSON")),
+		TArgument<bool>(TEXT("-BlockSizes"),			TEXT("Write block sizes in JSON output")),
+		TArgument<bool>(TEXT("-BlockHashes"),			TEXT("Add block hash values in JSON output")),
 
 		TArgument<FStringView>(TEXT("-BucketPrefix"),	TEXT("Bucket path to enumerate")),
 		TArgument<FStringView>(TEXT("-BucketPath"),		TEXT("Alternative way to set Bucket/BucketPrefix")),
