Constraints: remove obsolete UseEvaluationGraph

#rnx

#Synced-CL 48605401

diff --git a/Engine/Source/Runtime/Experimental/Animation/Constraints/Private/ConstraintsEvaluationGraph.cpp b/Engine/Source/Runtime/Experimental/Animation/Constraints/Private/ConstraintsEvaluationGraph.cpp
--- a/Engine/Source/Runtime/Experimental/Animation/Constraints/Private/ConstraintsEvaluationGraph.cpp
+++ b/Engine/Source/Runtime/Experimental/Animation/Constraints/Private/ConstraintsEvaluationGraph.cpp
@@ -9,12 +9,6 @@
 
 namespace ConstraintsEvaluationGraph
 {
-
-static bool	bUseEvaluationGraph = true;
-static FAutoConsoleVariableRef CVarUseEvaluationGraph(
-	TEXT("Constraints.UseEvaluationGraph"),
-	bUseEvaluationGraph,
-	TEXT("Use Evaluation Graph to update constraints when manipulating.") );
 	
 static bool	bDebugGraph = false;
 static FAutoConsoleVariableRef CVarDebugEvaluationGraph(
@@ -24,11 +18,6 @@ static FAutoConsoleVariableRef CVarDebugEvaluationGraph(
 	
 }
 
-bool FConstraintsEvaluationGraph::UseEvaluationGraph()
-{
-	return ConstraintsEvaluationGraph::bUseEvaluationGraph;
-}
-
 FConstraintNode& FConstraintsEvaluationGraph::GetNode(const TWeakObjectPtr<UTickableConstraint>& InConstraint)
 {
 	const int32 Found = Nodes.IndexOfByPredicate([InConstraint](const FConstraintNode& Node)
diff --git a/Engine/Source/Runtime/Experimental/Animation/Constraints/Private/ConstraintsManager.cpp b/Engine/Source/Runtime/Experimental/Animation/Constraints/Private/ConstraintsManager.cpp
--- a/Engine/Source/Runtime/Experimental/Animation/Constraints/Private/ConstraintsManager.cpp
+++ b/Engine/Source/Runtime/Experimental/Animation/Constraints/Private/ConstraintsManager.cpp
@@ -930,7 +930,7 @@ TArray< TWeakObjectPtr<UTickableConstraint> > FConstraintsManagerController::Get
 	const UE::Private::FScopedControllerContext Context(World);
 	
 	// use evaluation graph cached data
-	if (FConstraintsEvaluationGraph::UseEvaluationGraph() && bSorted)
+	if (bSorted)
 	{
 		TArray<TWeakObjectPtr<UTickableConstraint>> SortedConstraints;
 		if (Subsystem->GetEvaluationGraph(Context.GetWorld()).GetSortedConstraints(SortedConstraints))
@@ -940,20 +940,11 @@ TArray< TWeakObjectPtr<UTickableConstraint> > FConstraintsManagerController::Get
 	}
 
 	TArray<TWeakObjectPtr<UTickableConstraint>> Constraints = Subsystem->GetConstraints(Context.GetWorld());
-
 	// Remove stale constraints (note that GCd constraints should already have been caught by UConstraintSubsystem)
 	Constraints.RemoveAll([](const TWeakObjectPtr<UTickableConstraint>& ExistingConstraint) -> bool
 	{
 		return ExistingConstraint.IsStale();
 	});
-	
-	if (!bSorted)
-	{
-		return Constraints;
-	}
-
-	UE::Constraints::Graph::SortConstraints(Context.GetWorld(), Constraints);
-	
 	return Constraints;
 }
 
@@ -1039,10 +1030,6 @@ void FConstraintsManagerController::InvalidateEvaluationGraph() const
 
 void FConstraintsManagerController::FlushEvaluationGraph() const
 {
-	if (!FConstraintsEvaluationGraph::UseEvaluationGraph())
-	{
-		return;
-	}
 	UConstraintSubsystem* SubSystem = UConstraintSubsystem::Get();
 	if (!SubSystem || !World)
 	{
diff --git a/Engine/Source/Runtime/Experimental/Animation/Constraints/Public/ConstraintsEvaluationGraph.h b/Engine/Source/Runtime/Experimental/Animation/Constraints/Public/ConstraintsEvaluationGraph.h
--- a/Engine/Source/Runtime/Experimental/Animation/Constraints/Public/ConstraintsEvaluationGraph.h
+++ b/Engine/Source/Runtime/Experimental/Animation/Constraints/Public/ConstraintsEvaluationGraph.h
@@ -67,8 +67,6 @@ public:
 
 	/* Get Sorted Constraints */
 	bool GetSortedConstraints(TArray<TWeakObjectPtr<UTickableConstraint>>& OutConstraints);
-
-	static bool UseEvaluationGraph();
 	
 private:
 
