Change min particle batch :
A particle batch for the Broad Phase is not necessirly one collisioin computatioin on the Narrow Phase.
One particle can collisde with several particles, which can have many shapes.
By setting the minimum batch to one it increase perfromance even on the low end platform.

#Synced-CL 48444396

diff --git a/Engine/Source/Runtime/Experimental/Chaos/Private/Chaos/Framework/Parallel.cpp b/Engine/Source/Runtime/Experimental/Chaos/Private/Chaos/Framework/Parallel.cpp
--- a/Engine/Source/Runtime/Experimental/Chaos/Private/Chaos/Framework/Parallel.cpp
+++ b/Engine/Source/Runtime/Experimental/Chaos/Private/Chaos/Framework/Parallel.cpp
@@ -11,7 +11,7 @@ namespace Chaos
 	CHAOS_API int32 InnerParallelForBatchSize = 0;
 	CHAOS_API int32 MinRangeBatchSize = 0;
 	CHAOS_API int32 MaxNumWorkers = 100;
-	CHAOS_API int32 SmallBatchSize = 10;
+	CHAOS_API int32 SmallBatchSize = 1;
 	CHAOS_API int32 LargeBatchSize = 100;
 	CHAOS_API int32 MinParallelTaskSize = 0;
 	
