Unified Deformable Solver:
-supported FSolverParticlesRange for Gauss-Seidel spring constraints
-added anisotropic and zero rest length attributes in the constructor

#Synced-CL 48489583

diff --git a/Engine/Source/Runtime/Experimental/Chaos/Private/Chaos/Deformable/ChaosDeformableSolver.cpp b/Engine/Source/Runtime/Experimental/Chaos/Private/Chaos/Deformable/ChaosDeformableSolver.cpp
--- a/Engine/Source/Runtime/Experimental/Chaos/Private/Chaos/Deformable/ChaosDeformableSolver.cpp
+++ b/Engine/Source/Runtime/Experimental/Chaos/Private/Chaos/Deformable/ChaosDeformableSolver.cpp
@@ -183,7 +183,7 @@ namespace Chaos::Softs
 			AllTetAlphaJArray.Reset(new TArray<FSolverReal>());
 			AllCorotatedCodEMeshArray.Reset(new TArray<FSolverReal>());
 			AllSkinEMeshArray.Reset(new TArray<FSolverReal>());
-			GSWeakConstraints.Reset(new FGaussSeidelSpringConstraints({}, {}, {}, {}, {}, GDeformableXPBDWeakConstraintParams));
+			GSWeakConstraints.Reset(new FGaussSeidelSpringConstraints({}, {}, {}, {}, {}, {}, {}, GDeformableXPBDWeakConstraintParams));
 			GSDynamicWeakConstraints.Reset(new FGaussSeidelDynamicSpringConstraints(GDeformableXPBDWeakConstraintParams));
 			GSSphereRepulsionConstraints.Reset(new FGaussSeidelParticleSphereRepulsionConstraints(Property.SphereRepulsionRadius, Property.SphereRepulsionStiffness, Evolution->Particles(), GDeformableXPBDWeakConstraintParams));
 			GSVolumeConstraints.Reset(new FGaussSeidelUnilateralVolumeConstraints(Evolution->Particles(), {}, {}));
@@ -2041,7 +2041,7 @@ namespace Chaos::Softs
 		MuscleIndexOffset.Empty();
 		AllUnconstrainedSurfaceElementsCorotatedCod.Reset(new TArray<Chaos::TVec3<int32>>());
 		AllUnconstrainedSurfaceElementsSkin.Reset(new TArray<Chaos::TVec3<int32>>());
-		GSWeakConstraints.Reset(new FGaussSeidelSpringConstraints({}, {}, {}, {}, {}, GDeformableXPBDWeakConstraintParams));
+		GSWeakConstraints.Reset(new FGaussSeidelSpringConstraints({}, {}, {}, {}, {}, {}, {}, GDeformableXPBDWeakConstraintParams));
 		GSDynamicWeakConstraints.Reset(new FGaussSeidelDynamicSpringConstraints(GDeformableXPBDWeakConstraintParams));
 		GSSphereRepulsionConstraints.Reset(new FGaussSeidelParticleSphereRepulsionConstraints(Property.SphereRepulsionRadius, Property.SphereRepulsionStiffness, Evolution->Particles(), GDeformableXPBDWeakConstraintParams));
 		GSVolumeConstraints.Reset(new FGaussSeidelUnilateralVolumeConstraints(Evolution->Particles(), {}, {}));
diff --git a/Engine/Source/Runtime/Experimental/Chaos/Public/Chaos/Deformable/GaussSeidelDynamicWeakConstraints.h b/Engine/Source/Runtime/Experimental/Chaos/Public/Chaos/Deformable/GaussSeidelDynamicWeakConstraints.h
--- a/Engine/Source/Runtime/Experimental/Chaos/Public/Chaos/Deformable/GaussSeidelDynamicWeakConstraints.h
+++ b/Engine/Source/Runtime/Experimental/Chaos/Public/Chaos/Deformable/GaussSeidelDynamicWeakConstraints.h
@@ -31,7 +31,7 @@ namespace Chaos::Softs
 		typedef typename FGaussSeidelSpringConstraints::FGaussSeidelConstraintHandle GSConstraintHandle;
 
 		FGaussSeidelDynamicSpringConstraints(const FDeformableXPBDWeightedSpringConstraintParams& InParams)
-		: Base({}, {}, {}, {}, {}, InParams)
+			: Base({}, {}, {}, {}, {}, {}, {}, InParams)
 		{
 			ConstraintsData.AddArray(&Handles);
 		}
diff --git a/Engine/Source/Runtime/Experimental/Chaos/Public/Chaos/Deformable/GaussSeidelWeakConstraints.h b/Engine/Source/Runtime/Experimental/Chaos/Public/Chaos/Deformable/GaussSeidelWeakConstraints.h
--- a/Engine/Source/Runtime/Experimental/Chaos/Public/Chaos/Deformable/GaussSeidelWeakConstraints.h
+++ b/Engine/Source/Runtime/Experimental/Chaos/Public/Chaos/Deformable/GaussSeidelWeakConstraints.h
@@ -282,7 +282,7 @@ namespace Chaos::Softs
 	};
 
 	// Spring constraint for position-based nonlinear Gauss Seidel solver [Chen et al., 2024]
-	struct FGaussSeidelSpringConstraints 
+	struct FGaussSeidelSpringConstraints
 	{
 		//TODO(Yizhou): Add unittest for Gauss Seidel Weak Constraints
 		FGaussSeidelSpringConstraints(
@@ -291,12 +291,44 @@ namespace Chaos::Softs
 			const TArray<FSolverReal>& InStiffness,
 			const TArray<TArray<int32>>& InSecondIndices,
 			const TArray<TArray<FSolverReal>>& InSecondWeights,
+			const TArray<bool>& InIsAnisotrpic,
+			const TArray<bool>& InZeroRestLength,
 			const FDeformableXPBDWeightedSpringConstraintParams& InParams
-		): DebugDrawParams(InParams)
+		) : DebugDrawParams(InParams)
 		{
 			ensureMsgf(InIndices.Num() == InSecondIndices.Num(), TEXT("Input Double Bindings have wrong size"));
 
-			ConstraintsData.Resize(0);
+			ConstraintsData.AddConstraints(InIndices.Num());
+
+			for (int32 SpringIdx = 0; SpringIdx < InIndices.Num(); ++SpringIdx)
+			{
+				FGaussSeidelSpringConstraintSingleData SingleConstraintData;
+				SingleConstraintData.SingleIndices = InIndices[SpringIdx];
+				SingleConstraintData.SingleSecondIndices = InSecondIndices[SpringIdx];
+				SingleConstraintData.SingleWeights = InWeights[SpringIdx];
+				SingleConstraintData.SingleSecondWeights = InSecondWeights[SpringIdx];
+				SingleConstraintData.SingleStiffness = InStiffness[SpringIdx];
+				SingleConstraintData.bIsAnisotropic = InIsAnisotrpic[SpringIdx];
+				SingleConstraintData.bIsZeroRestLength = InZeroRestLength[SpringIdx];
+				ConstraintsData.SetSingleConstraint(SingleConstraintData, SpringIdx);
+				for (int32 SecondLocalIdx = 0; SecondLocalIdx < InSecondIndices[SpringIdx].Num(); ++SecondLocalIdx)
+				{
+					ensureMsgf(!InIndices[SpringIdx].Contains(InSecondIndices[SpringIdx][SecondLocalIdx]), TEXT("Indices and Second Indices overlaps. Currently not supported"));
+				}
+			}
+		}
+
+		UE_DEPRECATED(5.8, "Replaced with constructor with IsAnisotropic and IsZeroRestLength.")
+			FGaussSeidelSpringConstraints(
+				const TArray<TArray<int32>>& InIndices,
+				const TArray<TArray<FSolverReal>>& InWeights,
+				const TArray<FSolverReal>& InStiffness,
+				const TArray<TArray<int32>>& InSecondIndices,
+				const TArray<TArray<FSolverReal>>& InSecondWeights,
+				const FDeformableXPBDWeightedSpringConstraintParams& InParams
+			) : DebugDrawParams(InParams)
+		{
+			ensureMsgf(InIndices.Num() == InSecondIndices.Num(), TEXT("Input Double Bindings have wrong size"));
 
 			ConstraintsData.AddConstraints(InIndices.Num());
 
@@ -329,7 +361,7 @@ namespace Chaos::Softs
 
 		virtual ~FGaussSeidelSpringConstraints() = default;
 
-		void ComputeInitialWCData(const FSolverParticles& InParticles)
+		void ComputeInitialWCData(const FSolverParticlesRange& InParticles)
 		{
 			TArray<TArray<int32>> ExtraConstraints;
 			ExtraConstraints.Init(TArray<int32>(), ConstraintsData.Size());
@@ -342,14 +374,14 @@ namespace Chaos::Softs
 				}
 				for (int32 LocalSecondIdx = 0; LocalSecondIdx < ConstraintsData.GetSecondIndices(ConstraintIdx).Num(); ++LocalSecondIdx)
 				{
-					ExtraConstraints[ConstraintIdx][LocalSecondIdx+ConstraintsData.GetIndices(ConstraintIdx).Num()] = ConstraintsData.GetSecondIndices(ConstraintIdx)[LocalSecondIdx];
+					ExtraConstraints[ConstraintIdx][LocalSecondIdx + ConstraintsData.GetIndices(ConstraintIdx).Num()] = ConstraintsData.GetSecondIndices(ConstraintIdx)[LocalSecondIdx];
 				}
 			}
 			WCIncidentElements = Chaos::Utilities::ComputeIncidentElements(ExtraConstraints, &WCIncidentElementsLocal);
 
 			//Update rest state normal and nodal weights
 			UpdateTriangleNormalAndNodalWeight(InParticles, /*bUseParticleX = */true);
-			
+
 			NoCollisionNodalWeights = NodalWeights;
 			NoCollisionConstraints = ExtraConstraints;
 			InitialWCSize = ConstraintsData.Size();
@@ -384,7 +416,12 @@ namespace Chaos::Softs
 			}
 		}
 
-		TVec3<FSolverReal> ComputeSpringEdge(const FSolverParticles& InParticles, const TArray<int32>& LocalIndices,
+		void ComputeInitialWCData(const FSolverParticles& InParticles)
+		{
+			ComputeInitialWCData(FSolverParticlesRange(const_cast<FSolverParticles*>(&InParticles), /*InOffset = */ 0, /*InRange = */ InParticles.Size()));
+		}
+
+		TVec3<FSolverReal> ComputeSpringEdge(const FSolverParticlesRange& InParticles, const TArray<int32>& LocalIndices,
 			const TArray<int32>& LocalSecondIndices, const TArray<FSolverReal>& Weight, const TArray<FSolverReal>& SecondWeight, bool bUseParticleX) const
 		{
 			TVec3<FSolverReal> SpringEdge((FSolverReal)0.);
@@ -428,6 +465,13 @@ namespace Chaos::Softs
 			return SpringEdge;
 		}
 
+		TVec3<FSolverReal> ComputeSpringEdge(const FSolverParticles& InParticles, const TArray<int32>& LocalIndices,
+			const TArray<int32>& LocalSecondIndices, const TArray<FSolverReal>& Weight, const TArray<FSolverReal>& SecondWeight, bool bUseParticleX) const
+		{
+			return ComputeSpringEdge(FSolverParticlesRange(const_cast<FSolverParticles*>(&InParticles), /*InOffset = */ 0, /*InRange = */ InParticles.Size()), 
+				LocalIndices, LocalSecondIndices, Weight, SecondWeight, bUseParticleX);
+		}
+
 		void AddWCHessian(const int32 p, const FSolverReal Dt, Chaos::PMatrix<FSolverReal, 3, 3>& ParticleHessian) const
 		{
 			if (NodalWeights[p].Num() > 0)
@@ -504,7 +548,7 @@ namespace Chaos::Softs
 			}
 		}
 
-		void VisualizeAllBindings(const FSolverParticles& InParticles, const FSolverReal Dt) const
+		void VisualizeAllBindings(const FSolverParticlesRange& InParticles, const FSolverReal Dt) const
 		{
 #if WITH_EDITOR
 			auto DoubleVert = [](Chaos::TVec3<FSolverReal> V) { return FVector3d(V.X, V.Y, V.Z); };
@@ -550,7 +594,7 @@ namespace Chaos::Softs
 #endif
 		}
 
-		void Init(const FSolverParticles& InParticles, const FSolverReal Dt)
+		void Init(const FSolverParticlesRange& InParticles, const FSolverReal Dt)
 		{
 			UpdateTriangleNormalAndNodalWeight(InParticles, /*bUseParticleX =*/ false);
 			if (DebugDrawParams.bVisualizeBindings)
@@ -559,7 +603,12 @@ namespace Chaos::Softs
 			}
 		}
 
-		void UpdateTriangleNormalAndNodalWeight(const FSolverParticles& InParticles, bool bUseParticleX)
+		void Init(const FSolverParticles& InParticles, const FSolverReal Dt)
+		{
+			Init(FSolverParticlesRange(const_cast<FSolverParticles*>(&InParticles), /*InOffset = */ 0, /*InRange = */ InParticles.Size()), Dt);
+		}
+
+		void UpdateTriangleNormalAndNodalWeight(const FSolverParticlesRange& InParticles, bool bUseParticleX)
 		{
 			for (int32 i = 0; i < ConstraintsData.Size(); i++)
 			{
@@ -983,7 +1032,7 @@ namespace Chaos::Softs
 			return ExtraConstraints;
 		}
 
-		// Deprecated, now replaced with AddWCResidual for more general cases
+		UE_DEPRECATED(5.8, "Replaced with AddWCResidual for more general cases.")
 		void AddZeroRestLengthWCResidualAndHessian(const FSolverParticles& InParticles, const int32 ConstraintIndex, const int32 LocalIndex, const FSolverReal Dt, TVec3<FSolverReal>& ParticleResidual, Chaos::PMatrix<FSolverReal, 3, 3>& ParticleHessian) const
 		{
 			const FGaussSeidelSpringConstraintSingleData& SingleData = ConstraintsData.GetSingleConstraintData(ConstraintIndex);
@@ -1017,7 +1066,7 @@ namespace Chaos::Softs
 			}
 		}
 
-		void AddWCResidual(const FSolverParticles& InParticles, const int32 ConstraintIndex, const int32 LocalIndex, const FSolverReal Dt, TVec3<FSolverReal>& ParticleResidual, Chaos::PMatrix<FSolverReal, 3, 3>& ParticleHessian) const
+		void AddWCResidual(const FSolverParticlesRange& InParticles, const int32 ConstraintIndex, const int32 LocalIndex, const FSolverReal Dt, TVec3<FSolverReal>& ParticleResidual, Chaos::PMatrix<FSolverReal, 3, 3>& ParticleHessian) const
 		{
 			const FGaussSeidelSpringConstraintSingleData& SingleData = ConstraintsData.GetSingleConstraintData(ConstraintIndex);
 
@@ -1052,6 +1101,11 @@ namespace Chaos::Softs
 			}
 		}
 
+		void AddWCResidual(const FSolverParticles& InParticles, const int32 ConstraintIndex, const int32 LocalIndex, const FSolverReal Dt, TVec3<FSolverReal>& ParticleResidual, Chaos::PMatrix<FSolverReal, 3, 3>& ParticleHessian) const
+		{
+			AddWCResidual(FSolverParticlesRange(const_cast<FSolverParticles*>(&InParticles), /*InOffset = */ 0, /*InRange = */ InParticles.Size()), ConstraintIndex, LocalIndex, Dt, ParticleResidual, ParticleHessian);
+		}
+
 		FGaussSeidelSpringConstraintData ConstraintsData;
 
 		TArray<TArray<FSolverReal>> NodalWeights;
