Two way vehicle ground interaction: Update suspension constraint and modular vehicle so vehicle applies downward force on ground bodies, wheels work better with moving platforms.

#Synced-CL 48553644

diff --git a/Engine/Plugins/Experimental/ChaosModularVehicle/Source/ChaosModularVehicleEngine/Private/ChaosModularVehicle/ModularVehicleSimulationCU.cpp b/Engine/Plugins/Experimental/ChaosModularVehicle/Source/ChaosModularVehicleEngine/Private/ChaosModularVehicle/ModularVehicleSimulationCU.cpp
--- a/Engine/Plugins/Experimental/ChaosModularVehicle/Source/ChaosModularVehicleEngine/Private/ChaosModularVehicle/ModularVehicleSimulationCU.cpp
+++ b/Engine/Plugins/Experimental/ChaosModularVehicle/Source/ChaosModularVehicleEngine/Private/ChaosModularVehicle/ModularVehicleSimulationCU.cpp
@@ -4,6 +4,7 @@
 #include "ChaosModularVehicle/ModularVehicleDefaultAsyncInput.h"
 #include "SimModule/SimModulesInclude.h"
 #include "SimModule/ModuleInput.h"
+#include "PhysicsEngine/PhysicsObjectExternalInterface.h"
 #include "PhysicsProxy/ClusterUnionPhysicsProxy.h"
 #include "PhysicsProxy/SingleParticlePhysicsProxy.h"
 #include "PhysicalMaterials/PhysicalMaterial.h"
@@ -20,8 +21,9 @@ DEFINE_LOG_CATEGORY(LogModularVehicleSim);
 bool bModularVehicle_DumpModuleTreeStructure_Enabled = false;
 FAutoConsoleVariableRef CVarModularVehicleDumpModuleTreeStructureEnabled2(TEXT("p.ModularVehicle.DumpModuleTreeStructure.Enabled"), bModularVehicle_DumpModuleTreeStructure_Enabled, TEXT("Enable/Disable logging of module tree structure every time there is a change."));
 
-#if CHAOS_DEBUG_DRAW
 FAutoConsoleVariableRef CVarChaosModularVehiclesRaycastsEnabled(TEXT("p.ModularVehicle.SuspensionRaycastsEnabled"), GModularVehicleDebugParams.SuspensionRaycastsEnabled, TEXT("Enable/Disable Suspension Raycasts."));
+
+#if CHAOS_DEBUG_DRAW
 FAutoConsoleVariableRef CVarChaosModularVehiclesShowRaycasts(TEXT("p.ModularVehicle.ShowSuspensionRaycasts"), GModularVehicleDebugParams.ShowSuspensionRaycasts, TEXT("Enable/Disable Suspension Raycast Visualisation."));
 FAutoConsoleVariableRef CVarChaosModularVehiclesShowWheelData(TEXT("p.ModularVehicle.ShowWheelData"), GModularVehicleDebugParams.ShowWheelData, TEXT("Enable/Disable Displaying Wheel Simulation Data."));
 FAutoConsoleVariableRef CVarChaosModularVehiclesShowRaycastMaterial(TEXT("p.ModularVehicle.ShowRaycastMaterial"), GModularVehicleDebugParams.ShowRaycastMaterial, TEXT("Enable/Disable Raycast Material Hit Visualisation."));
@@ -303,7 +305,7 @@ void FModularVehicleSimulation::PerformAdditionalSimWork(UWorld* InWorld, const
 							const Chaos::FSimModuleTree::FSimModuleNode& WheelNode = ModuleArray[Suspension->GetWheelSimTreeIndex()];
 
 							Chaos::FWheelBaseInterface* Wheel = static_cast<Chaos::FWheelBaseInterface*>(WheelNode.SimModule);
-							if (Wheel && HitResult.PhysMaterial.IsValid())
+							if (Wheel)
 							{
 								if (GModularVehicleDebugParams.FrictionOverride > 0)
 								{
@@ -311,7 +313,28 @@ void FModularVehicleSimulation::PerformAdditionalSimWork(UWorld* InWorld, const
 								}
 								else
 								{
-									Wheel->SetSurfaceFriction(HitResult.PhysMaterial->Friction);
+									if (HitResult.PhysMaterial.IsValid())
+									{
+										Wheel->SetSurfaceFriction(HitResult.PhysMaterial->Friction);
+									}
+								}
+
+								IPhysicsProxyBase* HitProxy = nullptr;
+								if (InWorld)
+								{
+									if (Chaos::FPhysicsObjectHandle PhysicsObject = HitResult.PhysicsObject)
+									{
+										HitProxy = Chaos::FPhysicsObjectInterface::GetProxy({ &PhysicsObject, 1 });
+									}
+
+									if (HitProxy && HitProxy->GetType() == EPhysicsProxyType::SingleParticleProxy)
+									{
+										if (FSingleParticlePhysicsProxy* ParticleProxy = static_cast<FSingleParticlePhysicsProxy*>(HitProxy))
+										{
+											Chaos::FPBDRigidParticleHandle* GroundParticle = ParticleProxy->GetHandle_LowLevel() ? ParticleProxy->GetHandle_LowLevel()->CastToRigidParticle() : nullptr;
+											Wheel->SetGroundInteraction(GroundParticle, HitResult.ImpactPoint, HitResult.ImpactNormal);
+										}
+									}
 								}
 							}
 						}
@@ -407,6 +430,12 @@ void FModularVehicleSimulation::PerformAdditionalSimWork(UWorld* InWorld, const
 						HitDistance = HitResult.Distance - WheelRadius;
 					}
 
+					IPhysicsProxyBase* HitProxy = nullptr;
+					if (Chaos::FPhysicsObjectHandle PhysicsObject = HitResult.PhysicsObject)
+					{
+						HitProxy = Chaos::FPhysicsObjectInterface::GetProxy({ &PhysicsObject, 1 });
+					}
+
 					const TEnumAsByte<EPhysicalSurface> DefaultSurfaceType = EPhysicalSurface::SurfaceType_Default;
 					FSuspensionTargetPoint TargetPoint(
 						HitPoint
@@ -414,6 +443,7 @@ void FModularVehicleSimulation::PerformAdditionalSimWork(UWorld* InWorld, const
 						, HitDistance
 						, HitResult.bBlockingHit
 						, HitResult.PhysMaterial.IsValid() ? HitResult.PhysMaterial->SurfaceType : DefaultSurfaceType
+						, HitProxy
 					);
 
 					Suspension->SetTargetPoint(TargetPoint);
diff --git a/Engine/Plugins/Experimental/ChaosModularVehicle/Source/ChaosModularVehicleEngine/Private/ChaosModularVehicle/SuspensionSimModule.cpp b/Engine/Plugins/Experimental/ChaosModularVehicle/Source/ChaosModularVehicleEngine/Private/ChaosModularVehicle/SuspensionSimModule.cpp
--- a/Engine/Plugins/Experimental/ChaosModularVehicle/Source/ChaosModularVehicleEngine/Private/ChaosModularVehicle/SuspensionSimModule.cpp
+++ b/Engine/Plugins/Experimental/ChaosModularVehicle/Source/ChaosModularVehicleEngine/Private/ChaosModularVehicle/SuspensionSimModule.cpp
@@ -171,8 +171,9 @@ using namespace Chaos;
 					Chaos::FPhysicsSolver* Solver = Proxy->GetSolver<Chaos::FPhysicsSolver>();
 					const FVector& CurrentTargetPosition = GetTargetPosition();
 					const FVector& CurrentImpactNormal = GetImpactNormal();
+					const IPhysicsProxyBase* GroundProxy = GetHitProxy();
 					const bool bCurrentWheelInContact = IsWheelInContact();
-					Solver->SetSuspensionTarget(Constraint, CurrentTargetPosition, CurrentImpactNormal, bCurrentWheelInContact);
+					Solver->SetSuspensionTarget(Constraint, CurrentTargetPosition, CurrentImpactNormal, bCurrentWheelInContact, GroundProxy);
 				}
 			}
 		}
diff --git a/Engine/Plugins/Experimental/ChaosVehiclesPlugin/Source/ChaosVehicles/Private/ChaosWheeledVehicleMovementComponent.cpp b/Engine/Plugins/Experimental/ChaosVehiclesPlugin/Source/ChaosVehicles/Private/ChaosWheeledVehicleMovementComponent.cpp
--- a/Engine/Plugins/Experimental/ChaosVehiclesPlugin/Source/ChaosVehicles/Private/ChaosWheeledVehicleMovementComponent.cpp
+++ b/Engine/Plugins/Experimental/ChaosVehiclesPlugin/Source/ChaosVehicles/Private/ChaosWheeledVehicleMovementComponent.cpp
@@ -24,6 +24,7 @@
 #include "Chaos/PBDSuspensionConstraintData.h"
 #include "Chaos/DebugDrawQueue.h"
 #include "UObject/UE5MainStreamObjectVersion.h"
+#include "PhysicsEngine/PhysicsObjectExternalInterface.h"
 
 #include "PhysicsProxy/SuspensionConstraintProxy.h"
 #include "PBDRigidsSolver.h"
@@ -620,7 +621,14 @@ void UChaosWheeledVehicleSimulation::ApplySuspensionForces(float DeltaTime, TArr
 
 							Chaos::FPhysicsSolver* Solver = Proxy->GetSolver<Chaos::FPhysicsSolver>();
 
-							Solver->SetSuspensionTarget(Constraint, TargetPos, HitResult.Normal, PWheel.InContact());
+							IPhysicsProxyBase* HitProxy = nullptr;
+
+							if (Chaos::FPhysicsObjectHandle PhysicsObject = HitResult.PhysicsObject)
+							{
+								HitProxy = Chaos::FPhysicsObjectInterface::GetProxy({ &PhysicsObject, 1 });
+							}
+
+							Solver->SetSuspensionTarget(Constraint, TargetPos, HitResult.Normal, PWheel.InContact(), HitProxy);
 						}
 					}
 				}
diff --git a/Engine/Source/Programs/HeadlessChaos/Private/HeadlessChaosTestConstraints.cpp b/Engine/Source/Programs/HeadlessChaos/Private/HeadlessChaosTestConstraints.cpp
--- a/Engine/Source/Programs/HeadlessChaos/Private/HeadlessChaosTestConstraints.cpp
+++ b/Engine/Source/Programs/HeadlessChaos/Private/HeadlessChaosTestConstraints.cpp
@@ -161,12 +161,15 @@ namespace ChaosTest {
 			Evolution.GetGravityForces().SetAcceleration(FVec3(0, 0, 0), 0);
 
 			// chassis particle
-			auto* DynamicParticle = Evolution.CreateDynamicParticles(1)[0];
+			TPBDRigidParticleHandle<FReal, 3>* DynamicParticle = Evolution.CreateDynamicParticles(1)[0];
 			DynamicParticle->SetX(FVec3(0, 10, 10));
 			//	DynamicParticle->R() = FRotation3::FromAxisAngle(FVec3(0,1,0), PI); // upside down
 			DynamicParticle->I() = TVec3<FRealSingle>(100.0f);
 			DynamicParticle->InvI() = TVec3<FRealSingle>(1.0f / 100.0f);
 
+			TGeometryParticleHandle<FReal, 3>* GroundParticle = Evolution.CreateStaticParticles(1)[0];
+			GroundParticle->SetX(FVec3(0, 10, 10));
+
 			FPBDSuspensionConstraints SuspensionConstraints;
 			FVec3 SuspensionLocalLocationA(FVec3(0, 0, 0));
 
@@ -174,6 +177,7 @@ namespace ChaosTest {
 
 			// hard-stop will activate because 9 breaks the min suspension limit, anything greater than 8 will do this
 			SuspensionConstraints.AddConstraint(DynamicParticle, SuspensionLocalLocationA, SuspensionSettings);
+			SuspensionConstraints.SetGroundParticle(0, GroundParticle);
 
 			Evolution.AddConstraintContainer(SuspensionConstraints);
 			Evolution.EnableParticle(DynamicParticle);
@@ -214,6 +218,9 @@ namespace ChaosTest {
 			DynamicParticle->I() = TVec3<FRealSingle>(100000.0f);
 			DynamicParticle->InvI() = TVec3<FRealSingle>(1.0f / 100000.0f);
 
+			TGeometryParticleHandle<FReal, 3>* GroundParticle = Evolution.CreateStaticParticles(1)[0];
+			GroundParticle->SetX(FVec3(0, 10, 10));
+
 			FPBDSuspensionConstraints SuspensionConstraints;
 
 			// local offset from particle origin
@@ -225,6 +232,8 @@ namespace ChaosTest {
 			SuspensionConstraints.AddConstraint(DynamicParticle, SuspensionLocalLocationA, SuspensionSettings);
 			SuspensionSettings.Target = FVec3(45, 10, 9);
 			SuspensionConstraints.AddConstraint(DynamicParticle, SuspensionLocalLocationB, SuspensionSettings);
+			SuspensionConstraints.SetGroundParticle(0, GroundParticle);
+			SuspensionConstraints.SetGroundParticle(1, GroundParticle);
 
 			Evolution.AddConstraintContainer(SuspensionConstraints);
 			Evolution.EnableParticle(DynamicParticle);
@@ -277,6 +286,9 @@ namespace ChaosTest {
 		DynamicParticle->I() = TVec3<FRealSingle>(100000.0f);
 		DynamicParticle->InvI() = TVec3<FRealSingle>(1.0f / 100000.0f);
 
+		TGeometryParticleHandle<FReal, 3>* GroundParticle = Evolution.CreateStaticParticles(1)[0];
+		GroundParticle->SetX(FVec3(0, 10, 10));
+
 		// Suspension setup
 		FPBDSuspensionSettings SuspensionSettings;
 		SuspensionSettings.Enabled = true;
@@ -295,8 +307,6 @@ namespace ChaosTest {
 
 		FPBDSuspensionConstraints SuspensionConstraints;
 
-		// spring will activate because 6 is between 3 and 8
-
 		for (int SusIndex=0; SusIndex < SusLocalOffset.Num(); SusIndex++)
 		{
 			FVec3 WorldTarget = FVec3(0, 0, 9);
@@ -304,6 +314,7 @@ namespace ChaosTest {
 		}
 
 		Evolution.AddConstraintContainer(SuspensionConstraints);
+		SuspensionConstraints.SetGroundParticle(0, GroundParticle);
 		Evolution.EnableParticle(DynamicParticle);
 
 		const FVec3& Pos = Evolution.GetParticleHandles().Handle(0)->GetX();
@@ -319,7 +330,7 @@ namespace ChaosTest {
 			Evolution.EndFrame(DeltaTime);
 
 			//UE_LOG(LogChaos, Warning, TEXT("Pos %s"), *Pos.ToString());
-			EXPECT_GT(Pos.Z, 12.f - PositionTolerance); // should never go past hard-stop
+ 			EXPECT_GT(Pos.Z, 12.f - PositionTolerance); // should never go past hard-stop
 		}
 
 		EXPECT_GT(Pos.Z, 12.f - PositionTolerance);	// suspension min limit
diff --git a/Engine/Source/Runtime/Experimental/Chaos/Private/Chaos/PBDSuspensionConstraintData.cpp b/Engine/Source/Runtime/Experimental/Chaos/Private/Chaos/PBDSuspensionConstraintData.cpp
--- a/Engine/Source/Runtime/Experimental/Chaos/Private/Chaos/PBDSuspensionConstraintData.cpp
+++ b/Engine/Source/Runtime/Experimental/Chaos/Private/Chaos/PBDSuspensionConstraintData.cpp
@@ -44,4 +44,14 @@ namespace Chaos
 			});
 	}
 
+	void FSuspensionConstraint::SetGroundBody(FPhysicsObjectHandle& InBody)
+	{
+		check(InBody);
+
+		GroundBody.Modify(/*bInvalidate=*/true, DirtyFlags, Proxy, [InBody](FPhysicsObjectProperty& Data)
+			{
+				Data.PhysicsBody = InBody;
+			});
+	}
+
 } // Chaos
diff --git a/Engine/Source/Runtime/Experimental/Chaos/Private/Chaos/PBDSuspensionConstraints.cpp b/Engine/Source/Runtime/Experimental/Chaos/Private/Chaos/PBDSuspensionConstraints.cpp
--- a/Engine/Source/Runtime/Experimental/Chaos/Private/Chaos/PBDSuspensionConstraints.cpp
+++ b/Engine/Source/Runtime/Experimental/Chaos/Private/Chaos/PBDSuspensionConstraints.cpp
@@ -65,13 +65,14 @@ namespace Chaos
 	{
 		int32 NewIndex = ConstrainedParticles.Num();
 		ConstrainedParticles.Add(Particle);
+		GroundParticles.Add(nullptr);
 		SuspensionLocalOffset.Add(InSuspensionLocalOffset);
 		ConstraintSettings.Add(InConstraintSettings);
 		ConstraintResults.AddDefaulted();
 		ConstraintEnabledStates.Add(true); // Note: assumes always enabled on creation
-		ConstraintSolverBodies.Add(nullptr);
-		StaticCollisionBodies.Add(FSolverBody::MakeInitialized());
+		ConstraintSolverBodies.Add({ nullptr, nullptr });
 		CollisionSolvers.Add(Private::FPBDCollisionSolver());
+		CollisionSolvers[CollisionSolvers.Num() - 1].MakeInitialized();
 		CollisionSolverManifoldPoints.Add(Private::FPBDCollisionSolverManifoldPoint());
 
 		Handles.Add(HandleAllocator.AllocHandle(this, NewIndex));
@@ -96,6 +97,7 @@ namespace Chaos
 
 		// Swap the last constraint into the gap to keep the array packed
 		ConstrainedParticles.RemoveAtSwap(ConstraintIndex);
+		GroundParticles.RemoveAtSwap(ConstraintIndex);
 		SuspensionLocalOffset.RemoveAtSwap(ConstraintIndex);
 		ConstraintSettings.RemoveAtSwap(ConstraintIndex);
 		ConstraintResults.RemoveAtSwap(ConstraintIndex);
@@ -103,7 +105,6 @@ namespace Chaos
 		ConstraintSolverBodies.RemoveAtSwap(ConstraintIndex);
 		CollisionSolvers.RemoveAtSwap(ConstraintIndex);
 		CollisionSolverManifoldPoints.RemoveAtSwap(ConstraintIndex);
-		StaticCollisionBodies.RemoveAtSwap(ConstraintIndex);
 
 		Handles.RemoveAtSwap(ConstraintIndex);
 
@@ -114,9 +115,91 @@ namespace Chaos
 		}
 	}
 
+	void FPBDSuspensionConstraints::DisconnectConstraints(const TSet<TGeometryParticleHandle<FReal, 3>*>& RemovedParticles)
+	{
+		for (FGeometryParticleHandle* RemovedParticle : RemovedParticles)
+		{
+			for (FConstraintHandle* ConstraintHandle : RemovedParticle->ParticleConstraints())
+			{
+				if (FPBDSuspensionConstraintHandle* SuspensionHandle = ConstraintHandle->As<FPBDSuspensionConstraintHandle>())
+				{
+					SuspensionHandle->SetEnabled(false); // constraint lifespan is managed by the proxy
+
+					int ConstraintIndex = SuspensionHandle->GetConstraintIndex();
+					if (ConstraintIndex != INDEX_NONE)
+					{
+						if (ConstrainedParticles[ConstraintIndex] == RemovedParticle)
+						{
+							ConstrainedParticles[ConstraintIndex] = nullptr;
+							GroundParticles[ConstraintIndex] = nullptr;
+						}
+					}
+				}
+			}
+		}
+	}
+
+	void FPBDSuspensionConstraints::SetConstraintEnabled(int32 ConstraintIndex, bool bEnabled)
+	{
+		if (ConstraintIndex < ConstrainedParticles.Num())
+		{
+			const FGenericParticleHandle Particle = FGenericParticleHandle(ConstrainedParticles[ConstraintIndex]);
+
+			if (bEnabled)
+			{
+				// only enable constraint if the particle is valid and not disabled
+				if (Particle->Handle() != nullptr && !Particle->Disabled())
+				{
+					ConstraintEnabledStates[ConstraintIndex] = true;
+				}
+			}
+			else
+			{
+				// desirable to allow disabling no matter what state the endpoint
+				ConstraintEnabledStates[ConstraintIndex] = false;
+			}
+		}
+	}
+
+	void FPBDSuspensionConstraints::SetGroundParticle(int32 ConstraintIndex, Chaos::FGeometryParticleHandle* InGroundParticle)
+	{
+		if (GroundParticles.IsValidIndex(ConstraintIndex))
+		{
+			if (GroundParticles[ConstraintIndex] != InGroundParticle)
+			{
+				if (FConstraintContainerHandle* Constraint = GetConstraintHandle(ConstraintIndex))
+				{
+					Constraint->SetGroundParticleChanged();
+				}
+			}
+
+			GroundParticles[ConstraintIndex] = InGroundParticle;
+		}
+	}
+
 	void FPBDSuspensionConstraints::AddConstraintsToGraph(Private::FPBDIslandManager& IslandManager)
 	{
-		IslandManager.AddContainerConstraints(*this);
+		for (FPBDSuspensionConstraintHandle* Constraint : GetConstraints())
+		{
+			const bool bIsInGraph = Constraint->IsInConstraintGraph();
+			const bool bShouldBeInGraph = Constraint->IsEnabled();
+			const bool bShouldUpdateGraph = Constraint->HasGroundParticleChanged();
+
+			if (bShouldBeInGraph && !bIsInGraph)
+			{
+				IslandManager.AddConstraint(ContainerId, Constraint, Constraint->GetConstrainedParticles());
+			}
+			else if (bIsInGraph && !bShouldBeInGraph)
+			{
+				IslandManager.RemoveConstraint(Constraint);
+			}
+			else if (bIsInGraph && bShouldUpdateGraph)
+			{
+				IslandManager.RemoveConstraint(Constraint);
+				IslandManager.AddConstraint(ContainerId, Constraint, Constraint->GetConstrainedParticles());
+				Constraint->ClearGroundParticleChanged();
+			}
+		}
 	}
 
 	void FPBDSuspensionConstraints::AddBodies(FSolverBodyContainer& SolverBodyContainer)
@@ -201,7 +284,26 @@ namespace Chaos
 
 	void FPBDSuspensionConstraints::AddBodies(const int32 ConstraintIndex, FSolverBodyContainer& SolverBodyContainer)
 	{
-		ConstraintSolverBodies[ConstraintIndex] = SolverBodyContainer.FindOrAdd(ConstrainedParticles[ConstraintIndex]);
+		// FindOrAdd expects a valid pointer
+		if (ConstrainedParticles[ConstraintIndex])
+		{
+			if (GroundParticles[ConstraintIndex])
+			{
+				ConstraintSolverBodies[ConstraintIndex] =
+				{
+					SolverBodyContainer.FindOrAdd(ConstrainedParticles[ConstraintIndex]),
+					SolverBodyContainer.FindOrAdd(GroundParticles[ConstraintIndex])
+				};
+			}
+			else
+			{
+				ConstraintSolverBodies[ConstraintIndex] =
+				{
+					SolverBodyContainer.FindOrAdd(ConstrainedParticles[ConstraintIndex]),
+					nullptr
+				};
+			}
+		}
 	}
 
 	void FPBDSuspensionConstraints::GatherInput(const int32 ConstraintIndex, const FReal Dt)
@@ -212,19 +314,29 @@ namespace Chaos
 
 		// Hard-stop Collision Manifold Generation
 		{
-			check(ConstraintSolverBodies[ConstraintIndex] != nullptr);
 
 			FPBDCollisionSolver* Solver = &CollisionSolvers[ConstraintIndex];
 			Solver->Reset(&CollisionSolverManifoldPoints[ConstraintIndex], 1);
 			Solver->SetStiffness(1);
 			Solver->SetHardContact();
 
-			FSolverBody* Body0 = ConstraintSolverBodies[ConstraintIndex];	// vehicle chassis			
-			FSolverBody* Body1 = &StaticCollisionBodies[ConstraintIndex];	// Spoofed terrain
+			// this is the case if the raycast has not found a body to hit - can we disable constraint in that instance rather than running
+			if (ConstraintSolverBodies[ConstraintIndex][1] == nullptr)
+			{
+				return;
+			}
+
+			check(ConstraintSolverBodies[ConstraintIndex][0] != nullptr);
+			check(ConstraintSolverBodies[ConstraintIndex][1] != nullptr);
+
+			FSolverBody* Body0 = ConstraintSolverBodies[ConstraintIndex][0];	// vehicle chassis			
+			FSolverBody* Body1 = ConstraintSolverBodies[ConstraintIndex][1];	// body that the wheel is touching
+
+			bool bHitSelf = (Body0 == Body1);
 
 			const FPBDSuspensionSettings& Setting = ConstraintSettings[ConstraintIndex];
 
-			if (Body0->IsDynamic() && Setting.Enabled)
+			if (Body0->IsDynamic() && !bHitSelf && Setting.Enabled)
 			{
 				const FVec3& T = Setting.Target;
 
@@ -241,7 +353,8 @@ namespace Chaos
 				FVec3 AxisWorld = BodyQ.RotateVector(SuspensionCoMAxis);
 				FReal Distance = FVec3::DotProduct(WorldSpaceX - T, AxisWorld);
 
-				const FVec3 WorldArm = BodyQ.RotateVector(SuspensionCoMOffset);
+				const FVec3 WorldArm0 = BodyQ.RotateVector(SuspensionCoMOffset);
+				const FVec3 WorldArm1 = WorldSpaceX - Body1->CorrectedP();
 
 				// The hard-stop can only apply correction perpendicular to the surface
 				const FVec3& WorldContactNormal = Setting.Normal;
@@ -249,12 +362,6 @@ namespace Chaos
 				const FReal HardStopDistance = Setting.MinLength - Distance;
 				const FReal WorldContactDeltaNormal = FVec3::DotProduct(HardStopDistance * AxisWorld, WorldContactNormal);
 
-				// position the spoofed terrain at the same position at body0
-				FVec3 PosBody1 = T + Distance*AxisWorld;
-				Body1->SetP(PosBody1);
-				Body1->SetX(PosBody1);
-				check(Body1->InvM() == 0);
-
 				Solver->SetSolverBodies(*Body0, *Body1);
 				Solver->SolverBody0().Init();
 				Solver->SolverBody1().Init();
@@ -263,11 +370,9 @@ namespace Chaos
 #if CHAOS_DEBUG_DRAW
 				if (bChaos_Suspension_DebugDraw_Hardstop)
 				{
-					FVec3 Body0Center = WorldArm + BodyP;
-					FVec3 Body1Center = PosBody1;
+					FVec3 Body0Center = WorldSpaceX;
 					FReal Radius = 30;
 					FDebugDrawQueue::GetInstance().DrawDebugCircle(Body0Center, Radius, 60, FColor::Yellow, false, -1.0f, 0, 3, FVector(1, 0, 0), FVector(0, 1, 0), false);
-					FDebugDrawQueue::GetInstance().DrawDebugCircle(Body1Center, Radius, 60, FColor::Green, false, -1.0f, 0, 3, FVector(1, 0, 0), FVector(0, 1, 0), false);
 					FString TextOut = FString::Format(TEXT("{0}"), { WorldContactDeltaNormal });
 					FDebugDrawQueue::GetInstance().DrawDebugString(Body0Center + FVec3(0, 50, 50), TextOut, nullptr, FColor::White, -1.f, true, 1.0f);
 				}
@@ -278,8 +383,8 @@ namespace Chaos
 				Solver->InitManifoldPoint(
 					0,										// ManifoldIndex
 					FSolverReal(Dt),						// Delta Time
-					FSolverVec3(WorldArm),					// RelativeContactPosition0,
-					FSolverVec3(0, 0, 0),					// RelativeContactPosition1,
+					FSolverVec3(WorldArm0),					// RelativeContactPosition0,
+					FSolverVec3(WorldArm1),					// RelativeContactPosition1,
 					FSolverVec3(WorldContactNormal),		// WorldContactNormal
 					FSolverVec3(0, 0, 0),					// WorldContactTangentU,
 					FSolverVec3(0, 0, 0),					// WorldContactTangentV,
@@ -305,7 +410,12 @@ namespace Chaos
 			ConstraintResults[ConstraintIndex].HardStopNetImpulse = ManifoldPoint.NetImpulseNormal * ManifoldPoint.ContactNormal;
 		}
 
-		ConstraintSolverBodies[ConstraintIndex] = nullptr;
+		ConstraintSolverBodies[ConstraintIndex] =
+		{
+			nullptr,
+			nullptr
+		};
+
 		CollisionSolvers[ConstraintIndex].Reset(nullptr, 0);
 	}
 
@@ -356,12 +466,23 @@ namespace Chaos
 
 	void FPBDSuspensionConstraints::ApplySingle(int32 ConstraintIndex, const FReal Dt)
 	{
-		check(ConstraintSolverBodies[ConstraintIndex] != nullptr);
-		FSolverBody& Body = *ConstraintSolverBodies[ConstraintIndex];
+		// this is the case if the raycast has not found a body to hit - can we disable constraint in that instance rather than running
+  		if (ConstraintSolverBodies[ConstraintIndex][1] == nullptr)
+		{
+			return;
+		}
+
+		check(ConstraintSolverBodies[ConstraintIndex][0] != nullptr);
+		check(ConstraintSolverBodies[ConstraintIndex][1] != nullptr);
+		FSolverBody& Body0 = *ConstraintSolverBodies[ConstraintIndex][0];
+		FSolverBody& Body1 = *ConstraintSolverBodies[ConstraintIndex][1];
+
+		bool bHitSelf = (ConstraintSolverBodies[ConstraintIndex][0] == ConstraintSolverBodies[ConstraintIndex][1]);
+
 		const FPBDSuspensionSettings& Setting = ConstraintSettings[ConstraintIndex];
 		FPBDSuspensionResults& Results = ConstraintResults[ConstraintIndex];
 
-		if (Body.IsDynamic() && Setting.Enabled)
+		if (Body0.IsDynamic() && !bHitSelf && Setting.Enabled)
 		{
 			const FVec3& T = Setting.Target;
 
@@ -373,9 +494,11 @@ namespace Chaos
 			const FVec3 SuspensionCoMAxis = Particle->RotationOfMass().UnrotateVector(Setting.Axis);
 
 			// @todo(chaos): use linearized error calculation
-			const FRotation3 BodyQ = Body.CorrectedQ();
-			const FVec3 BodyP = Body.CorrectedP();
+			const FRotation3 BodyQ = Body0.CorrectedQ();
+			const FVec3 BodyP = Body0.CorrectedP();
 			const FVec3 WorldSpaceX = BodyQ.RotateVector(SuspensionCoMOffset) + BodyP;
+			const FRotation3 BodyQ2 = Body1.CorrectedQ();
+			const FVec3 BodyP2 = Body1.CorrectedP();
 
 			FVec3 AxisWorld = BodyQ.RotateVector(SuspensionCoMAxis);
 
@@ -388,13 +511,13 @@ namespace Chaos
 			// Ingeniously blends the surface normal at low speeds to stop vehicles sliding slowly down steep slopes
 			if (SurfaceNormal.Z > Chaos_Suspension_SlopeThreshold)
 			{
-				if (Body.V().SquaredLength() < (SpeedThreshold * SpeedThreshold))
+				if (Body0.V().SquaredLength() < (SpeedThreshold * SpeedThreshold))
 				{
 					SurfaceNormal = FVec3(0.f, 0.f, 1.f);
 				}
 				else
 				{
-					const FReal Speed = FMath::Abs(Body.V().Length());
+					const FReal Speed = FMath::Abs(Body0.V().Length());
 					if (Speed < SpeedBlendThreshold)
 					{
 						SurfaceNormal = FMath::Lerp(FVec3(0.f, 0.f, 1.f), SurfaceNormal, Speed / SpeedBlendThreshold);
@@ -417,8 +540,8 @@ namespace Chaos
 			// velocity in the damping calculation
 			// @todo(chaos): consider moving the damping term to the velocity solve phase
 			const FVec3 Diff = WorldSpaceX - BodyP;
-			const FVec3 V = FVec3::CalculateVelocity(Body.X(), BodyP, Dt);
-			const FVec3 W = FRotation3::CalculateAngularVelocity(Body.R(), BodyQ, Dt);
+			const FVec3 V = FVec3::CalculateVelocity(Body0.X(), BodyP, Dt);
+			const FVec3 W = FRotation3::CalculateAngularVelocity(Body0.R(), BodyQ, Dt);
 			FVec3 ArmVelocity = V - FVec3::CrossProduct(Diff, W);
 
 
@@ -451,7 +574,7 @@ namespace Chaos
 
 					// @todo(chaos): this is not using the correct effective mass
 					const bool bAccelerationMode = false;
-					const FReal SpringMassScale = (bAccelerationMode) ? FReal(1) / Body.InvM() : FReal(1);
+					const FReal SpringMassScale = (bAccelerationMode) ? FReal(1) / Body0.InvM() : FReal(1);
 					const FReal S = SpringMassScale * Setting.SpringStiffness * Dt * Dt;
 					const FReal D = SpringMassScale * Setting.SpringDamping * Dt;
 					
@@ -470,10 +593,19 @@ namespace Chaos
 
 			const FVec3 Arm = WorldSpaceX - BodyP;
 
-			const FVec3 DP = Body.InvM() * DX;
-			const FVec3 DR = Body.InvI() * FVec3::CrossProduct(Arm, DX);
-			Body.ApplyTransformDelta(DP, DR);
-			Body.UpdateRotationDependentState();
+			const FVec3 DP = Body0.InvM() * DX;
+			const FVec3 DR = Body0.InvI() * FVec3::CrossProduct(Arm, DX);
+			Body0.ApplyTransformDelta(DP, DR);
+			Body0.UpdateRotationDependentState();
+
+			if (Body1.IsDynamic())
+			{
+				const FVec3 Arm2 = WorldSpaceX - BodyP2;
+				const FVec3 DP2 = Body1.InvM() * -DX;
+				const FVec3 DR2 = Body1.InvI() * FVec3::CrossProduct(Arm2, -DX);
+				Body1.ApplyTransformDelta(DP2, DR2);
+				Body1.UpdateRotationDependentState();
+			}
 
 			Results.NetPushOut += DX;
 			Results.Length = Distance;
diff --git a/Engine/Source/Runtime/Experimental/Chaos/Private/PBDRigidsSolver.cpp b/Engine/Source/Runtime/Experimental/Chaos/Private/PBDRigidsSolver.cpp
--- a/Engine/Source/Runtime/Experimental/Chaos/Private/PBDRigidsSolver.cpp
+++ b/Engine/Source/Runtime/Experimental/Chaos/Private/PBDRigidsSolver.cpp
@@ -1183,12 +1183,12 @@ namespace Chaos
 			});
 	}
 
-	void FPBDRigidsSolver::SetSuspensionTarget(FSuspensionConstraint* GTConstraint, const FVector& TargetPos, const FVector& Normal, bool Enabled)
+	void FPBDRigidsSolver::SetSuspensionTarget(FSuspensionConstraint* GTConstraint, const FVector& TargetPos, const FVector& Normal, bool Enabled, const IPhysicsProxyBase* GroundProxy)
 	{
 		EnsureIsInPhysicsThreadContext();
 		FSuspensionConstraintPhysicsProxy* SuspensionProxy = GTConstraint->GetProxy<FSuspensionConstraintPhysicsProxy>();
 		check(SuspensionProxy);
-		SuspensionProxy->UpdateTargetOnPhysicsThread(this, TargetPos, Normal, Enabled);
+		SuspensionProxy->UpdateTargetOnPhysicsThread(this, TargetPos, GroundProxy, Normal, Enabled);
 	}
 
 	void FPBDRigidsSolver::EnableRewindCapture(int32 NumFrames, bool InUseCollisionResimCache, TUniquePtr<IRewindCallback>&& RewindCallback)
diff --git a/Engine/Source/Runtime/Experimental/Chaos/Private/PhysicsProxy/SuspensionConstraintProxy.cpp b/Engine/Source/Runtime/Experimental/Chaos/Private/PhysicsProxy/SuspensionConstraintProxy.cpp
--- a/Engine/Source/Runtime/Experimental/Chaos/Private/PhysicsProxy/SuspensionConstraintProxy.cpp
+++ b/Engine/Source/Runtime/Experimental/Chaos/Private/PhysicsProxy/SuspensionConstraintProxy.cpp
@@ -109,14 +109,25 @@ void FSuspensionConstraintPhysicsProxy::DestroyOnPhysicsThread(FPBDRigidsSolver*
 	}
 }
 
-void FSuspensionConstraintPhysicsProxy::UpdateTargetOnPhysicsThread(FPBDRigidsSolver* RBDSolver, const FVector& TargetPos, const FVector& Normal, bool Enabled)
+void FSuspensionConstraintPhysicsProxy::UpdateTargetOnPhysicsThread(FPBDRigidsSolver* RBDSolver, const FVector& TargetPos, const IPhysicsProxyBase* GroundProxy, const FVector& Normal, bool Enabled)
 {
-	if (Constraint_PT)
+	if (Constraint_PT && Constraint_PT->IsValid())
 	{
 		auto& SuspensionConstraints = RBDSolver->GetSuspensionConstraints();
 		SuspensionConstraints.GetSettings(Constraint_PT->GetConstraintIndex()).Target = TargetPos;
 		SuspensionConstraints.GetSettings(Constraint_PT->GetConstraintIndex()).Normal = Normal;
 		SuspensionConstraints.GetSettings(Constraint_PT->GetConstraintIndex()).Enabled = Enabled;
+
+		FGeometryParticleHandle* GroundParticle = nullptr;
+		if (GroundProxy)
+		{
+			if (GroundProxy->GetType() == EPhysicsProxyType::SingleParticleProxy)
+			{
+				GroundParticle = ((FSingleParticlePhysicsProxy*)GroundProxy)->GetHandle_LowLevel();
+			}
+		}
+
+		SuspensionConstraints.SetGroundParticle(Constraint_PT->GetConstraintIndex(), GroundParticle);
 	}
 }
 
diff --git a/Engine/Source/Runtime/Experimental/Chaos/Public/Chaos/PBDSuspensionConstraintData.h b/Engine/Source/Runtime/Experimental/Chaos/Public/Chaos/PBDSuspensionConstraintData.h
--- a/Engine/Source/Runtime/Experimental/Chaos/Public/Chaos/PBDSuspensionConstraintData.h
+++ b/Engine/Source/Runtime/Experimental/Chaos/Public/Chaos/PBDSuspensionConstraintData.h
@@ -39,6 +39,7 @@ namespace Chaos
 
 		CHAOS_API void SetParticleProxy(IPhysicsProxyBase* InParticleProxy);
 		CHAOS_API void SetPhysicsBody(FPhysicsObjectHandle& InBody);
+		CHAOS_API void SetGroundBody(FPhysicsObjectHandle& InBody);
 
 		const FPBDSuspensionSettings& GetSuspensionSettings() const { return SuspensionSettings.Read(); }
 
@@ -53,6 +54,7 @@ namespace Chaos
 		TChaosProperty<FSuspensionLocation, EChaosProperty::SuspensionLocation> SuspensionLocation; // location = spring local offset
 		TChaosProperty<FProxyBaseProperty, EChaosProperty::SuspensionParticleProxy> SuspensionProxy;
 		TChaosProperty<FPhysicsObjectProperty, EChaosProperty::SuspensionPhysicsObject> SuspensionBody;
+		TChaosProperty<FPhysicsObjectProperty, EChaosProperty::GroundPhysicsObject> GroundBody;
 
 		virtual void SyncRemoteDataImp(FDirtyPropertiesManager& Manager, int32 DataIdx, FDirtyChaosProperties& RemoteData) override
 		{
@@ -60,6 +62,7 @@ namespace Chaos
 			SuspensionLocation.SyncRemote(Manager, DataIdx, RemoteData);
 			SuspensionProxy.SyncRemote(Manager, DataIdx, RemoteData);
 			SuspensionBody.SyncRemote(Manager, DataIdx, RemoteData);
+			GroundBody.SyncRemote(Manager, DataIdx, RemoteData);
 		}
 	};
 
diff --git a/Engine/Source/Runtime/Experimental/Chaos/Public/Chaos/PBDSuspensionConstraints.h b/Engine/Source/Runtime/Experimental/Chaos/Public/Chaos/PBDSuspensionConstraints.h
--- a/Engine/Source/Runtime/Experimental/Chaos/Public/Chaos/PBDSuspensionConstraints.h
+++ b/Engine/Source/Runtime/Experimental/Chaos/Public/Chaos/PBDSuspensionConstraints.h
@@ -10,7 +10,6 @@
 #include "Chaos/Particle/ParticleUtilities.h"
 #include "Chaos/Utilities.h"
 
-//PRAGMA_DISABLE_OPTIMIZATION
 
 namespace Chaos
 {
@@ -39,6 +38,16 @@ namespace Chaos
 
 		CHAOS_API virtual FParticlePair GetConstrainedParticles() const override final;
 
+		FGeometryParticleHandle* GetVehicleParticle() const { return GetConstrainedParticles()[0]; }
+		FGeometryParticleHandle* GetGroundParticle() const { return GetConstrainedParticles()[1]; }
+
+		void SetGroundParticleChanged() { bGroundParticleChanged = true; }
+		void ClearGroundParticleChanged() { bGroundParticleChanged = false; }
+		bool HasGroundParticleChanged() const
+		{
+			return bGroundParticleChanged;
+		}
+
 		static const FConstraintHandleTypeID& StaticType()
 		{
 			static FConstraintHandleTypeID STypeID(TEXT("FSuspensionConstraintHandle"), &FIndexedConstraintHandle::StaticType());
@@ -47,6 +56,8 @@ namespace Chaos
 	protected:
 		using Base::ConstraintIndex;
 		using Base::ConcreteContainer;
+
+		bool bGroundParticleChanged = false;
 	};
 
 	class FPBDSuspensionConstraints : public TPBDIndexedConstraintContainer<FPBDSuspensionConstraints>
@@ -56,6 +67,7 @@ namespace Chaos
 		using FConstraintContainerHandle = FPBDSuspensionConstraintHandle;
 		using FConstraintHandleAllocator = TConstraintHandleAllocator<FPBDSuspensionConstraints>;
 		using FHandles = TArray<FConstraintContainerHandle*>;
+		using FConstraints = TArrayView<FPBDSuspensionConstraintHandle* const>;
 
 		FPBDSuspensionConstraints(const FPBDSuspensionSolverSettings& InSolverSettings = FPBDSuspensionSolverSettings())
 			: TPBDIndexedConstraintContainer<FPBDSuspensionConstraints>(FConstraintContainerHandle::StaticType())
@@ -100,6 +112,7 @@ namespace Chaos
 		 */
 		CHAOS_API void RemoveConstraint(int ConstraintIndex);
 
+		FConstraints GetConstraints() { return MakeArrayView(Handles); }
 
 		/*
 		* Disconnect the constraints from the attached input particles.
@@ -109,53 +122,14 @@ namespace Chaos
 		* The constraint is unuseable at this point and pending deletion.
 		*/
 
-		void DisconnectConstraints(const TSet<TGeometryParticleHandle<FReal, 3>*>& RemovedParticles)
-		{
-			for (FGeometryParticleHandle* RemovedParticle : RemovedParticles)
-			{
-				for (FConstraintHandle* ConstraintHandle : RemovedParticle->ParticleConstraints())
-				{
-					if (FPBDSuspensionConstraintHandle* SuspensionHandle = ConstraintHandle->As<FPBDSuspensionConstraintHandle>())
-					{
-						SuspensionHandle->SetEnabled(false); // constraint lifespan is managed by the proxy
-
-						int ConstraintIndex = SuspensionHandle->GetConstraintIndex();
-						if (ConstraintIndex != INDEX_NONE)
-						{
-							if (ConstrainedParticles[ConstraintIndex] == RemovedParticle)
-							{
-								ConstrainedParticles[ConstraintIndex] = nullptr;
-							}
-						}
-					}
-				}
-			}
-		}
+		void DisconnectConstraints(const TSet<TGeometryParticleHandle<FReal, 3>*>& RemovedParticles);
 
 		bool IsConstraintEnabled(int32 ConstraintIndex) const
 		{
 			return ConstraintEnabledStates[ConstraintIndex];
 		}
 
-		void SetConstraintEnabled(int32 ConstraintIndex, bool bEnabled)
-		{
-			const FGenericParticleHandle Particle = FGenericParticleHandle(ConstrainedParticles[ConstraintIndex]);
-
-			if (bEnabled)
-			{
-				// only enable constraint if the particle is valid and not disabled
-				if (Particle->Handle() != nullptr && !Particle->Disabled())
-				{
-					ConstraintEnabledStates[ConstraintIndex] = true;
-				}
-			}
-			else
-			{
-				// desirable to allow disabling no matter what state the endpoint
-				ConstraintEnabledStates[ConstraintIndex] = false;
-			}
-
-		}
+		void SetConstraintEnabled(int32 ConstraintIndex, bool bEnabled);
 
 		//
 		// Constraint API
@@ -180,6 +154,8 @@ namespace Chaos
 			ConstraintSettings[ConstraintIndex].Target = TargetPos;
 		}
 
+		void SetGroundParticle(int32 ConstraintIndex, Chaos::FGeometryParticleHandle* InGroundParticle);
+
 		const FPBDSuspensionResults& GetResults(int32 ConstraintIndex) const
 		{
 			return ConstraintResults[ConstraintIndex];
@@ -209,7 +185,11 @@ namespace Chaos
 		 */
 		TVec2<TGeometryParticleHandle<FReal, 3>*> GetConstrainedParticles(int32 ConstraintIndex) const
 		{
-			return { ConstrainedParticles[ConstraintIndex], nullptr };
+			check(ConstraintIndex < ConstrainedParticles.Num());
+			check(ConstraintIndex < GroundParticles.Num());
+			check(ConstrainedParticles.Num() == GroundParticles.Num());
+
+			return { ConstrainedParticles[ConstraintIndex], GroundParticles[ConstraintIndex] };
 		}
 
 		/**
@@ -270,20 +250,18 @@ namespace Chaos
 		FPBDSuspensionSolverSettings SolverSettings;
 
 		TArray<FGeometryParticleHandle*> ConstrainedParticles;
+		TArray<FGeometryParticleHandle*> GroundParticles;
 		TArray<FVec3> SuspensionLocalOffset;
 		TArray<FPBDSuspensionSettings> ConstraintSettings;
 		TArray<FPBDSuspensionResults> ConstraintResults;
 		TArray<bool> ConstraintEnabledStates;
 
-		TArray<FSolverBody*> ConstraintSolverBodies;
+		TArray<FSolverBodyPtrPair> ConstraintSolverBodies;
 
 		FHandles Handles;
 		FConstraintHandleAllocator HandleAllocator;
 
 		TArray<Private::FPBDCollisionSolver> CollisionSolvers;
 		TArray<Private::FPBDCollisionSolverManifoldPoint> CollisionSolverManifoldPoints;
-		TArray<FSolverBody> StaticCollisionBodies;
 	};
 }
-
-//PRAGMA_ENABLE_OPTIMIZATION
diff --git a/Engine/Source/Runtime/Experimental/Chaos/Public/Chaos/ParticleProperties.inl b/Engine/Source/Runtime/Experimental/Chaos/Public/Chaos/ParticleProperties.inl
--- a/Engine/Source/Runtime/Experimental/Chaos/Public/Chaos/ParticleProperties.inl
+++ b/Engine/Source/Runtime/Experimental/Chaos/Public/Chaos/ParticleProperties.inl
@@ -25,6 +25,8 @@ CHAOS_PROPERTY(SuspensionSettings, FPBDSuspensionSettings, EPhysicsProxyType::Su
 CHAOS_PROPERTY(SuspensionParticleProxy, FParticleProxyProperty, EPhysicsProxyType::SuspensionConstraintType)
 CHAOS_PROPERTY(SuspensionLocation, FSuspensionLocation, EPhysicsProxyType::SuspensionConstraintType)
 CHAOS_PROPERTY(SuspensionPhysicsObject, FPhysicsObjectProperty, EPhysicsProxyType::SuspensionConstraintType)
+CHAOS_PROPERTY(GroundPhysicsObject, FPhysicsObjectProperty, EPhysicsProxyType::SuspensionConstraintType)
+
 
 CHAOS_PROPERTY(CharacterGroundConstraintSettings, FCharacterGroundConstraintSettings, EPhysicsProxyType::CharacterGroundConstraintType)
 CHAOS_PROPERTY(CharacterGroundConstraintDynamicData, FCharacterGroundConstraintDynamicData, EPhysicsProxyType::CharacterGroundConstraintType)
diff --git a/Engine/Source/Runtime/Experimental/Chaos/Public/PBDRigidsSolver.h b/Engine/Source/Runtime/Experimental/Chaos/Public/PBDRigidsSolver.h
--- a/Engine/Source/Runtime/Experimental/Chaos/Public/PBDRigidsSolver.h
+++ b/Engine/Source/Runtime/Experimental/Chaos/Public/PBDRigidsSolver.h
@@ -228,7 +228,7 @@ namespace Chaos
 
 		FPBDSuspensionConstraints& GetSuspensionConstraints() { return MEvolution->GetSuspensionConstraints(); }
 		const FPBDSuspensionConstraints& GetSuspensionConstraints() const { return MEvolution->GetSuspensionConstraints(); }
-		CHAOS_API void SetSuspensionTarget(Chaos::FSuspensionConstraint* GTConstraint, const FVector& TargetPos, const FVector& Normal, bool Enabled);
+		CHAOS_API void SetSuspensionTarget(Chaos::FSuspensionConstraint* GTConstraint, const FVector& TargetPos, const FVector& Normal, bool Enabled, const IPhysicsProxyBase* GroundProxy = nullptr);
 
 		FCharacterGroundConstraintContainer& GetCharacterGroundConstraints() { return MEvolution->GetCharacterGroundConstraints(); }
 		const FCharacterGroundConstraintContainer& GetCharacterGroundConstraints() const { return MEvolution->GetCharacterGroundConstraints(); }
diff --git a/Engine/Source/Runtime/Experimental/Chaos/Public/PhysicsProxy/SuspensionConstraintProxy.h b/Engine/Source/Runtime/Experimental/Chaos/Public/PhysicsProxy/SuspensionConstraintProxy.h
--- a/Engine/Source/Runtime/Experimental/Chaos/Public/PhysicsProxy/SuspensionConstraintProxy.h
+++ b/Engine/Source/Runtime/Experimental/Chaos/Public/PhysicsProxy/SuspensionConstraintProxy.h
@@ -46,7 +46,7 @@ public:
 	CHAOS_API void DestroyOnGameThread();
 	CHAOS_API void DestroyOnPhysicsThread(FPBDRigidsSolver* InSolver);
 
-	CHAOS_API void UpdateTargetOnPhysicsThread(FPBDRigidsSolver* InSolver, const FVector& TargetPos, const FVector& Normal, bool Enabled);
+	CHAOS_API void UpdateTargetOnPhysicsThread(FPBDRigidsSolver* InSolver, const FVector& TargetPos, const IPhysicsProxyBase* GroundProxy, const FVector& Normal, bool Enabled);
 
 
 	//
diff --git a/Engine/Source/Runtime/Experimental/ChaosVehicles/ChaosVehiclesCore/Private/SimModule/SimModuleTree.cpp b/Engine/Source/Runtime/Experimental/ChaosVehicles/ChaosVehiclesCore/Private/SimModule/SimModuleTree.cpp
--- a/Engine/Source/Runtime/Experimental/ChaosVehicles/ChaosVehiclesCore/Private/SimModule/SimModuleTree.cpp
+++ b/Engine/Source/Runtime/Experimental/ChaosVehicles/ChaosVehiclesCore/Private/SimModule/SimModuleTree.cpp
@@ -9,6 +9,7 @@
 #include "Chaos/PhysicsObjectInternalInterface.h"
 #include "Chaos/DebugDrawQueue.h"
 #include "VehicleUtility.h"
+#include "SimModule/SimModulesInclude.h"
 
 #include UE_INLINE_GENERATED_CPP_BY_NAME(SimModuleTree)
 
@@ -483,9 +484,20 @@ void FSimModuleTree::UpdateModuleVelocites(IPhysicsProxyBase* PhysicsProxy, Chao
 						FVector WorldLocation = BodyTransform.TransformPosition(LocalPos);
 						const Chaos::FVec3 Arm = WorldLocation - Particle->GetX();
 
-						//Chaos::FDebugDrawQueue::GetInstance().DrawDebugLine(Particle->X(), Particle->X() + Arm, FColor::Yellow, false, -1.f, 0, 2.f);
-
-						FVector WorldVelocity = Particle->GetV() - Chaos::FVec3::CrossProduct(Arm, Particle->GetW());
+						// Take surface velocity into account
+						Chaos::FPBDRigidParticleHandle* GroundParticle = nullptr;
+						FVector GroundVel = FVector::ZeroVector;
+						if (FWheelBaseInterface* Wheel = Module->Cast<FWheelBaseInterface>())
+						{
+							GroundParticle = Wheel->GetGroundBody();
+							if (GroundParticle && GroundParticle != RootParticle)
+							{
+								const Chaos::FVec3 GroundArm = WorldLocation - GroundParticle->GetX();
+								GroundVel = GroundParticle->GetV() - Chaos::FVec3::CrossProduct(GroundArm, GroundParticle->GetW());
+							}
+						}
+
+						FVector WorldVelocity = Particle->GetV() - Chaos::FVec3::CrossProduct(Arm, Particle->GetW()) - GroundVel;
 						FVector LocalVelocity = OffsetTransform.InverseTransformVector(BodyTransform.InverseTransformVector(WorldVelocity));
 
 						FVector LocalAngular = BodyTransform.InverseTransformVector(Particle->GetW());
diff --git a/Engine/Source/Runtime/Experimental/ChaosVehicles/ChaosVehiclesCore/Private/SimModule/WheelModule.cpp b/Engine/Source/Runtime/Experimental/ChaosVehicles/ChaosVehiclesCore/Private/SimModule/WheelModule.cpp
--- a/Engine/Source/Runtime/Experimental/ChaosVehicles/ChaosVehiclesCore/Private/SimModule/WheelModule.cpp
+++ b/Engine/Source/Runtime/Experimental/ChaosVehicles/ChaosVehiclesCore/Private/SimModule/WheelModule.cpp
@@ -2,12 +2,20 @@
 
 #include "SimModule/WheelModule.h"
 #include "SimModule/SimModuleTree.h"
+#include "Chaos/DebugDrawQueue.h"
+#include "Chaos/Particle/ParticleUtilities.h"
 #include "VehicleUtility.h"
 
 #if VEHICLE_DEBUGGING_ENABLED
 UE_DISABLE_OPTIMIZATION_SHIP
 #endif
 
+bool bModularVehicle_ApplyWheelForceToGroundBody = true;
+FAutoConsoleVariableRef CVarModularVehicleApplyWheelForceToGroundBody(TEXT("p.ModularVehicle.ApplyWheelForceToGroundBody.Enabled"), bModularVehicle_ApplyWheelForceToGroundBody, TEXT("Enable/Disable aplication of wheel forces to the ground body."));
+
+float bModularVehicle_WheelForceToGroundBodyScaling = 1.0f;
+FAutoConsoleVariableRef CVarModularVehicleWheelForceToGroundBodyScaling(TEXT("p.ModularVehicle.WheelForceToGroundBodyScaling"), bModularVehicle_WheelForceToGroundBodyScaling, TEXT("Sacle froces applied to ground body from wheel."));
+
 namespace Chaos
 {
 	FWheelSimModule::FWheelSimModule(const FWheelSettings& Settings)
@@ -179,11 +187,38 @@ namespace Chaos
 				ForceFromFriction.X = FinalLateralForce;
 			}
 
+			// This code applies torque to the body the wheel is touching.
+			FVector AppliedWorldForce = -ForceFromFriction;
+
 			if (Setup().ReverseDirection)
 			{
 				ForceFromFriction = -ForceFromFriction;
 			}
+
 			AddLocalForceAtPosition(SteeringRotator.RotateVector(ForceFromFriction), Setup().ForceOffset);
+
+			if (bModularVehicle_ApplyWheelForceToGroundBody && GroundBody && GroundBody->IsDynamic())
+			{
+				const Chaos::FVec3 WorldCOM = Chaos::FParticleUtilitiesXR::GetCoMWorldPosition(GroundBody);
+
+				const FTransform VehicleTransform = Inputs.VehicleWorldTransform;
+				FTransform WheelRelativeTransform = GetParentRelativeTransform();
+
+				AppliedWorldForce = SteeringRotator.RotateVector(AppliedWorldForce);
+				Chaos::FVec3 WorldForce = WheelRelativeTransform.TransformVector(AppliedWorldForce);
+				WorldForce = VehicleTransform.TransformVector(WorldForce);
+				WorldForce *= bModularVehicle_WheelForceToGroundBodyScaling;
+				Chaos::FVec3 Arm = ImpactPoint - WorldCOM;
+				Chaos::FVec3 WorldTorque = Chaos::FVec3::CrossProduct(Arm, WorldForce);
+
+				GroundBody->AddForce(WorldForce);
+				GroundBody->AddTorque(WorldTorque);
+
+				//ImpactPoint.Z += 10.0f;
+				//Chaos::FDebugDrawQueue::GetInstance().DrawDebugLine(ImpactPoint, ImpactPoint + WorldForce * 0.01f, FColor::Cyan, false, -1.f, 0, 2.f);
+				//Chaos::FDebugDrawQueue::GetInstance().DrawDebugLine(ImpactPoint, ImpactPoint - Arm, FColor::Red, false, -1.f, 0, 2.f);
+				//Chaos::FDebugDrawQueue::GetInstance().DrawDebugSphere(WorldCOM, 20, 16, FColor::Red, false, -1.f, 0, 2.f);
+			}
 		}
 
 		TransmitTorque(VehicleModuleSystem, DriveTorque, BrakeTorque);
diff --git a/Engine/Source/Runtime/Experimental/ChaosVehicles/ChaosVehiclesCore/Public/SimModule/SuspensionBaseInterface.h b/Engine/Source/Runtime/Experimental/ChaosVehicles/ChaosVehiclesCore/Public/SimModule/SuspensionBaseInterface.h
--- a/Engine/Source/Runtime/Experimental/ChaosVehicles/ChaosVehiclesCore/Public/SimModule/SuspensionBaseInterface.h
+++ b/Engine/Source/Runtime/Experimental/ChaosVehicles/ChaosVehiclesCore/Public/SimModule/SuspensionBaseInterface.h
@@ -34,28 +34,22 @@ struct FSuspensionTargetPoint
 {
 	FSuspensionTargetPoint() {}
 
-	FSuspensionTargetPoint(const FVector& InTargetPosition, const FVector& InImpactNormal, const float InHitDistance, const bool bInWheelInContact, const TEnumAsByte<EPhysicalSurface> InSurfaceType)
+	FSuspensionTargetPoint(const FVector& InTargetPosition, const FVector& InImpactNormal, const float InHitDistance, const bool bInWheelInContact, const TEnumAsByte<EPhysicalSurface> InSurfaceType, IPhysicsProxyBase* InHitProxy)
 		: TargetPosition(InTargetPosition)
 		, ImpactNormal(InImpactNormal)
 		, HitDistance(InHitDistance)
 		, bWheelInContact(bInWheelInContact)
 		, SurfaceType(InSurfaceType)
+		, HitProxy(InHitProxy)
 	{}
 
-	void Serialize(FArchive& Ar)
-	{
-		Ar << TargetPosition;
-		Ar << ImpactNormal;
-		Ar << HitDistance;
-		Ar << bWheelInContact;
-		Ar << SurfaceType;
-	}
 
 	FVector TargetPosition = FVector::ZeroVector;
 	FVector ImpactNormal = FVector::ZeroVector;
 	float HitDistance = 0.0f;
 	bool bWheelInContact = false;
 	TEnumAsByte<EPhysicalSurface> SurfaceType = EPhysicalSurface::SurfaceType_Default;
+	IPhysicsProxyBase* HitProxy = nullptr;
 };
 
 class FSuspensionBaseInterface 
@@ -96,6 +90,9 @@ public:
 	void SetTargetPosition(const FVector& NewValue) { TargetPoint.TargetPosition = NewValue; }
 	const FVector& GetTargetPosition() const { return TargetPoint.TargetPosition; }
 
+	void SetHitProxy(IPhysicsProxyBase* NewValue) { TargetPoint.HitProxy = NewValue; }
+	const IPhysicsProxyBase* GetHitProxy() const { return TargetPoint.HitProxy; }
+
 protected:
 
 	int WheelSimTreeIndex = INVALID_IDX;
diff --git a/Engine/Source/Runtime/Experimental/ChaosVehicles/ChaosVehiclesCore/Public/SimModule/TorqueSimModule.h b/Engine/Source/Runtime/Experimental/ChaosVehicles/ChaosVehiclesCore/Public/SimModule/TorqueSimModule.h
--- a/Engine/Source/Runtime/Experimental/ChaosVehicles/ChaosVehiclesCore/Public/SimModule/TorqueSimModule.h
+++ b/Engine/Source/Runtime/Experimental/ChaosVehicles/ChaosVehiclesCore/Public/SimModule/TorqueSimModule.h
@@ -137,13 +137,25 @@ namespace Chaos
 		void SetForceIntoSurface(float ForceIntoSurfaceIn) { ForceIntoSurface = ForceIntoSurfaceIn; }
 		float GetForceIntoSurface() const { return ForceIntoSurface; }
 		float GetSurfaceFriction() const { return SurfaceFriction; }
+		void SetGroundInteraction(Chaos::FPBDRigidParticleHandle* GroundParticleIn, const FVector& ImpactPointIn, const FVector& ImpactNormalIn)
+		{ 
+			GroundBody = GroundParticleIn; 
+			ImpactPoint = ImpactPointIn;
+			ImpactNormal = ImpactNormalIn;
+		}
 
+		Chaos::FPBDRigidParticleHandle* GetGroundBody() const { return GroundBody; }
+		FVector GetImpactPoint() const { return ImpactPoint; }
+		FVector GetImpactNormal() const { return ImpactNormal; }
 		virtual float GetWheelRadius() const = 0;
 
 	protected:
 		int SuspensionSimTreeIndex;
 		float SurfaceFriction;
 		float ForceIntoSurface;
+		Chaos::FPBDRigidParticleHandle* GroundBody = nullptr;
+		FVector ImpactPoint = FVector::ZeroVector;
+		FVector ImpactNormal = FVector(0,0,1);
 
 	};
 
