[NetPhysics] Add a setting for cooldown frames between resimulations and a setting to clamp number of resim frames.
- By default, when out of resim frames physics replication switches to Predictive Interpolation temporarily
- p.Resim.ClampedFramesOverride 0
- p.Resim.CooldownFrames 0

#Synced-CL 48646594

diff --git a/Engine/Source/Runtime/Experimental/Chaos/Private/RewindData.cpp b/Engine/Source/Runtime/Experimental/Chaos/Private/RewindData.cpp
--- a/Engine/Source/Runtime/Experimental/Chaos/Private/RewindData.cpp
+++ b/Engine/Source/Runtime/Experimental/Chaos/Private/RewindData.cpp
@@ -205,6 +205,7 @@ FRewindData::FRewindData(FPBDRigidsSolver* InSolver, int32 NumFrames, bool InRew
 	, bRewindDataOptimization(InRewindDataOptimization)
 	, LatestTargetFrame(0)
 {
+	ResimClampedFrames = Managers.Capacity();
 	RegisterEvolutionCallbacks();
 }
 
@@ -219,6 +220,7 @@ FRewindData::FRewindData(FPBDRigidsSolver* InSolver, int32 NumFrames, int32 InCu
 	, bRewindDataOptimization(bCVarRewindDataOptimization)
 	, LatestTargetFrame(0)
 {
+	ResimClampedFrames = Managers.Capacity();
 	RegisterEvolutionCallbacks();
 }
 
@@ -793,6 +795,9 @@ FAutoConsoleVariableRef CVarResimAllowRewindToResimulatedFrames(TEXT("p.Resim.Al
 CHAOS_API float ResimAllowRewindToResimulatedFramesOverlapPercent = 1.0f;
 FAutoConsoleVariableRef CVarResimAllowRewindToResimulatedFramesOverlapPercent(TEXT("p.Resim.AllowRewindToResimulatedFrames.OverlapPercent"), ResimAllowRewindToResimulatedFramesOverlapPercent, TEXT("Value in percent as a multiplier, 0.3 = 30%. When p.Resim.AllowRewindToResimulatedFrames is true, this sets how many frames of overlap we allow when rewinding back past a previous resimulation. If a resimulation is performed between frame 100-110 (10 frames), 30% overlap would allow 3 frames to get resimulated again, i.e. we would be allowed to rewind back to frame 107 again if we need to."));
 
+int32 ResimCooldownFrames = 0;
+FAutoConsoleVariableRef CVarResimCooldownFrames(TEXT("p.Resim.CooldownFrames"), ResimCooldownFrames, TEXT("Number of frames that needs to pass between each resimulation."));
+
 void FRewindData::RewindToFrame(int32 RewindFrame)
 {
 	FPBDRigidsEvolution* Evolution = Solver ? Solver->GetEvolution() : nullptr;
@@ -828,6 +833,8 @@ void FRewindData::RewindToFrame(int32 RewindFrame)
 	}
 
 	SetResimFrame(RewindFrame);
+	ResimCooldownUntilFrame = LatestFrame + FMath::Max(ResimCooldownFrames, 0);
+
 	CurFrame = RewindFrame;
 	bNeedsSave = false;
 
@@ -1259,7 +1266,7 @@ void FRewindData::FinishFrame()
 void FRewindData::DumpHistory_Internal(const int32 FramePrintOffset, const FString& Filename)
 {
 	FStringOutputDevice Out;
-	const int32 EarliestFrame = GetEarliestFrame_Internal();
+	const int32 EarliestFrame = LatestFrame - FramesSaved;
 	for(int32 Frame = EarliestFrame; Frame < CurFrame; ++Frame)
 	{
 		for (int32 Phase = 0; Phase < FFrameAndPhase::EParticleHistoryPhase::NumPhases; ++Phase)
@@ -2555,6 +2562,11 @@ void FRewindData::SetTargetStateAtFrame(FGeometryParticleHandle& Handle, const i
 
 void FRewindData::RequestResimulation(int32 Frame, Chaos::FGeometryParticleHandle* Particle)
 {
+	if (IsResim() || LatestFrame < ResimCooldownUntilFrame)
+	{
+		return;
+	}
+
 	// Update ResimFrame but don't allow to set a newer frame than already set
 	SetResimFrame((ResimFrame == INDEX_NONE) ? Frame : FMath::Min(ResimFrame, Frame));
 
@@ -2572,6 +2584,17 @@ void FRewindData::RequestResimulation(int32 Frame, Chaos::FGeometryParticleHandl
 	}
 }
 
+int32 ResimClampedFramesOverride = 0;
+FAutoConsoleVariableRef CVarResimClampedFramesOverride(TEXT("p.Resim.ClampedFramesOverride"), ResimClampedFramesOverride, TEXT("Override how far behind current frame is considered the earliest available resim frame. Gets clamped between 0 and number of frames cached in history."));
+
+int32 FRewindData::GetEarliestFrame_Internal() const
+{
+	checkf(ResimClampedFrames > INDEX_NONE, TEXT("FRewindData not initialized correctly, ensure use of proper constructor or Init() function calls"));
+
+	const int32 EffectiveResimClamp = FMath::Min((ResimClampedFramesOverride > 0) ? ResimClampedFramesOverride : ResimClampedFrames, FramesSaved);
+	return LatestFrame - EffectiveResimClamp;
+}
+
 void FRewindData::BlockResim()
 {
 	if (LatestFrame > BlockResimFrame)
diff --git a/Engine/Source/Runtime/Experimental/Chaos/Public/RewindData.h b/Engine/Source/Runtime/Experimental/Chaos/Public/RewindData.h
--- a/Engine/Source/Runtime/Experimental/Chaos/Public/RewindData.h
+++ b/Engine/Source/Runtime/Experimental/Chaos/Public/RewindData.h
@@ -1703,6 +1703,7 @@ public:
 		bRewindDataOptimization = InRewindDataOptimization;
 		LatestTargetFrame = 0;
 		Managers = TCircularBuffer<FFrameManagerInfo>(NumFrames + 1);
+		ResimClampedFrames = Managers.Capacity();
 
 		RegisterEvolutionCallbacks();
 	}
@@ -1714,6 +1715,7 @@ public:
 		LatestFrame = InCurrentFrame;
 		LatestTargetFrame = 0;
 		Managers = TCircularBuffer<FFrameManagerInfo>(NumFrames + 1);
+		ResimClampedFrames = Managers.Capacity();
 
 		RegisterEvolutionCallbacks();
 	}
@@ -1743,7 +1745,7 @@ public:
 		DirtyJoints.Remove(Joint, AllowShrinking);
 	}
 
-	int32 GetEarliestFrame_Internal() const { return CurFrame - FramesSaved; }
+	int32 CHAOS_API GetEarliestFrame_Internal() const;
 
 	/* Extend the current history size to be sure to include the given frame */
 	void CHAOS_API ExtendHistoryWithFrame(const int32 Frame);
@@ -1940,12 +1942,13 @@ public:
 	/** Get and set the frame we resimulate from */
 	const int32 GetResimFrame() const { return ResimFrame; }
 	void SetResimFrame(int32 Frame) { ResimFrame = Frame; }
+	const bool IsResimCooldown() const { return (LatestFrame < ResimCooldownUntilFrame); }
 
 	/** Called at the end of resim */
 	void EndResim();
 
-	/** Check if a frame number is within the current rewind history */
-	bool IsFrameWithinRewindHistory(int32 Frame) { return Frame < CurrentFrame() && Frame >= GetEarliestFrame_Internal(); }
+	/** Check if a frame number is within the currently allowed rewind span */
+	bool IsFrameWithinRewindHistory(int32 Frame) { return Frame < LatestFrame && Frame >= GetEarliestFrame_Internal(); }
 
 	/** Request a resimulation by setting a requested frame to rewind to
 	* @param Particle can be used to mark a specific particle as the one triggering the resimulation */
@@ -2155,7 +2158,7 @@ private:
 	auto GetPastStateAtFrameImp(const TDirtyObjs& DirtyObjs, const TObj& Handle, int32 Frame, FFrameAndPhase::EParticleHistoryPhase Phase) const
 	{
 		ensure(!IsResim());
-		ensure(Frame >= GetEarliestFrame_Internal());	//can't get state from before the frame we rewound to
+		ensure(Frame >= LatestFrame - FramesSaved);
 
 		const auto* Info = DirtyObjs.Find(&Handle);
 		const auto* State = Info ? &Info->GetHistory() : nullptr;
@@ -2203,6 +2206,8 @@ private:
 	bool bNeedsSave;	//Indicates that some data is pointing at head and requires saving before a rewind
 	bool bRewindDataOptimization;
 	int32 ResimFrame = INDEX_NONE;
+	int32 ResimCooldownUntilFrame = INDEX_NONE;
+	int32 ResimClampedFrames = INDEX_NONE;
 	int32 LatestTargetFrame;
 
 	// Used to block rewinding past a physics change we currently don't handle
