Add assert to check for multiple calls to FOnDemandIoStoreFactory::CreateInstance.

#rnx

#Synced-CL 48498703

diff --git a/Engine/Source/Runtime/Experimental/IoStore/OnDemand/Private/OnDemandModule.cpp b/Engine/Source/Runtime/Experimental/IoStore/OnDemand/Private/OnDemandModule.cpp
--- a/Engine/Source/Runtime/Experimental/IoStore/OnDemand/Private/OnDemandModule.cpp
+++ b/Engine/Source/Runtime/Experimental/IoStore/OnDemand/Private/OnDemandModule.cpp
@@ -84,6 +84,8 @@ class FOnDemandIoStoreFactory final
 public:
 	virtual IOnDemandIoStore* CreateInstance() override
 	{
+		check(IoStore == nullptr);
+
 		IoStore = MakeShared<FOnDemandIoStore>();
 		return IoStore.Get();
 	}
