Fix bug using ON_SCOPE_EXIT and capturing TValueOrError return value
- depending on RVO the captured FResult (TValueOrError<void, FError> can be move before
the scope exit lambda has been called, causing a crash when sending the error to analytics.

#rnx

#tests Manual testing

#Synced-CL 48600773

diff --git a/Engine/Source/Runtime/Experimental/IoStore/OnDemand/Private/OnDemandInstallCache.cpp b/Engine/Source/Runtime/Experimental/IoStore/OnDemand/Private/OnDemandInstallCache.cpp
--- a/Engine/Source/Runtime/Experimental/IoStore/OnDemand/Private/OnDemandInstallCache.cpp
+++ b/Engine/Source/Runtime/Experimental/IoStore/OnDemand/Private/OnDemandInstallCache.cpp
@@ -1102,48 +1102,51 @@ FResult FCasJournal::Commit(FTransaction&& Transaction, uint64& OutByteCount, ui
 
 	IPlatformFile& Ipf = FPlatformFileManager::Get().GetPlatformFile();
 
-	FResult Result			= MakeValue();
-	uint64 TotalEntrySize	= 0;
-	ON_SCOPE_EXIT
-	{
-		FOnDemandInstallCacheStats::OnJournalCommit(Result, TotalEntrySize);
-	};
-
 	// Validate header and footer
 	{
 		TUniquePtr<IFileHandle> FileHandle(Ipf.OpenRead(*Transaction.JournalFile));
 		if (FileHandle.IsValid() == false)
 		{
-			return Result = MakeJournalError(
+			FResult Result = MakeJournalError(
 				ECasErrorCode::CommitJournalFailed, EIoErrorCode::FileOpenFailed, FString::Printf(TEXT("Failed to open journal '%s'"), *Transaction.JournalFile));
+			FOnDemandInstallCacheStats::OnJournalCommit(Result, 0);
+			return Result;
 		}
 
 		const int64 FileSize = FileHandle->Size();
 		if (FileSize < FHeader::Size())
 		{
-			return Result = MakeJournalError(
+			FResult Result = MakeJournalError(
 				ECasErrorCode::CommitJournalFailed, EIoErrorCode::FileCorrupt, FString::Printf(TEXT("Failed to validate journal header in '%s'"), *Transaction.JournalFile));
+			FOnDemandInstallCacheStats::OnJournalCommit(Result, 0);
+			return Result;
 		}
 
 		FHeader Header;
 		if ((FileHandle->Read(reinterpret_cast<uint8*>(&Header), FHeader::Size()) == false) || (Header.IsValid() == false))
 		{
-			return Result = MakeJournalError(
+			FResult Result = MakeJournalError(
 				ECasErrorCode::CommitJournalFailed, EIoErrorCode::ReadError, FString::Printf(TEXT("Failed to validate journal header in '%s'"), *Transaction.JournalFile));
+			FOnDemandInstallCacheStats::OnJournalCommit(Result, 0);
+			return Result;
 		}
 
 		const int64 FooterPos = FileSize - FFooter::Size();
 		if (FileHandle->Seek(FooterPos) == false)
 		{
-			return Result = MakeJournalError(
+			FResult Result = MakeJournalError(
 				ECasErrorCode::CommitJournalFailed, EIoErrorCode::FileSeekFailed, FString::Printf(TEXT("Failed to validate journal footer in '%s'"), *Transaction.JournalFile));
+			FOnDemandInstallCacheStats::OnJournalCommit(Result, 0);
+			return Result;
 		}
 
 		FFooter Footer;
 		if ((FileHandle->Read(reinterpret_cast<uint8*>(&Footer), FFooter::Size()) == false) || (Footer.IsValid() == false))
 		{
-			return Result = MakeJournalError(
+			FResult Result = MakeJournalError(
 				ECasErrorCode::CommitJournalFailed, EIoErrorCode::ReadError, FString::Printf(TEXT("Found to validate journal footer in '%s'"), *Transaction.JournalFile));
+			FOnDemandInstallCacheStats::OnJournalCommit(Result, 0);
+			return Result;
 		}
 	}
 
@@ -1156,23 +1159,29 @@ FResult FCasJournal::Commit(FTransaction&& Transaction, uint64& OutByteCount, ui
 
 		if (EntriesPos < 0) 
 		{
-			return Result = MakeJournalError(
+			FResult Result = MakeJournalError(
 				ECasErrorCode::CommitJournalFailed, EIoErrorCode::FileOpenFailed, FString::Printf(TEXT("Failed to open journal '%s'"), *Transaction.JournalFile));
+			FOnDemandInstallCacheStats::OnJournalCommit(Result, 0);
+			return Result;
 		}
 
 		if (FileHandle->Seek(EntriesPos) == false)
 		{
-			return Result = MakeJournalError(
+			FResult Result = MakeJournalError(
 				ECasErrorCode::CommitJournalFailed, EIoErrorCode::FileSeekFailed, FString::Printf(TEXT("Failed to seek to journal entries offset %" INT64_FMT " in '%s'"), EntriesPos, *Transaction.JournalFile));
+			FOnDemandInstallCacheStats::OnJournalCommit(Result, 0);
+			return Result;
 		}
 
-		TotalEntrySize = Transaction.Entries.Num() * FEntry::Size();
+		const int64 TotalEntrySize = Transaction.Entries.Num() * FEntry::Size();
 		if (FileHandle->Write(
 			reinterpret_cast<const uint8*>(Transaction.Entries.GetData()),
 			TotalEntrySize) == false)
 		{
-			return Result = MakeJournalError(
+			FResult Result = MakeJournalError(
 				ECasErrorCode::CommitJournalFailed, EIoErrorCode::WriteError, FString::Printf(TEXT("Failed to write journal entries to '%s'"), *Transaction.JournalFile));
+			FOnDemandInstallCacheStats::OnJournalCommit(Result, 0);
+			return Result;
 		}
 		OutOpCount++;
 		OutByteCount += TotalEntrySize;
@@ -1181,23 +1190,28 @@ FResult FCasJournal::Commit(FTransaction&& Transaction, uint64& OutByteCount, ui
 		FMemory::Memcpy(&Footer.Magic, &FFooter::MagicSequence, sizeof(FFooter::MagicSequence));
 		if (FileHandle->Write(reinterpret_cast<uint8*>(&Footer), FFooter::Size()) == false)
 		{
-			return Result = MakeJournalError(
+			FResult Result = MakeJournalError(
 				ECasErrorCode::CommitJournalFailed, EIoErrorCode::WriteError, FString::Printf(TEXT("Failed to write journal footer to '%s'"), *Transaction.JournalFile));
+			FOnDemandInstallCacheStats::OnJournalCommit(Result, 0);
+			return Result;
 		}
 		OutOpCount++;
 		OutByteCount += uint64(FFooter::Size());
 
 		if (FileHandle->Flush() == false)
 		{
-			return Result = MakeJournalError(
+			FResult Result = MakeJournalError(
 				ECasErrorCode::CommitJournalFailed, EIoErrorCode::FileFlushFailed, FString::Printf(TEXT("Failed to flush journal entries to '%s'"), *Transaction.JournalFile));
+			FOnDemandInstallCacheStats::OnJournalCommit(Result, 0);
+			return Result;
 		}
 		OutOpCount++;
 
 		UE_LOG(LogIoStoreOnDemand, Log, TEXT("Committed %d CAS journal entries of total %.2lf KiB to '%s'"),
 			Transaction.Entries.Num(), ToKiB(TotalEntrySize), *Transaction.JournalFile);
 
-		check(Result.HasValue());
+		FResult Result = MakeValue();
+		FOnDemandInstallCacheStats::OnJournalCommit(Result, TotalEntrySize);
 		return Result;
 	}
 }
@@ -3002,17 +3016,16 @@ FResult FOnDemandInstallCache::Defrag(
 {
 	using namespace UE::UnifiedError::IoStoreOnDemand;
 
-	if (TotalBytesToFree && *TotalBytesToFree == 0)
+	if (OutTotalFreedBytes)
 	{
-		if (OutTotalFreedBytes)
-		{
-			*OutTotalFreedBytes = 0;
-		}
+		*OutTotalFreedBytes = 0;
+	}
 
+	if (TotalBytesToFree && *TotalBytesToFree == 0)
+	{
 		return MakeValue();
 	}
 
-	FResult Result						= MakeValue();
 	uint64 TotalCachedBytes				= 0;
 	uint64 TotalFragmentedBytes			= 0;
 	uint64 BlocksToDefragFragmentedSize	= 0;
@@ -3095,7 +3108,9 @@ FResult FOnDemandInstallCache::Defrag(
 			UE_LOGFMT(LogIoStoreOnDemand, Error, "{ErrorMessage}", ErrorMessage);
 			FInstallCacheErrorContext Ctx = MakeInstallCacheErrorContext(TotalCachedBytes, *TotalBytesToFree);
 			Ctx.IoError = FIoStatus(EIoErrorCode::OutOfDiskSpace, ErrorMessage);
-			return Result = MakeError(InstallCacheDefragError::MakeError(MoveTemp(Ctx), UE::UnifiedError::EDetailFilter::All));
+			FResult Result = MakeError(InstallCacheDefragError::MakeError(MoveTemp(Ctx), UE::UnifiedError::EDetailFilter::All));
+			FOnDemandInstallCacheStats::OnDefrag(Result, TotalFragmentedBytes, TotalCachedBytes, BlocksToDefragFragmentedSize, BlocksToDefragTotalSize);
+			return Result;
 		}
 	}
 	else
@@ -3136,22 +3151,15 @@ FResult FOnDemandInstallCache::Defrag(
 		{
 			// Already defragged
 			UE_LOG(LogIoStoreOnDemand, Display, TEXT("Cache not fragmented."));
-			return Result = MakeValue();
+			FResult Result = MakeValue();
+			FOnDemandInstallCacheStats::OnDefrag(Result, TotalFragmentedBytes, TotalCachedBytes, BlocksToDefragFragmentedSize, BlocksToDefragTotalSize);
+			return Result;
 		}
 	}
 
 	UE_LOG(LogIoStoreOnDemand, Display, TEXT("Defrag found %" UINT64_FMT " fragmented bytes of %" UINT64_FMT " total bytes in %i blocks."), 
 		BlocksToDefragFragmentedSize, BlocksToDefragTotalSize, BlocksToDefrag.Num()); 
 
-	ON_SCOPE_EXIT
-	{
-		if (OutTotalFreedBytes != nullptr)
-		{
-			*OutTotalFreedBytes = Result.HasError() ? 0 : BlocksToDefragFragmentedSize;
-		}
-		FOnDemandInstallCacheStats::OnDefrag(Result, TotalFragmentedBytes, TotalCachedBytes, BlocksToDefragFragmentedSize, BlocksToDefragTotalSize);
-	};
-
 	if (TotalCachedBytes > MaxCacheSize)
 	{
 		// Ruh-Roh! There's not enough of the disk quota left to run a defrag!
@@ -3160,8 +3168,8 @@ FResult FOnDemandInstallCache::Defrag(
 		UE_LOG(LogIoStoreOnDemand, Error, TEXT("%s"), *ErrorMsg);
 
 		FInstallCacheErrorContext Ctx = MakeInstallCacheErrorContext(TotalCachedBytes, TotalBytesToFree ? *TotalBytesToFree : TotalFragmentedBytes);
-		Ctx.IoError = FIoStatus(EIoErrorCode::OutOfDiskSpace, ErrorMsg),
-			Result = MakeError(InstallCacheDefragError::MakeError(MoveTemp(Ctx), UE::UnifiedError::EDetailFilter::All));
+		Ctx.IoError = FIoStatus(EIoErrorCode::OutOfDiskSpace, ErrorMsg);
+		FResult Result = MakeError(InstallCacheDefragError::MakeError(MoveTemp(Ctx), UE::UnifiedError::EDetailFilter::All));
 
 		// Append a critical error entry to clear the cache at next startup
 		FCasJournal::FTransaction Transaction = FCasJournal::Begin(GetJournalFilename());
@@ -3171,6 +3179,7 @@ FResult FOnDemandInstallCache::Defrag(
 			UE_LOGFMT(LogIoStoreOnDemand, Error, "Failed to commit critical errors to CAS journal ({Error})", CommitResult.GetError());
 		}
 
+		FOnDemandInstallCacheStats::OnDefrag(Result, TotalFragmentedBytes, TotalCachedBytes, BlocksToDefragFragmentedSize, BlocksToDefragTotalSize);
 		return Result;
 	}
 
@@ -3236,7 +3245,9 @@ FResult FOnDemandInstallCache::Defrag(
 
 				FInstallCacheErrorContext Ctx = MakeInstallCacheErrorContext(TotalCachedBytes, TotalBytesToFree ? *TotalBytesToFree : TotalFragmentedBytes);
 				FileOpenResult.GetError().PushErrorContext(MoveTemp(Ctx));
-				return Result = MakeError(FileOpenResult.StealError());
+				FResult Result = MakeError(FileOpenResult.StealError());
+				FOnDemandInstallCacheStats::OnDefrag(Result, TotalFragmentedBytes, TotalCachedBytes, BlocksToDefragFragmentedSize, BlocksToDefragTotalSize);
+				return Result;
 			}
 
 			FSharedFileHandle FileHandle = FileOpenResult.StealValue();
@@ -3273,7 +3284,9 @@ FResult FOnDemandInstallCache::Defrag(
 				{
 					FInstallCacheErrorContext Ctx = MakeInstallCacheErrorContext(TotalCachedBytes, TotalBytesToFree ? *TotalBytesToFree : TotalFragmentedBytes);
 					ReadResult.GetError().PushErrorContext(MoveTemp(Ctx));
-					return Result = MakeError(ReadResult.StealError());
+					FResult Result = MakeError(ReadResult.StealError());
+					FOnDemandInstallCacheStats::OnDefrag(Result, TotalFragmentedBytes, TotalCachedBytes, BlocksToDefragFragmentedSize, BlocksToDefragTotalSize);
+					return Result;
 				}
 
 				const FIoHash ChunkHash = FIoHash::HashBuffer(Buffer.GetView());
@@ -3290,25 +3303,29 @@ FResult FOnDemandInstallCache::Defrag(
 						UE_LOGFMT(LogIoStoreOnDemand, Error, "Failed to commit critical errors to CAS journal ({Error})", CommitResult.GetError());
 					}
 
-					if (Result = FlushPendingChunks(DefragPendingChunks, DefragBlock.LastAccess); Result.HasError())
+					if (FResult Result = FlushPendingChunks(DefragPendingChunks, DefragBlock.LastAccess); Result.HasError())
 					{
+						FOnDemandInstallCacheStats::OnDefrag(Result, TotalFragmentedBytes, TotalCachedBytes, BlocksToDefragFragmentedSize, BlocksToDefragTotalSize);
 						return Result;
 					}
 					check(DefragPendingChunks.IsEmpty());
 	
-					return Result = MakeError(InstallCacheDefragError::MakeError(
+					FResult Result = MakeError(InstallCacheDefragError::MakeError(
 						FChunkHashMismatchErrorContext
 						{
 							.ChunkId		= ReffedChunk.ChunkId, // May not be unique, just return the first found for debugging purposes
 							.ExpectedHash	= ReffedChunk.Hash,
 							.ActualHash		= ChunkHash
 						}));
+					FOnDemandInstallCacheStats::OnDefrag(Result, TotalFragmentedBytes, TotalCachedBytes, BlocksToDefragFragmentedSize, BlocksToDefragTotalSize);
+					return Result;
 				}
 
 				if (DefragPendingChunks.TotalSize > FPendingChunks::MaxPendingBytes)
 				{
-					if (Result = FlushPendingChunks(DefragPendingChunks, DefragBlock.LastAccess); Result.HasError())
+					if (FResult Result = FlushPendingChunks(DefragPendingChunks, DefragBlock.LastAccess); Result.HasError())
 					{
+						FOnDemandInstallCacheStats::OnDefrag(Result, TotalFragmentedBytes, TotalCachedBytes, BlocksToDefragFragmentedSize, BlocksToDefragTotalSize);
 						return Result;
 					}
 					check(DefragPendingChunks.IsEmpty());
@@ -3319,8 +3336,9 @@ FResult FOnDemandInstallCache::Defrag(
 
 			FileHandle.Reset();
 
-			if (Result = FlushPendingChunks(DefragPendingChunks, DefragBlock.LastAccess); Result.HasError()) 
+			if (FResult Result = FlushPendingChunks(DefragPendingChunks, DefragBlock.LastAccess); Result.HasError())
 			{
+				FOnDemandInstallCacheStats::OnDefrag(Result, TotalFragmentedBytes, TotalCachedBytes, BlocksToDefragFragmentedSize, BlocksToDefragTotalSize);
 				return Result;
 			}
 			check(DefragPendingChunks.IsEmpty());
@@ -3331,8 +3349,9 @@ FResult FOnDemandInstallCache::Defrag(
 		// Flushing should overwrite the lookup info for the cas addr to point at the new block.
 		// Can now remove the old block
 		TArray<FCasAddr> DeletedChunkAddresses;
-		if (Result = Cas.DeleteBlock(DefragBlock.BlockId, DeletedChunkAddresses); Result.HasError())
+		if (FResult Result = Cas.DeleteBlock(DefragBlock.BlockId, DeletedChunkAddresses); Result.HasError())
 		{
+			FOnDemandInstallCacheStats::OnDefrag(Result, TotalFragmentedBytes, TotalCachedBytes, BlocksToDefragFragmentedSize, BlocksToDefragTotalSize);
 			return Result;
 		}
 
@@ -3344,8 +3363,9 @@ FResult FOnDemandInstallCache::Defrag(
 		}
 		Transaction.BlockDeleted(DefragBlock.BlockId);
 
-		if (Result = FCasJournal::Commit(MoveTemp(Transaction)); Result.HasError())
+		if (FResult Result = FCasJournal::Commit(MoveTemp(Transaction)); Result.HasError())
 		{
+			FOnDemandInstallCacheStats::OnDefrag(Result, TotalFragmentedBytes, TotalCachedBytes, BlocksToDefragFragmentedSize, BlocksToDefragTotalSize);
 			return Result;
 		}
 	}
@@ -3353,6 +3373,13 @@ FResult FOnDemandInstallCache::Defrag(
 	UE_LOG(LogIoStoreOnDemand, Display, TEXT("Defrag removed %" UINT64_FMT " fragmented bytes of %" UINT64_FMT " total bytes in %i blocks."),
 		BlocksToDefragTotalSize, BlocksToDefragTotalSize, BlocksToDefrag.Num());
 
+	if (OutTotalFreedBytes)
+	{
+		*OutTotalFreedBytes = BlocksToDefragFragmentedSize;
+	}
+
+	FResult Result = MakeValue();
+	FOnDemandInstallCacheStats::OnDefrag(Result, TotalFragmentedBytes, TotalCachedBytes, BlocksToDefragFragmentedSize, BlocksToDefragTotalSize);
 	return Result;
 }
 
@@ -3502,14 +3529,12 @@ FResult FOnDemandInstallCache::FlushPendingChunksImpl(const FPendingChunks& Chun
 {
 	using namespace UE::UnifiedError::IoStoreOnDemand;
 
-	FResult Result				= MakeValue();
 	uint64 TotalCasBytes		= 0;
 	uint64 TotalJournalBytes	= 0;
 	uint32 TotalOpCount			= 0;
 
 	ON_SCOPE_EXIT 
 	{ 
-		FOnDemandInstallCacheStats::OnFlush(Result, TotalCasBytes);
 		Governor.OnInstallCacheFlushed(TotalCasBytes + TotalJournalBytes, TotalOpCount);
 	};
 
@@ -3540,7 +3565,9 @@ FResult FOnDemandInstallCache::FlushPendingChunksImpl(const FPendingChunks& Chun
 		{
 			FInstallCacheErrorContext Ctx = MakeInstallCacheErrorContext();
 			OpenWriteResult.GetError().PushErrorContext(MoveTemp(Ctx));
-			return Result = MakeError(OpenWriteResult.StealError());
+			FResult Result = MakeError(OpenWriteResult.StealError());
+			FOnDemandInstallCacheStats::OnFlush(Result, TotalCasBytes);
+			return Result;
 		}
 		
 		FUniqueFileHandle CasFileHandle = OpenWriteResult.StealValue();
@@ -3575,10 +3602,11 @@ FResult FOnDemandInstallCache::FlushPendingChunksImpl(const FPendingChunks& Chun
 
 			if (CasFileHandle->Write(Ar.GetData(), BytesToWrite) == false)
 			{
-				Result = MakeCasError<void>(
+				FResult Result = MakeCasError<void>(
 					ECasErrorCode::WriteBlockFailed, EIoErrorCode::WriteError, FString::Printf(TEXT("Failed to write %" INT64_FMT " bytes to CAS block %u"), BytesToWrite, CurrentBlockId.Id));
 				FInstallCacheErrorContext Ctx = MakeInstallCacheErrorContext();
 				Result.GetError().PushErrorContext(MoveTemp(Ctx));
+				FOnDemandInstallCacheStats::OnFlush(Result, TotalCasBytes);
 				return Result;
 			}
 			TotalOpCount++;
@@ -3586,10 +3614,11 @@ FResult FOnDemandInstallCache::FlushPendingChunksImpl(const FPendingChunks& Chun
 
 			if (CasFileHandle->Flush() == false)
 			{
-				Result = MakeCasError<void>(
+				FResult Result = MakeCasError<void>(
 					ECasErrorCode::WriteBlockFailed, EIoErrorCode::FileFlushFailed, FString::Printf(TEXT("Failed to flush %" INT64_FMT " bytes to CAS block %u"), BytesToWrite, CurrentBlockId.Id));
 				FInstallCacheErrorContext Ctx = MakeInstallCacheErrorContext();
 				Result.GetError().PushErrorContext(MoveTemp(Ctx));
+				FOnDemandInstallCacheStats::OnFlush(Result, TotalCasBytes);
 				return Result;
 			}
 			TotalOpCount++;
@@ -3630,10 +3659,11 @@ FResult FOnDemandInstallCache::FlushPendingChunksImpl(const FPendingChunks& Chun
 
 		uint64 JrnlBytesWritten	= 0;
 		uint32 JrnlOps			= 0;
-		if (Result = FCasJournal::Commit(MoveTemp(Transaction), JrnlBytesWritten, JrnlOps); Result.HasError())
+		if (FResult Result = FCasJournal::Commit(MoveTemp(Transaction), JrnlBytesWritten, JrnlOps); Result.HasError())
 		{
 			FInstallCacheErrorContext Ctx = MakeInstallCacheErrorContext();
 			Result.GetError().PushErrorContext(MoveTemp(Ctx));
+			FOnDemandInstallCacheStats::OnFlush(Result, TotalCasBytes);
 			return Result;
 		}
 
@@ -3646,6 +3676,8 @@ FResult FOnDemandInstallCache::FlushPendingChunksImpl(const FPendingChunks& Chun
 		}
 	}
 
+	FResult Result = MakeValue();
+	FOnDemandInstallCacheStats::OnFlush(Result, TotalCasBytes);
 	return Result;
 }
 
