Minor clean up to IAX statistics code to make it easier to follow.

#rnx

- Moved the recording of the max duration of requests out of FHttpRecentHistoryStatistics (which is supposed to maintain a running average of values for a recent number of requests) to it's own value.

#Synced-CL 48498703

diff --git a/Engine/Source/Runtime/Experimental/IoStore/OnDemand/Private/Statistics.cpp b/Engine/Source/Runtime/Experimental/IoStore/OnDemand/Private/Statistics.cpp
--- a/Engine/Source/Runtime/Experimental/IoStore/OnDemand/Private/Statistics.cpp
+++ b/Engine/Source/Runtime/Experimental/IoStore/OnDemand/Private/Statistics.cpp
@@ -403,7 +403,6 @@ struct FHttpRecentHistoryStatistics
 	int64			Duration[HistoryCount] = {};
 	int64			Bytes[HistoryCount] = {};
 	int64			TotalDuration = 0;
-	int64			MaxDuration = 0;
 	int64			TotalBytes = 0;
 	int64 			Index = 0;
 
@@ -420,8 +419,6 @@ struct FHttpRecentHistoryStatistics
 		TotalBytes += SizeBytes;
 		Bytes[Index] = SizeBytes;
 
-		MaxDuration = FMath::Max(NewDuration, MaxDuration);
-
 		Index = (Index + 1) % HistoryCount;
 	}
 
@@ -434,11 +431,6 @@ struct FHttpRecentHistoryStatistics
 	{
 		return static_cast<double>(TotalDuration) / static_cast<double>(HistoryCount);
 	}
-
-	int64 GetMaxDuration() const
-	{
-		return MaxDuration;
-	}
 };
 
 FHttpRecentHistoryStatistics GHttpHistory[(uint8)EHttpRequestType::NUM_SOURCES];
@@ -1361,18 +1353,21 @@ void FOnDemandIoBackendStats::OnHttpDequeue(EHttpRequestType Type)
 
 void FOnDemandIoBackendStats::OnHttpGet(EHttpRequestType Type, uint64 SizeBytes, uint64 DurationMs)
 {
-	GHttpPendingCount[(uint8)Type].Decrement();
-	GHttpInflightCount[(uint8)Type].Decrement();
-	GHttpGetCount[(uint8)Type].Increment();
-	GHttpDownloadedBytes[(uint8)Type].Add(SizeBytes);
-	GHttpDurationMsSum[(uint8)Type] += DurationMs;
-	GHttpDurationMs[(uint8)Type].Set(DurationMs);
+	const uint8 Index = (uint8)Type;
+
+	GHttpPendingCount[Index].Decrement();
+	GHttpInflightCount[Index].Decrement();
+	GHttpGetCount[Index].Increment();
+	GHttpDownloadedBytes[Index].Add(SizeBytes);
+	GHttpDurationMsSum[Index] += DurationMs;
+	GHttpDurationMs[Index].Set(DurationMs);
+
+	GHttpHistory[Index].OnGet(SizeBytes, DurationMs);
 
-	GHttpHistory[(uint8)Type].OnGet(SizeBytes, DurationMs);
+	GHttpBandwidthMbps[Index].Set(GHttpHistory[Index].GetBandwidthMbps());
+	GHttpDurationMsAvg[Index] = GHttpHistory[Index].GetAverage();
 
-	GHttpBandwidthMbps[(uint8)Type].Set(GHttpHistory[(uint8)Type].GetBandwidthMbps());
-	GHttpDurationMsAvg[(uint8)Type] = GHttpHistory[(uint8)Type].GetAverage();
-	GHttpDurationMsMax[(uint8)Type] = (int32)GHttpHistory[(uint8)Type].GetMaxDuration();
+	GHttpDurationMsMax[Index] = FMath::Max(GHttpDurationMsMax[Index], (int32)DurationMs);
 }
 
 void FOnDemandIoBackendStats::OnHttpCancel(EHttpRequestType Type)
