Change Dataflow context to support FText for logs

#Synced-CL 48601912

diff --git a/Engine/Plugins/Experimental/Dataflow/Source/DataflowEditor/Private/Dataflow/DataflowEditorToolkit.cpp b/Engine/Plugins/Experimental/Dataflow/Source/DataflowEditor/Private/Dataflow/DataflowEditorToolkit.cpp
--- a/Engine/Plugins/Experimental/Dataflow/Source/DataflowEditor/Private/Dataflow/DataflowEditorToolkit.cpp
+++ b/Engine/Plugins/Experimental/Dataflow/Source/DataflowEditor/Private/Dataflow/DataflowEditorToolkit.cpp
@@ -132,7 +132,7 @@ namespace UE::Dataflow::Private
 				}
 				else
 				{
-					DebugDrawComponent = NewObject<UDataflowDebugDrawComponent>(RootActor, FName(LOCTEXT("DataflowDebugDrawComponent", "Dataflow Debug Draw Component").ToString(), RF_Transient));
+					DebugDrawComponent = NewObject<UDataflowDebugDrawComponent>(RootActor, FName(LOCTEXT("DataflowDebugDrawComponent", "Dataflow Debug Draw Component").ToString()), RF_Transient);
 					DebugDrawComponent->RegisterComponentWithWorld(DataflowScene->GetWorld());
 				}
 			}
@@ -557,9 +557,7 @@ void FDataflowEditorToolkit::RegisterContextHandlers()
 		OnNodeBeginEvaluateMulticastDelegateHandle = DataflowContext->OnNodeBeginEvaluateMulticast.AddSP(this, &FDataflowEditorToolkit::OnNodeBeginEvaluate);
 		OnNodeFinishEvaluateMulticastDelegateHandle = DataflowContext->OnNodeFinishEvaluateMulticast.AddSP(this, &FDataflowEditorToolkit::OnNodeFinishEvaluate);
 
-		OnContextHasInfoDelegateHandle = DataflowContext->OnContextHasInfo.AddSP(this, &FDataflowEditorToolkit::OnContextHasInfo);
-		OnContextHasWarningDelegateHandle = DataflowContext->OnContextHasWarning.AddSP(this, &FDataflowEditorToolkit::OnContextHasWarning);
-		OnContextHasErrorDelegateHandle = DataflowContext->OnContextHasError.AddSP(this, &FDataflowEditorToolkit::OnContextHasError);
+		OnContextLogDelegateHandle = DataflowContext->GetOnContextLogMulticast().AddSP(this, &FDataflowEditorToolkit::OnContextLog);
 	}
 
 	FDataflowAssetDelegates::OnNodeInvalidated.AddSP(this, &FDataflowEditorToolkit::OnNodeInvalidated);
@@ -585,9 +583,7 @@ void FDataflowEditorToolkit::UnregisterContextHandlers()
 		DataflowContext->OnNodeBeginEvaluateMulticast.Remove(OnNodeBeginEvaluateMulticastDelegateHandle);
 		DataflowContext->OnNodeFinishEvaluateMulticast.Remove(OnNodeFinishEvaluateMulticastDelegateHandle);
 
-		DataflowContext->OnContextHasInfo.Remove(OnContextHasInfoDelegateHandle);
-		DataflowContext->OnContextHasWarning.Remove(OnContextHasWarningDelegateHandle);
-		DataflowContext->OnContextHasError.Remove(OnContextHasErrorDelegateHandle);	
+		DataflowContext->GetOnContextLogMulticast().Remove(OnContextLogDelegateHandle);
 	}
 
 	FDataflowAssetDelegates::OnNodeInvalidated.RemoveAll(this);
@@ -2069,6 +2065,12 @@ void FDataflowEditorToolkit::SetDataflowPathFromNodeAndOutput(const FDataflowNod
 }
 
 void FDataflowEditorToolkit::LogMessage(const EMessageSeverity::Type Severity, const FDataflowNode* Node, const FDataflowOutput* Output, const FString& Message) const
+{
+	FText TextMessage = FText::FromString(Message);
+	LogMessage(Severity, Node, Output, TextMessage);
+}
+
+void FDataflowEditorToolkit::LogMessage(const EMessageSeverity::Type Severity, const FDataflowNode* Node, const FDataflowOutput* Output, const FText& Message) const
 {
 	if (DataflowOutputLog)
 	{
@@ -2078,22 +2080,13 @@ void FDataflowEditorToolkit::LogMessage(const EMessageSeverity::Type Severity, c
 	}
 }
 
-void FDataflowEditorToolkit::OnContextHasInfo(const FDataflowNode* Node, const FDataflowOutput* Output, const FString& Info)
-{
-	LogMessage(EMessageSeverity::Info, Node, Output, Info);
-}
-
-void FDataflowEditorToolkit::OnContextHasWarning(const FDataflowNode* Node, const FDataflowOutput* Output, const FString& Warning)
+void FDataflowEditorToolkit::OnContextLog(const UE::Dataflow::FContext::FLogMessage& InLogMessage)
 {
-	LogMessage(EMessageSeverity::Warning, Node, Output, Warning);
-}
-
-void FDataflowEditorToolkit::OnContextHasError(const FDataflowNode* Node, const FDataflowOutput* Output, const FString& Error)
-{
-	LogMessage(EMessageSeverity::Error, Node, Output, Error);
-	if (TSharedPtr<SDockTab> OutputLogDockTab = WeakOutputLogDockTab.Pin())
+	if (DataflowOutputLog)
 	{
-		OutputLogDockTab->DrawAttention();
+		TSharedPtr<const FDataflowNode> Node = InLogMessage.Node.Pin();
+		const FDataflowOutput* Output = Node ? Node->FindOutput(InLogMessage.OutputGuid) : nullptr;
+		LogMessage(InLogMessage.Severity, Node.Get(), Output, InLogMessage.Text);
 	}
 }
 
diff --git a/Engine/Plugins/Experimental/Dataflow/Source/DataflowEditor/Private/Dataflow/DataflowOutputLog.cpp b/Engine/Plugins/Experimental/Dataflow/Source/DataflowEditor/Private/Dataflow/DataflowOutputLog.cpp
--- a/Engine/Plugins/Experimental/Dataflow/Source/DataflowEditor/Private/Dataflow/DataflowOutputLog.cpp
+++ b/Engine/Plugins/Experimental/Dataflow/Source/DataflowEditor/Private/Dataflow/DataflowOutputLog.cpp
@@ -67,7 +67,7 @@ void FDataflowOutputLog::ClearMessageLog()
 
 namespace UE::Dataflow::Private
 {
-	TSharedRef<FTokenizedMessage> BuildTokenizedMessage(const EMessageSeverity::Type InSeverity, const FString& InMessage, const FDataflowPath& InPath)
+	TSharedRef<FTokenizedMessage> BuildTokenizedMessage(const EMessageSeverity::Type InSeverity, const FText& InMessage, const FDataflowPath& InPath)
 	{
 		// Token0 - Severity
 		// Token1 - Time code
@@ -83,7 +83,7 @@ namespace UE::Dataflow::Private
 		TSharedRef<IMessageToken> TokenPath = FTextToken::Create(FText::FromString("[" + InPath.ToString() + "]"));
 		TokenizedMessage->AddToken(TokenPath);
 		TokenizedMessage->SetMessageLink(TokenPath);
-		TSharedRef<IMessageToken> TokenMessage = FTextToken::Create(FText::FromString(InMessage));
+		TSharedRef<IMessageToken> TokenMessage = FTextToken::Create(InMessage);
 		TokenizedMessage->AddToken(TokenMessage);
 
 		return TokenizedMessage;
@@ -102,14 +102,18 @@ namespace UE::Dataflow::Private
 
 void FDataflowOutputLog::AddMessage(const EMessageSeverity::Type InSeverity, const FString& InMessage, const FDataflowPath& InPath)
 {
-	TSharedRef<FTokenizedMessage> TokenizedMessage = UE::Dataflow::Private::BuildTokenizedMessage(InSeverity, InMessage, InPath);
+	const FText TextMessage = FText::FromString(InMessage);
+	AddMessage(InSeverity, TextMessage, InPath);
+}
 
+void FDataflowOutputLog::AddMessage(const EMessageSeverity::Type InSeverity, const FText& InMessage, const FDataflowPath& InPath)
+{
 	if (MessageLogListing.IsValid())
 	{
+		TSharedRef<FTokenizedMessage> TokenizedMessage = UE::Dataflow::Private::BuildTokenizedMessage(InSeverity, InMessage, InPath);
 		MessageLogListing->AddMessage(TokenizedMessage, false);
 	}
-
-	UE_LOG(LogDataflowContextOutput, Verbose, TEXT("[%s][%s]"), *InPath.ToString(), *InMessage);
+	UE_LOG(LogDataflowContextOutput, Verbose, TEXT("[%s][%s]"), *InPath.ToString(), *InMessage.ToString());
 }
 
 void FDataflowOutputLog::CreateMessageLogWidget()
diff --git a/Engine/Plugins/Experimental/Dataflow/Source/DataflowEditor/Private/Dataflow/DataflowSNode.cpp b/Engine/Plugins/Experimental/Dataflow/Source/DataflowEditor/Private/Dataflow/DataflowSNode.cpp
--- a/Engine/Plugins/Experimental/Dataflow/Source/DataflowEditor/Private/Dataflow/DataflowSNode.cpp
+++ b/Engine/Plugins/Experimental/Dataflow/Source/DataflowEditor/Private/Dataflow/DataflowSNode.cpp
@@ -125,21 +125,24 @@ TSharedPtr<SGraphPin> SDataflowEdNode::CreatePinWidget(UEdGraphPin* Pin) const
 			{
 				if (FDataflowOutput* Output = DataflowNode->FindOutput(Pin->GetFName()))
 				{
-					if (const TSharedPtr<UE::Dataflow::FContext> DataflowContext = DataflowInterface->GetDataflowContext())
+					if (DataflowInterface)
 					{
-						TSet<UE::Dataflow::FContextCacheKey> CacheKeys;
-						const int32 NumKeys = DataflowContext->GetKeys(CacheKeys);
-
-						//
-						// DataStore is empty or 
-						// CacheKey is not in DataStore or
-						// Node's Timestamp is invalid or
-						// Node's Timestamp is greater than CacheKey's Timestamp -> Pin is invalid
-						//
-						bIsOutputInvalid = !NumKeys ||
-							!CacheKeys.Contains(Output->CacheKey()) ||
-							DataflowNode->GetTimestamp().IsInvalid() ||
-							!DataflowContext->IsCacheEntryAfterTimestamp(Output->CacheKey(), DataflowNode->GetTimestamp());
+						if (const TSharedPtr<UE::Dataflow::FContext> DataflowContext = DataflowInterface->GetDataflowContext())
+						{
+							TSet<UE::Dataflow::FContextCacheKey> CacheKeys;
+							const int32 NumKeys = DataflowContext->GetKeys(CacheKeys);
+
+							//
+							// DataStore is empty or 
+							// CacheKey is not in DataStore or
+							// Node's Timestamp is invalid or
+							// Node's Timestamp is greater than CacheKey's Timestamp -> Pin is invalid
+							//
+							bIsOutputInvalid = !NumKeys ||
+								!CacheKeys.Contains(Output->CacheKey()) ||
+								DataflowNode->GetTimestamp().IsInvalid() ||
+								!DataflowContext->IsCacheEntryAfterTimestamp(Output->CacheKey(), DataflowNode->GetTimestamp());
+						}
 					}
 				}
 			}
@@ -181,7 +184,7 @@ TSharedRef<SWidget> SDataflowEdNode::CreateTitleRightWidget()
 			})
 		.OnCheckStateChanged_Lambda([this](const ECheckBoxState NewState)
 			{
-				if (DataflowGraphNode)
+				if (DataflowGraphNode && DataflowInterface)
 				{
 					const bool bShouldWireframeRenderNode = (NewState == ECheckBoxState::Checked);
 					DataflowGraphNode->SetShouldWireframeRenderNode(bShouldWireframeRenderNode);
@@ -207,6 +210,7 @@ TSharedRef<SWidget> SDataflowEdNode::CreateTitleRightWidget()
 	const TSharedPtr<const FDataflowNode> DataflowNode = DataflowGraphNode ? DataflowGraphNode->GetDataflowNode() : TSharedPtr<const FDataflowNode>();
 
 	const bool bHasRenderingCheckBox = (
+		DataflowInterface &&
 		DataflowInterface->NodesHaveToggleWidget() &&
 		DataflowNode && 
 		(DataflowNode->NumOutputs() > 0 || DataflowNode->CanDebugDraw())
@@ -228,7 +232,7 @@ TArray<FOverlayWidgetInfo> SDataflowEdNode::GetOverlayWidgets(bool bSelected, co
 {
 	TArray<FOverlayWidgetInfo> Widgets = SGraphNode::GetOverlayWidgets(bSelected, WidgetSize);
 
-	if (DataflowGraphNode && DataflowGraphNode->GetDataflowNode())
+	if (DataflowGraphNode && DataflowGraphNode->GetDataflowNode() && DataflowInterface)
 	{
 		if (DataflowInterface->NodesHaveFreezeWidget())
 		{
@@ -332,7 +336,6 @@ void SDataflowEdNode::UpdateErrorInfo()
 	// Priority should be:
 	// 1. Error
 	// 2. Warning
-	// 3. Failed
 	// 4. Deprecated
 	// 5. Experimental
 
@@ -357,12 +360,6 @@ void SDataflowEdNode::UpdateErrorInfo()
 						ErrorColor = FAppStyle::GetColor("ErrorReporting.WarningBackgroundColor");
 						return;
 					}
-					else if (Context->NodeFailed(DataflowNode.Get()))
-					{
-						ErrorMsg = FString(TEXT("FAILED"));
-						ErrorColor = CFailedColor;
-						return;
-					}
 				}
 			}
 
diff --git a/Engine/Plugins/Experimental/Dataflow/Source/DataflowEditor/Public/Dataflow/DataflowEditorToolkit.h b/Engine/Plugins/Experimental/Dataflow/Source/DataflowEditor/Public/Dataflow/DataflowEditorToolkit.h
--- a/Engine/Plugins/Experimental/Dataflow/Source/DataflowEditor/Public/Dataflow/DataflowEditorToolkit.h
+++ b/Engine/Plugins/Experimental/Dataflow/Source/DataflowEditor/Public/Dataflow/DataflowEditorToolkit.h
@@ -75,6 +75,7 @@ public:
 	UE_API bool IsSimulationDataflowAsset() const;
 	UE_API FName GetGraphLogName() const;
 	UE_API void LogMessage(const EMessageSeverity::Type Severity, const FDataflowNode* Node, const FDataflowOutput* Output, const FString& Message) const;
+	UE_API void LogMessage(const EMessageSeverity::Type Severity, const FDataflowNode* Node, const FDataflowOutput* Output, const FText& Message) const;
 
 	/** Editor dataflow content accessors */
 	UE_API const TObjectPtr<UDataflowBaseContent>& GetEditorContent() const;
@@ -183,9 +184,7 @@ protected:
 	UE_API void OnBeginEvaluate();
 	UE_API void OnFinishEvaluate();
 	UE_API void OnOutputLogMessageTokenClicked(const FString TokenString);
-	UE_API void OnContextHasInfo(const FDataflowNode* Node, const FDataflowOutput* Output, const FString& Info);
-	UE_API void OnContextHasWarning(const FDataflowNode* Node, const FDataflowOutput* Output, const FString& Warning);
-	UE_API void OnContextHasError(const FDataflowNode* Node, const FDataflowOutput* Output, const FString& Error);
+	UE_API void OnContextLog(const UE::Dataflow::FContext::FLogMessage& InLogMessage);
 
 private:
 	
@@ -368,9 +367,7 @@ private:
 	FDelegateHandle OnNodeFinishEvaluateMulticastDelegateHandle;
 	FDelegateHandle OnOutputLogMessageTokenClickedDelegateHandle;
 
-	FDelegateHandle OnContextHasInfoDelegateHandle;
-	FDelegateHandle OnContextHasWarningDelegateHandle;
-	FDelegateHandle OnContextHasErrorDelegateHandle;
+	FDelegateHandle OnContextLogDelegateHandle;
 
 	// The currently selected set of dataflow nodes. 
 	TSet<TObjectPtr<UDataflowEdNode>> SelectedDataflowNodes;
diff --git a/Engine/Plugins/Experimental/Dataflow/Source/DataflowEditor/Public/Dataflow/DataflowOutputLog.h b/Engine/Plugins/Experimental/Dataflow/Source/DataflowEditor/Public/Dataflow/DataflowOutputLog.h
--- a/Engine/Plugins/Experimental/Dataflow/Source/DataflowEditor/Public/Dataflow/DataflowOutputLog.h
+++ b/Engine/Plugins/Experimental/Dataflow/Source/DataflowEditor/Public/Dataflow/DataflowOutputLog.h
@@ -27,7 +27,11 @@ public:
 
 	TSharedRef<IMessageLogListing> GetMessageLog() const;
 	void ClearMessageLog();
+	
+	UE_DEPRECATED(5.8, "AddMessage with a FString parameter is now deprecated, use the FText version")
 	void AddMessage(const EMessageSeverity::Type InSeverity, const FString& InMessage, const FDataflowPath& InPath);
+	
+	void AddMessage(const EMessageSeverity::Type InSeverity, const FText& InMessage, const FDataflowPath& InPath);
 
 	DECLARE_MULTICAST_DELEGATE_OneParam(FOnOutputLogMessageTokenClicked, const FString);
 	FOnOutputLogMessageTokenClicked& GetOnOutputLogMessageTokenClickedDelegate() { return OnOutputLogMessageTokenClickedDelegate; }
diff --git a/Engine/Source/Runtime/Experimental/Dataflow/Core/Private/Dataflow/DataflowInputOutput.cpp b/Engine/Source/Runtime/Experimental/Dataflow/Core/Private/Dataflow/DataflowInputOutput.cpp
--- a/Engine/Source/Runtime/Experimental/Dataflow/Core/Private/Dataflow/DataflowInputOutput.cpp
+++ b/Engine/Source/Runtime/Experimental/Dataflow/Core/Private/Dataflow/DataflowInputOutput.cpp
@@ -229,9 +229,12 @@ const FDataflowInput* FDataflowOutput::GetPassthroughInput() const
 
 void FDataflowOutput::Invalidate(const UE::Dataflow::FTimestamp& ModifiedTimestamp)
 {
-	for (FDataflowConnection* Con : GetConnections())
+	for (FDataflowConnection* Connection : GetConnections())
 	{
-		Con->Invalidate(ModifiedTimestamp);
+		if (Connection)
+		{
+			Connection->Invalidate(ModifiedTimestamp);
+		}
 	}
 }
 
@@ -295,7 +298,7 @@ bool FDataflowOutput::EvaluateImpl(UE::Dataflow::FContext& Context) const
 		OwningNode->Evaluate(Context, this);
 		
 		// Validation
-		if (!Context.NodeHasError(OwningNode) && !Context.NodeFailed(OwningNode))
+		if (!Context.NodeHasError(OwningNode))
 		{
 			if (!Context.HasData(CacheKey()))
 			{
@@ -462,5 +465,5 @@ void* FDataflowArrayOutput::RealAddress() const
 
 bool FDataflowOutput::HasNodeFailedOrErrored(UE::Dataflow::FContext& Context) const
 {
-	return Context.NodeHasError(GetOwningNode()) || Context.NodeFailed(GetOwningNode());
+	return Context.NodeHasError(GetOwningNode());
 }
diff --git a/Engine/Source/Runtime/Experimental/Dataflow/Core/Private/Dataflow/DataflowNodeParameters.cpp b/Engine/Source/Runtime/Experimental/Dataflow/Core/Private/Dataflow/DataflowNodeParameters.cpp
--- a/Engine/Source/Runtime/Experimental/Dataflow/Core/Private/Dataflow/DataflowNodeParameters.cpp
+++ b/Engine/Source/Runtime/Experimental/Dataflow/Core/Private/Dataflow/DataflowNodeParameters.cpp
@@ -10,6 +10,7 @@
 #include "Serialization/Archive.h"
 #include "GeometryCollection/ManagedArrayCollection.h"
 #include "Modules/ModuleManager.h"
+#include "Misc/ScopeLock.h"
 
 namespace UE::Dataflow
 {
@@ -382,93 +383,142 @@ namespace UE::Dataflow
 		return false;
 	}
 
-	void FContext::Info(const FString& InInfo, const FDataflowNode* InNode, const FDataflowOutput* InOutput)
+	void FContext::Info(const FText& InInfo, const FDataflowNode* InNode, const FDataflowOutput* InOutput)
 	{
-		if (!IsThreaded())
+		if (InNode)
 		{
-			if (OnContextHasInfo.IsBound())
-			{
-				OnContextHasInfo.Broadcast(InNode, InOutput, InInfo);
-			}
-			NodesWithInfo.Add(InNode);
+			FScopeLock Lock(&NodeTrackingLock);
+			NodesWithInfo.Add(InNode->GetGuid());
 		}
+		BroadCastLog(EMessageSeverity::Info, InInfo, InNode, InOutput);
+	}
+
+	void FContext::Info(const FString& InInfo, const FDataflowNode* InNode, const FDataflowOutput* InOutput)
+	{
+		FText Message = FText::FromString(InInfo);
+		Info(Message, InNode, InOutput);
 	}
 
 	int32 FContext::GetNumInfo() const
 	{
+		FScopeLock Lock(&NodeTrackingLock);
 		return NodesWithInfo.Num();
 	}
 
+	void FContext::Warning(const FText& InWarning, const FDataflowNode* InNode, const FDataflowOutput* InOutput)
+	{
+		if (InNode)
+		{
+			FScopeLock Lock(&NodeTrackingLock);
+			NodesWithWarning.Add(InNode->GetGuid());
+		}
+		BroadCastLog(EMessageSeverity::Warning, InWarning, InNode, InOutput);
+	}
+
 	void FContext::Warning(const FString& InWarning, const FDataflowNode* InNode, const FDataflowOutput* InOutput)
 	{
-		if (!IsThreaded())
+		FText Message = FText::FromString(InWarning);
+		Warning(Message, InNode, InOutput);
+	}
+
+	bool FContext::NodeHasWarning(const FDataflowNode* InNode) const
+	{
+		if (InNode)
 		{
-			if (OnContextHasWarning.IsBound())
-			{
-				OnContextHasWarning.Broadcast(InNode, InOutput, InWarning);
-			}
-			NodesWithWarning.Add(InNode);
+			FScopeLock Lock(&NodeTrackingLock);
+			return NodesWithWarning.Contains(InNode->GetGuid());
 		}
+		return false;
 	}
 
 	int32 FContext::GetNumWarnings() const
 	{
+		FScopeLock Lock(&NodeTrackingLock);
 		return NodesWithWarning.Num();
 	}
 
-	void FContext::Error(const FString& InError, const FDataflowNode* InNode, const FDataflowOutput* InOutput)
+	void FContext::Error(const FText& InError, const FDataflowNode* InNode, const FDataflowOutput* InOutput)
 	{
-		if (!IsThreaded())
+		if (InNode)
 		{
-			if (OnContextHasError.IsBound())
-			{
-				OnContextHasError.Broadcast(InNode, InOutput, InError);
-			}
-			NodesWithError.Add(InNode);
-
-			//
-		// Add all the connected nodes to NodesFailed set
-		//
-#if DATAFLOW_EDITOR_EVALUATION
-			const int32 NumNodesInCallStack = Callstack.Num();
-			if (NumNodesInCallStack > 1)
-			{
-				for (int32 Idx = NumNodesInCallStack - 2; Idx >= 0; --Idx)
-				{
-					const FDataflowNode* OwningNode = Callstack[Idx]->GetOwningNode();
-					const FString WarningStr = FString("Evaluation failed");
+			FScopeLock Lock(&NodeTrackingLock);
+			NodesWithError.Add(InNode->GetGuid());
+		}
+		BroadCastLog(EMessageSeverity::Error, InError, InNode, InOutput);
+	}
 
-					if (OnContextHasWarning.IsBound())
-					{
-						OnContextHasWarning.Broadcast(OwningNode, InOutput, WarningStr);
-					}
+	void FContext::Error(const FString& InError, const FDataflowNode* InNode, const FDataflowOutput* InOutput)
+	{
+		FText Message = FText::FromString(InError);
+		Error(Message, InNode, InOutput);
+	}
 
-					NodesFailed.Add(OwningNode);
-				}
-			}
-#endif
+	bool FContext::NodeHasError(const FDataflowNode* InNode) const
+	{
+		if (InNode)
+		{
+			FScopeLock Lock(&NodeTrackingLock);
+			return NodesWithError.Contains(InNode->GetGuid());
 		}
+		return false;
 	}
 
 	int32 FContext::GetNumErrors() const
 	{
+		FScopeLock Lock(&NodeTrackingLock);
 		return NodesWithError.Num();
 	}
 
+	void FContext::BroadCastLog(EMessageSeverity::Type Severity, const FText& InMessage, const FDataflowNode* InNode, const FDataflowOutput* InOutput)
+	{
+		FLogMessage LogMessage
+		{
+			.Severity = Severity,
+			.Node = InNode ? InNode->AsWeak() : TWeakPtr<const FDataflowNode>{},
+			.OutputGuid = (InOutput) ? InOutput->GetGuid() : FGuid(),
+			.Text = InMessage,
+		};
+
+		auto BroadCastLogOnGameThread =
+			[WeakDelegates = Delegates.ToWeakPtr(), Message = MoveTemp(LogMessage)]()
+			{
+				check(IsInGameThread());
+				if (TSharedPtr<FDelegates> Delegates = WeakDelegates.Pin())
+				{
+					if (Delegates->OnContextLog.IsBound())
+					{
+						Delegates->OnContextLog.Broadcast(Message);
+					}
+				}
+			};
+
+		if (IsInGameThread())
+		{
+			BroadCastLogOnGameThread();
+		}
+		else
+		{
+			Async(EAsyncExecution::TaskGraphMainThread, MoveTemp(BroadCastLogOnGameThread));
+		}
+	}
 	void FContext::ClearNodesData()
 	{
-		NodesWithInfo.Empty();
-		NodesWithWarning.Empty();
-		NodesWithError.Empty();
-		NodesFailed.Empty();
+		FScopeLock Lock(&NodeTrackingLock);
+		NodesWithInfo.Reset();
+		NodesWithWarning.Reset();
+		NodesWithError.Reset();
 	}
 
 	void FContext::ClearNodeData(const FDataflowNode* InNode)
 	{
-		NodesWithInfo.Remove(InNode);
-		NodesWithWarning.Remove(InNode);
-		NodesWithError.Remove(InNode);
-		NodesFailed.Remove(InNode);
+		if (InNode)
+		{
+			FScopeLock Lock(&NodeTrackingLock);
+			const FGuid NodeGuid = InNode->GetGuid();
+			NodesWithInfo.Remove(NodeGuid);
+			NodesWithWarning.Remove(NodeGuid);
+			NodesWithError.Remove(NodeGuid);
+		}
 	}
 
 	FContextScopedCallstack::FContextScopedCallstack(FContext& InContext, const FDataflowConnection* InConnection)
@@ -486,6 +536,11 @@ namespace UE::Dataflow
 		Context.PopFromCallstack(Connection);
 	}
 
+	FContext::FContext() 
+		: AsyncEvaluator(*this) 
+		, Delegates(MakeShared<FDelegates>())
+	{}
+
 	void FContext::CheckIntrinsicInputs(const FDataflowOutput& Connection)
 	{
 #if WITH_EDITOR
@@ -667,7 +722,7 @@ namespace UE::Dataflow
 				OnNodeFinishEvaluateMulticast.Broadcast(Connection.GetOwningNode(), &Connection);
 			}
 
-			const bool bHasError = NodeHasError(Connection.GetOwningNode()) || NodeFailed(Connection.GetOwningNode());
+			const bool bHasError = NodeHasError(Connection.GetOwningNode());
 
 			return !bHasError && RetValue;
 		}
diff --git a/Engine/Source/Runtime/Experimental/Dataflow/Core/Public/Dataflow/DataflowNodeParameters.h b/Engine/Source/Runtime/Experimental/Dataflow/Core/Public/Dataflow/DataflowNodeParameters.h
--- a/Engine/Source/Runtime/Experimental/Dataflow/Core/Public/Dataflow/DataflowNodeParameters.h
+++ b/Engine/Source/Runtime/Experimental/Dataflow/Core/Public/Dataflow/DataflowNodeParameters.h
@@ -3,9 +3,11 @@
 #pragma once
 
 #include "CoreMinimal.h"
+#include "Logging/TokenizedMessage.h"
 #include "Misc/ScopeExit.h"
 #include "UObject/UnrealType.h"
 #include "Templates/UniquePtr.h"
+#include "Templates/SharedPointer.h"
 #include "HAL/CriticalSection.h"
 #include "Dataflow/DataflowContextCache.h"
 #include "Dataflow/DataflowContextAssetStore.h"
@@ -145,8 +147,11 @@ namespace UE::Dataflow
 		FContextEvaluator AsyncEvaluator;
 
 	public:
-		FContext() : AsyncEvaluator(*this) {}
+		DATAFLOWCORE_API FContext();
+
+		PRAGMA_DISABLE_DEPRECATION_WARNINGS
 		virtual ~FContext() = default;
+		PRAGMA_ENABLE_DEPRECATION_WARNINGS
 
 		static FName StaticType() { return FName("FContext"); }
 
@@ -313,16 +318,23 @@ namespace UE::Dataflow
 
 		DATAFLOWCORE_API bool IsCacheEntryAfterTimestamp(FContextCacheKey InKey, const FTimestamp InTimestamp) const;
 
+		DATAFLOWCORE_API void Info(const FText& InInfo, const FDataflowNode* InNode = nullptr, const FDataflowOutput* InOutput = nullptr);
 		DATAFLOWCORE_API void Info(const FString& InInfo, const FDataflowNode* InNode = nullptr, const FDataflowOutput* InOutput = nullptr);
 		DATAFLOWCORE_API int32 GetNumInfo() const;
+		
+		DATAFLOWCORE_API void Warning(const FText& InWarning, const FDataflowNode* InNode = nullptr, const FDataflowOutput* InOutput = nullptr);
 		DATAFLOWCORE_API void Warning(const FString& InWarning, const FDataflowNode* InNode = nullptr, const FDataflowOutput* InOutput = nullptr);
 		DATAFLOWCORE_API int32 GetNumWarnings() const;
+
+		DATAFLOWCORE_API void Error(const FText& InError, const FDataflowNode* InNode = nullptr, const FDataflowOutput* InOutput = nullptr);
 		DATAFLOWCORE_API void Error(const FString& InError, const FDataflowNode* InNode = nullptr, const FDataflowOutput* InOutput = nullptr);
 		DATAFLOWCORE_API int32 GetNumErrors() const;
 
-		bool NodeHasWarning(const FDataflowNode* InNode) { return NodesWithWarning.Contains(InNode); };
-		bool NodeHasError(const FDataflowNode* InNode) { return NodesWithError.Contains(InNode); };
-		bool NodeFailed(const FDataflowNode* InNode) { return NodesFailed.Contains(InNode); };
+		DATAFLOWCORE_API bool NodeHasWarning(const FDataflowNode* InNode) const;
+		DATAFLOWCORE_API bool NodeHasError(const FDataflowNode* InNode) const;
+		
+		UE_DEPRECATED(5.8, "Deprecated state of a node, only Errors and Warnings are now reported")
+		bool NodeFailed(const FDataflowNode* InNode) const { return false; }
 
 		DATAFLOWCORE_API void ClearNodesData();
 		DATAFLOWCORE_API void ClearNodeData(const FDataflowNode* InNode);
@@ -334,20 +346,42 @@ namespace UE::Dataflow
 		FOnNodeFinishEvaluateMulticast OnNodeFinishEvaluateMulticast;
 
 		DECLARE_MULTICAST_DELEGATE_ThreeParams(FOnContextHasInfoMulticast, const FDataflowNode* Node, const FDataflowOutput* Output, const FString& Info);
+		UE_DEPRECATED(5.8, "Use OnContextLogMulticast instead.")
 		FOnContextHasInfoMulticast OnContextHasInfo;
 
 		DECLARE_MULTICAST_DELEGATE_ThreeParams(FOnContextHasWarningMulticast, const FDataflowNode* Node, const FDataflowOutput* Output, const FString& Info);
+		UE_DEPRECATED(5.8, "Use OnContextLogMulticast instead.")
 		FOnContextHasWarningMulticast OnContextHasWarning;
 
 		DECLARE_MULTICAST_DELEGATE_ThreeParams(FOnContextHasErrorMulticast, const FDataflowNode* Node, const FDataflowOutput* Output, const FString& Info);
+		UE_DEPRECATED(5.8, "Use OnContextLogMulticast instead.")
 		FOnContextHasErrorMulticast OnContextHasError;
 
+		struct FLogMessage
+		{
+			EMessageSeverity::Type Severity = EMessageSeverity::Info;
+			TWeakPtr<const FDataflowNode> Node;
+			FGuid OutputGuid;
+			FText Text;
+		};
+
+		DECLARE_MULTICAST_DELEGATE_OneParam(FOnContextLogMulticast, const FLogMessage& Message);
+		FOnContextLogMulticast& GetOnContextLogMulticast() { return Delegates->OnContextLog; }
+
 	protected:
 		DATAFLOWCORE_API void BeginContextEvaluation(const FDataflowNode* Node, const FDataflowOutput* Output);
 
 		DATAFLOWCORE_API const FContextPerfData& GetPerfData() const;
 
+		DATAFLOWCORE_API virtual void BroadCastLog(EMessageSeverity::Type Severity, const FText& InMessage, const FDataflowNode* InNode, const FDataflowOutput* InOutput);
+
 	private:
+		// have the delegate in a shared ref so we can keep a weakpointer on them when scheduling on the gamethread from another thread 
+		struct FDelegates
+		{
+			FOnContextLogMulticast OnContextLog;
+		};
+		TSharedRef<FDelegates> Delegates;
 
 		void CheckIntrinsicInputs(const FDataflowOutput& Connection);
 
@@ -355,10 +389,10 @@ namespace UE::Dataflow
 		FContextCallstack Callstack;
 		FContextPerfData PerfData;
 #endif
-		TSet<const FDataflowNode*> NodesWithInfo;
-		TSet<const FDataflowNode*> NodesWithWarning;
-		TSet<const FDataflowNode*> NodesWithError;
-		TSet<const FDataflowNode*> NodesFailed;
+		mutable FCriticalSection NodeTrackingLock;
+		TSet<FGuid> NodesWithInfo;
+		TSet<FGuid> NodesWithWarning;
+		TSet<FGuid> NodesWithError;
 
 		/** used to store the dependent asset created during the evaluation of the graph */
 		FContextAssetStore AssetStore;
diff --git a/Engine/Source/Runtime/Experimental/Dataflow/Engine/Private/Dataflow/DataflowSubGraphNodes.cpp b/Engine/Source/Runtime/Experimental/Dataflow/Engine/Private/Dataflow/DataflowSubGraphNodes.cpp
--- a/Engine/Source/Runtime/Experimental/Dataflow/Engine/Private/Dataflow/DataflowSubGraphNodes.cpp
+++ b/Engine/Source/Runtime/Experimental/Dataflow/Engine/Private/Dataflow/DataflowSubGraphNodes.cpp
@@ -13,6 +13,8 @@
 #include "GraphEditAction.h"
 #include "UObject/UObjectGlobals.h"
 
+#define LOCTEXT_NAMESPACE "DataflowSubGraphNodes"
+
 #include UE_INLINE_GENERATED_CPP_BY_NAME(DataflowSubGraphNodes)
 
 namespace UE::Dataflow::SubGraphNodes::Private
@@ -59,6 +61,25 @@ namespace UE::Dataflow::SubGraphNodes::Private
 			return DataStore.IsEmpty();
 		}
 
+		virtual void BroadCastLog(EMessageSeverity::Type Severity, const FText& InMessage, const FDataflowNode* InNode, const FDataflowOutput* InOutput) override
+		{
+			// make sure to forward to the parent context
+			if (Severity == EMessageSeverity::Error)
+			{
+				ParentContext.Error(InMessage, InNode, InOutput);
+			}
+			else if (Severity == EMessageSeverity::Warning)
+			{
+				ParentContext.Warning(InMessage, InNode, InOutput);
+			}
+			else
+			{
+				ParentContext.Info(InMessage, InNode, InOutput);
+			}
+
+			FEngineContext::BroadCastLog(Severity, InMessage, InNode, InOutput);
+		}
+
 		virtual void Evaluate(const FDataflowNode* Node, const FDataflowOutput* Output) override
 		{
 			BeginContextEvaluation(Node, Output);
@@ -78,25 +99,6 @@ namespace UE::Dataflow::SubGraphNodes::Private
 				return false;
 			}
 
-			auto OnInfo = [this](const FDataflowNode* Node, const FDataflowOutput* Output, const FString& Info)
-				{
-					ParentContext.Info(Info, Node, Output);
-				};
-
-			bool bHasWarning = false;
-			auto OnWarning = [this, &bHasWarning](const FDataflowNode* Node, const FDataflowOutput* Output, const FString& Warning)
-				{
-					ParentContext.Warning(Warning, Node, Output);
-					bHasWarning = true;
-				};
-
-			bool bHasError = false;
-			auto OnError = [this, &bHasError](const FDataflowNode* Node, const FDataflowOutput* Output, const FString& Error)
-				{
-					ParentContext.Error(Error, Node, Output);
-					bHasError = true;
-				};
-
 			bool bEvaluationSuccess = false;
 			if (const FDataflowInput* InputToPull = SubGraphOutputNode.FindInput(CallerOutput.GetName()))
 			{
@@ -105,10 +107,6 @@ namespace UE::Dataflow::SubGraphNodes::Private
 					// we need to make sure the graph output and caller output are of compatible type
 					if (UE::Dataflow::AreTypesCompatible(OutputToEvaluate->GetType(), CallerOutput.GetType()))
 					{
-						const FDelegateHandle OnInfoHandle = OnContextHasInfo.AddLambda(OnInfo);
-						const FDelegateHandle OnWarningHandle = OnContextHasWarning.AddLambda(OnWarning);
-						const FDelegateHandle OnErrorHandle = OnContextHasError.AddLambda(OnError);
-
 						if (Evaluate(*OutputToEvaluate))
 						{
 							// copy the cache value from this cachestore to the parent one 
@@ -117,20 +115,16 @@ namespace UE::Dataflow::SubGraphNodes::Private
 							);
 						}
 
-						if (bHasError)
+						if (GetNumErrors() > 0)
 						{
-							const FString ErrorMessage = "Subgraph call evaluation returned errors : see the details above";
+							const FText ErrorMessage(LOCTEXT("SubgraphContext_EvaluationError", "Evaluation of SubGraph returned errors : see the details above"));
 							ParentContext.Error(ErrorMessage, &CallerNode, &CallerOutput);
 						}
-						if (bHasWarning)
+						if (GetNumWarnings() > 0)
 						{
-							const FString ErrorMessage = "Subgraph call evaluation returned warning errors : see the details above";
-							ParentContext.Warning(ErrorMessage, &CallerNode, &CallerOutput);
+							const FText WarningMessage(LOCTEXT("SubgraphContext_EvaluationWarning", "Evaluation of SubGraph returned warnings: see the details above"));
+							ParentContext.Warning(WarningMessage, &CallerNode, &CallerOutput);
 						}
-
-						OnContextHasInfo.Remove(OnInfoHandle);
-						OnContextHasWarning.Remove(OnWarningHandle);
-						OnContextHasError.Remove(OnErrorHandle);
 					}
 				}
 			}
@@ -684,3 +678,4 @@ void FDataflowCallSubGraphNode::Evaluate(UE::Dataflow::FContext& Context, const
 	}
 }
 
+#undef LOCTEXT_NAMESPACE
\ No newline at end of file
