Dataflow : Make ForwardValue more robust to errors

#Synced-CL 48601912

diff --git a/Engine/Source/Runtime/Experimental/Dataflow/Core/Private/Dataflow/DataflowInputOutput.cpp b/Engine/Source/Runtime/Experimental/Dataflow/Core/Private/Dataflow/DataflowInputOutput.cpp
--- a/Engine/Source/Runtime/Experimental/Dataflow/Core/Private/Dataflow/DataflowInputOutput.cpp
+++ b/Engine/Source/Runtime/Experimental/Dataflow/Core/Private/Dataflow/DataflowInputOutput.cpp
@@ -382,16 +382,17 @@ void FDataflowOutput::ForwardInput(const FDataflowInput* Input, UE::Dataflow::FC
 			if (const FDataflowOutput* ConnectionOut = Input->GetConnection())
 			{
 				Input->PullValue(Context);
-				Context.SetDataReference(CacheKey(), Property, ConnectionOut->CacheKey(), GetOwningNodeTimestamp());
+				if (Context.HasData(ConnectionOut->CacheKey()))
+				{
+					Context.SetDataReference(CacheKey(), Property, ConnectionOut->CacheKey(), GetOwningNodeTimestamp());
+					return;
+				}
 			}
 		}
-		else
-		{
-			// if there's no connection we make a invalid reference 
-			// so when the input is going to pull the cached value , it will return a default value instead
-			SetNullValue(Context);
-		}
 	}
+	// Fallback , set a null value so that the downstream nodes get a default value
+	SetNullValue(Context);
+
 }
 
 void FDataflowOutput::SetNullValue(UE::Dataflow::FContext& Context) const
