Fix compact binary serialization for FOnDemandTocContainerEntry.

#rnx

- We now write out the member 'ContainerFlags' for FOnDemandTocContainerEntry.
- We now read the members 'ContainerId' and 'ContainerFlags' for FOnDemandTocContainerEntry.
- Switched some serialization lines to match the order of the data. This is backwards compatible.
- Wrapped the field names with UTF8TEXTVIEW to avoid strlen look ups.
#Synced-CL 48634335

diff --git a/Engine/Source/Runtime/Experimental/IoStore/OnDemandCore/Private/OnDemandToc.cpp b/Engine/Source/Runtime/Experimental/IoStore/OnDemandCore/Private/OnDemandToc.cpp
--- a/Engine/Source/Runtime/Experimental/IoStore/OnDemandCore/Private/OnDemandToc.cpp
+++ b/Engine/Source/Runtime/Experimental/IoStore/OnDemandCore/Private/OnDemandToc.cpp
@@ -280,13 +280,14 @@ FCbWriter& operator<<(FCbWriter& Writer, const FOnDemandTocContainerEntry& Conta
 	}
 	Writer.EndArray();
 
-	Writer.AddHash(UTF8TEXTVIEW("UTocHash"), ContainerEntry.UTocHash);
-
 	if (!ContainerEntry.Header.IsEmpty())
 	{
 		Writer.AddBinary(UTF8TEXTVIEW("Header"), ContainerEntry.Header.GetData(), ContainerEntry.Header.Num());
 	}
 
+	Writer.AddHash(UTF8TEXTVIEW("UTocHash"), ContainerEntry.UTocHash);
+	Writer.AddInteger(UTF8TEXTVIEW("ContainerFlags"), ContainerEntry.ContainerFlags);
+
 	Writer.EndObject();
 
 	return Writer;
@@ -296,10 +297,15 @@ bool LoadFromCompactBinary(FCbFieldView Field, FOnDemandTocContainerEntry& OutCo
 {
 	if (FCbObjectView Obj = Field.AsObjectView())
 	{
-		OutContainer.ContainerName = FString(Obj["Name"].AsString());
-		OutContainer.EncryptionKeyGuid = FString(Obj["EncryptionKeyGuid"].AsString());
+		if (FCbFieldView Id = Obj[UTF8TEXTVIEW("Id")]; Id)
+		{
+			LoadFromCompactBinary(Id, OutContainer.ContainerId);
+		}
 
-		FCbArrayView Entries = Obj["Entries"].AsArrayView();
+		OutContainer.ContainerName = FString(Obj[UTF8TEXTVIEW("Name")].AsString());
+		OutContainer.EncryptionKeyGuid = FString(Obj[UTF8TEXTVIEW("EncryptionKeyGuid")].AsString());
+
+		FCbArrayView Entries = Obj[UTF8TEXTVIEW("Entries")].AsArrayView();
 		OutContainer.Entries.Reserve(int32(Entries.Num()));
 		for (FCbFieldView ArrayField : Entries)
 		{
@@ -309,14 +315,14 @@ bool LoadFromCompactBinary(FCbFieldView Field, FOnDemandTocContainerEntry& OutCo
 			}
 		}
 
-		FCbArrayView BlockSizes = Obj["BlockSizes"].AsArrayView();
+		FCbArrayView BlockSizes = Obj[UTF8TEXTVIEW("BlockSizes")].AsArrayView();
 		OutContainer.BlockSizes.Reserve(int32(BlockSizes.Num()));
 		for (FCbFieldView ArrayField : BlockSizes)
 		{
 			OutContainer.BlockSizes.Add(ArrayField.AsUInt32());
 		}
 
-		FCbArrayView BlockHashes = Obj["BlockHashes"].AsArrayView();
+		FCbArrayView BlockHashes = Obj[UTF8TEXTVIEW("BlockHashes")].AsArrayView();
 		OutContainer.BlockHashes.Reserve(int32(BlockHashes.Num()));
 		for (FCbFieldView ArrayField : BlockHashes)
 		{
@@ -331,9 +337,7 @@ bool LoadFromCompactBinary(FCbFieldView Field, FOnDemandTocContainerEntry& OutCo
 			}
 		}
 
-		OutContainer.UTocHash = Obj["UTocHash"].AsHash();
-
-		FMemoryView Header = Obj["Header"].AsBinaryView();
+		FMemoryView Header = Obj[UTF8TEXTVIEW("Header")].AsBinaryView();
 		if (!Header.IsEmpty())
 		{
 			OutContainer.Header = MakeArrayView<const uint8>(
@@ -341,6 +345,13 @@ bool LoadFromCompactBinary(FCbFieldView Field, FOnDemandTocContainerEntry& OutCo
 				IntCastChecked<int32>(Header.GetSize()));
 		}
 
+		OutContainer.UTocHash = Obj[UTF8TEXTVIEW("UTocHash")].AsHash();
+
+		if (FCbFieldView Flags = Obj[UTF8TEXTVIEW("ContainerFlags")]; Flags)
+		{
+			OutContainer.ContainerFlags = Flags.AsInt8();
+		}
+
 		return true;
 	}
 
