Constraints: remove the deleted constraint's tick function from the prerequisites of all other constraints

#rnx

#Synced-CL 48713665

diff --git a/Engine/Source/Runtime/Experimental/Animation/Constraints/Private/ConstraintSubsystem.cpp b/Engine/Source/Runtime/Experimental/Animation/Constraints/Private/ConstraintSubsystem.cpp
--- a/Engine/Source/Runtime/Experimental/Animation/Constraints/Private/ConstraintSubsystem.cpp
+++ b/Engine/Source/Runtime/Experimental/Animation/Constraints/Private/ConstraintSubsystem.cpp
@@ -172,14 +172,49 @@ void UConstraintSubsystem::RemoveConstraint(UWorld* InWorld, UTickableConstraint
 	Modify();
 	OnConstraintRemovedFromSystem_BP.Broadcast(this, InConstraint, bDoNoCompensate);
 
+	// store tick function so that it can be removed from prerequisites
+	const FConstraintTickFunction* RemovedTickFunction = InConstraint->FindTickFunction(InWorld);
+	
 	// disable constraint
 	InConstraint->Modify();
 	InConstraint->TeardownConstraint(InWorld);
 
 	if (FConstraintsInWorld* Constraints = ConstraintsInWorldFind(InWorld))
 	{
-		Constraints->Constraints.Remove(InConstraint);
+		Constraints->Constraints.Remove(InConstraint);		
 		Constraints->InvalidateGraph();
+
+		// remove useless tick prerequisites on remaining constraints
+		const bool bTickFunctionRemoved = RemovedTickFunction && !InConstraint->FindTickFunction(InWorld);
+		if (bTickFunctionRemoved)
+		{
+			static constexpr bool bEvenIfPendingKill = true;
+			for (const TWeakObjectPtr<UTickableConstraint>& OtherConstraint: Constraints->Constraints)
+			{
+				if (OtherConstraint.IsValid() && !OtherConstraint.IsStale())
+				{
+					if (FConstraintTickFunction* OtherTickFunction = OtherConstraint->FindTickFunction(InWorld))
+					{
+						TArray<FTickPrerequisite>& Prerequisites = OtherTickFunction->GetPrerequisites();
+						for (int32 PrereqIndex = 0; PrereqIndex < Prerequisites.Num(); PrereqIndex++)
+						{
+							FTickPrerequisite& Prerequisite = Prerequisites[PrereqIndex];
+							UObject* PrereqObject = Prerequisite.PrerequisiteObject.Get(bEvenIfPendingKill);
+							if (!PrereqObject)
+							{
+								// remove prerequisite coming from a stale object (cf. FTickFunction::QueueTickFunction)
+								Prerequisites.RemoveAtSwap(PrereqIndex--);
+							}
+							else if (PrereqObject == this && Prerequisite.Get() == RemovedTickFunction)
+							{
+								// remove prerequisite from the constraint just removed (cf. UConstraintSubsystem::SetConstraintDependencies) 
+								Prerequisites.RemoveAtSwap(PrereqIndex--);
+							}
+						}	
+					}
+				}
+			}
+		}
 	}
 }
 
diff --git a/Engine/Source/Runtime/Experimental/Animation/Constraints/Private/ConstraintsManager.cpp b/Engine/Source/Runtime/Experimental/Animation/Constraints/Private/ConstraintsManager.cpp
--- a/Engine/Source/Runtime/Experimental/Animation/Constraints/Private/ConstraintsManager.cpp
+++ b/Engine/Source/Runtime/Experimental/Animation/Constraints/Private/ConstraintsManager.cpp
@@ -121,6 +121,16 @@ const FConstraintTickFunction& UTickableConstraint::GetTickFunction(const UWorld
 	return *NewFunction.Get();
 }
 
+FConstraintTickFunction* UTickableConstraint::FindTickFunction(const UWorld* InWorld) const
+{
+	if (!ensure(InWorld && InWorld->GetCurrentLevel()))
+	{
+		return nullptr;
+	}
+	const TUniquePtr<FConstraintTickFunction>* Found = ConstraintTicks.Find(InWorld->GetCurrentLevel());
+	return Found && Found->IsValid() ? Found->Get() : nullptr;
+}
+
 void UTickableConstraint::SetActive(const bool bIsActive)
 {
 	Active = bIsActive;
diff --git a/Engine/Source/Runtime/Experimental/Animation/Constraints/Public/ConstraintsManager.h b/Engine/Source/Runtime/Experimental/Animation/Constraints/Public/ConstraintsManager.h
--- a/Engine/Source/Runtime/Experimental/Animation/Constraints/Public/ConstraintsManager.h
+++ b/Engine/Source/Runtime/Experimental/Animation/Constraints/Public/ConstraintsManager.h
@@ -113,6 +113,9 @@ public:
 	CONSTRAINTS_API FConstraintTickFunction& GetTickFunction(const UWorld* InWorld);
 	CONSTRAINTS_API const FConstraintTickFunction& GetTickFunction(const UWorld* InWorld) const;
 
+	/** Returns the tick function registrered in the world's current level if any (assuming both world and current leve are valid) */
+	FConstraintTickFunction* FindTickFunction(const UWorld* InWorld) const;
+
 	// UObject interface
 	CONSTRAINTS_API virtual void PostLoad() override;
 	CONSTRAINTS_API virtual void PostDuplicate(bool bDuplicateForPIE) override;
