Fix compact binary serialization for FOnDemandToc.

#rnx

- Write out FOnDemandToc::Meta when serializing to compact binary.
- Wrapped the field names with UTF8TEXTVIEW to avoid strlen look ups.

#Synced-CL 48634335

diff --git a/Engine/Source/Runtime/Experimental/IoStore/OnDemandCore/Private/OnDemandToc.cpp b/Engine/Source/Runtime/Experimental/IoStore/OnDemandCore/Private/OnDemandToc.cpp
--- a/Engine/Source/Runtime/Experimental/IoStore/OnDemandCore/Private/OnDemandToc.cpp
+++ b/Engine/Source/Runtime/Experimental/IoStore/OnDemandCore/Private/OnDemandToc.cpp
@@ -551,6 +551,8 @@ FCbWriter& operator<<(FCbWriter& Writer, const FOnDemandToc& Toc)
 	Writer.BeginObject();
 	Writer << UTF8TEXTVIEW("Header") << Toc.Header;
 
+	Writer << UTF8TEXTVIEW("Meta") << Toc.Meta;
+
 	Writer.BeginArray(UTF8TEXTVIEW("Containers"));
 	for (const FOnDemandTocContainerEntry& Container : Toc.Containers)
 	{
@@ -590,20 +592,20 @@ bool LoadFromCompactBinary(FCbFieldView Field, FOnDemandToc& OutToc)
 {
 	if (FCbObjectView Obj = Field.AsObjectView())
 	{
-		if (!LoadFromCompactBinary(Obj["Header"], OutToc.Header))
+		if (!LoadFromCompactBinary(Obj[UTF8TEXTVIEW("Header")], OutToc.Header))
 		{
 			return false;
 		}
 
 		if (uint32(OutToc.Header.Version) >= uint32(EOnDemandTocVersion::Meta))
 		{
-			if (!LoadFromCompactBinary(Obj["Meta"], OutToc.Meta))
+			if (!LoadFromCompactBinary(Obj[UTF8TEXTVIEW("Meta")], OutToc.Meta))
 			{
 				return false;
 			}
 		}
 
-		FCbArrayView Containers = Obj["Containers"].AsArrayView();
+		FCbArrayView Containers = Obj[UTF8TEXTVIEW("Containers")].AsArrayView();
 		OutToc.Containers.Reserve(int32(Containers.Num()));
 		for (FCbFieldView ArrayField : Containers)
 		{
@@ -628,7 +630,7 @@ bool LoadFromCompactBinary(FCbFieldView Field, FOnDemandToc& OutToc)
 
 		if (uint32(OutToc.Header.Version) >= uint32(EOnDemandTocVersion::TagSets))
 		{
-			FCbArrayView TagSets = Obj["TagSets"].AsArrayView();
+			FCbArrayView TagSets = Obj[UTF8TEXTVIEW("TagSets")].AsArrayView();
 			OutToc.TagSets.Reserve(int32(TagSets.Num()));
 			for (FCbFieldView ArrayField : TagSets)
 			{
