Include FMetaData, if available, in json when stringify'ing a UPackage

#rnx
#Synced-CL 48489485

diff --git a/Engine/Source/Runtime/Experimental/JsonObjectGraph/Private/JsonStringifyImpl.cpp b/Engine/Source/Runtime/Experimental/JsonObjectGraph/Private/JsonStringifyImpl.cpp
--- a/Engine/Source/Runtime/Experimental/JsonObjectGraph/Private/JsonStringifyImpl.cpp
+++ b/Engine/Source/Runtime/Experimental/JsonObjectGraph/Private/JsonStringifyImpl.cpp
@@ -277,9 +277,16 @@ FJsonStringifyImpl::FJsonStringifyImpl(TConstArrayView<const UObject*> Roots, co
 	, MemoryWriter(Result)
 	, Writer(FJsonWriter::Create(&MemoryWriter))
 	, CurrentObject(nullptr)
+	, RootPackage(nullptr)
 {
 	const bool bFilterEditorOnly = (Options.Flags & EJsonStringifyFlags::FilterEditorOnlyData) != EJsonStringifyFlags::Default;
 	TArray<const UObject*> Exports;
+	// If the requested graph is simple we can store the root package and 
+	// include said package's metadata in the result:
+	if(Roots.Num() == 1 && Roots[0]->IsA<UPackage>())
+	{
+		RootPackage = Cast<UPackage>(Roots[0]);
+	}
 	GatherExports(Roots, {}, bFilterEditorOnly, RootObjects, Exports);
 
 	ObjectsToExport = TSet<const UObject*>(Exports);
@@ -354,10 +361,7 @@ void FJsonStringifyImpl::ToJsonBytes()
 	}
 	Writer->WriteArrayEnd();
 	
-	if(ShouldWritePackageSummary())
-	{
-		WritePackageSummary();
-	}
+	WritePackageSummary();
 
 	Writer->WriteObjectEnd();
 	// Result is now populated with Json representation
@@ -947,37 +951,126 @@ TArray<uint8> FJsonStringifyImpl::StructuredDataToJson(const UObject* Object, in
 
 void FJsonStringifyImpl::WritePackageSummary()
 {
-	Writer->WriteIdentifierPrefix(UE_JSON_PACKAGE_SUMMARY_KEY_TCHAR);
-	Writer->WriteObjectStartInline();
+	FPendingScope PackageSummary(this, 
+		[this]()
+		{
+			Writer->WriteIdentifierPrefix(UE_JSON_PACKAGE_SUMMARY_KEY_TCHAR);
+			Writer->WriteObjectStartInline();
+		},
+		[this]()
+		{
+			Writer->WriteObjectEnd();
+		}
+	);
 
+	WriteCustomVersions();
+	WritePackageMetadata();
+}
+
+void FJsonStringifyImpl::WriteCustomVersions()
+{
 #if WITH_TEXT_ARCHIVE_SUPPORT
-	if (Versions.Num() != 0)
+	if (Versions.Num() == 0)
 	{
-		// filter dupes, preserving order of first encounter - we need some stable order:
-		TSet<FGuid> Encountered;
-		Versions.SetNum(Algo::StableRemoveIf(Versions,
-			[&Encountered](const FCustomVersion& V)
-			{
-				if (Encountered.Contains(V.Key))
-				{
-					return true;
-				}
-				else
-				{
-					Encountered.Add(V.Key);
-					return false;
-				}
-			}));
+		return;
+	}
 
-		Writer->WriteIdentifierPrefix(UE_JSON_CUSTOM_VERSIONS_KEY_TCHAR);
-		FJsonStringifyStructuredArchive::WriteCustomVersionValueInline(Versions, Writer->GetIndentLevel(), MemoryWriter);
-		Writer->HACK_SetPreviousTokenWrittenSquareClose(); // the above will write the custom version as an array of tuples, track that write
+	if (CurrentScope)
+	{
+		CurrentScope->Apply();
 	}
+
+	// filter dupes, preserving order of first encounter - we need some stable order:
+	TSet<FGuid> Encountered;
+	Versions.SetNum(Algo::StableRemoveIf(Versions,
+		[&Encountered](const FCustomVersion& V)
+		{
+			if (Encountered.Contains(V.Key))
+			{
+				return true;
+			}
+			else
+			{
+				Encountered.Add(V.Key);
+				return false;
+			}
+		}));
+
+	Writer->WriteIdentifierPrefix(UE_JSON_CUSTOM_VERSIONS_KEY_TCHAR);
+	FJsonStringifyStructuredArchive::WriteCustomVersionValueInline(Versions, Writer->GetIndentLevel(), MemoryWriter);
+	Writer->HACK_SetPreviousTokenWrittenSquareClose(); // the above will write the custom version as an array of tuples, track that write
 #endif // WITH_TEXT_ARCHIVE_SUPPORT
+}
+
+void FJsonStringifyImpl::WritePackageMetadata()
+{
+#if WITH_METADATA
+	if(!RootPackage)
+	{
+		return;
+	}
+	
+	const FPackageFileVersion& Version = RootPackage->GetLinkerPackageVersion();
+	if(!Version.IsCompatible( 
+		FPackageFileVersion(VER_LATEST_ENGINE_UE4, EUnrealEngineObjectUE5Version::METADATA_SERIALIZATION_OFFSET)) )
+	{
+		return;
+	}
+	const FMetaData& MetaData = const_cast<UPackage*>(RootPackage)->GetMetaData();
+	const bool bHasObjectMetaData = MetaData.ObjectMetaDataMap.Num() > 0;
+	const bool bHasRootMetaData = MetaData.RootMetaDataMap.Num() > 0;
+	const bool bHasAnyMetadata = bHasObjectMetaData || bHasRootMetaData;
+	if(!bHasAnyMetadata)
+	{
+		return;
+	}
+	
+	const auto WriteNameStringMap = [this](const TMap<FName, FString>& NameStringMap)
+	{
+		Writer->WriteArrayStartInline();
+		for(const TPair<FName, FString>& Values : NameStringMap)
+		{
+			Writer->WriteObjectStart();
+			Writer->WriteValue(TEXT("Key"), Values.Key.ToString());
+			Writer->WriteValue(TEXT("Value"), Values.Value);
+			Writer->WriteObjectEnd();
+		}
+		Writer->WriteArrayEnd();
+	};
+
+	if (CurrentScope)
+	{
+		CurrentScope->Apply();
+	}
+
+	Writer->WriteIdentifierPrefix(TEXT("PackageMetadata"));
+	Writer->WriteObjectStartInline();
+
+	if(bHasObjectMetaData)
+	{
+		Writer->WriteIdentifierPrefix(TEXT("ObjectMetadata"));
+		Writer->WriteArrayStartInline();
+		for(const TPair<FSoftObjectPath, TMap<FName, FString>>& ObjectMetaData : MetaData.ObjectMetaDataMap)
+		{
+			Writer->WriteObjectStart();
+
+			Writer->WriteValue(TEXT("Key"), ObjectMetaData.Key.ToString());
+			
+			Writer->WriteIdentifierPrefix(TEXT("Values"));
+			WriteNameStringMap(ObjectMetaData.Value);
+
+			Writer->WriteObjectEnd();
+		}
+		Writer->WriteArrayEnd();
+	}
+	if(bHasRootMetaData)
+	{
+		Writer->WriteIdentifierPrefix(TEXT("RootMetadata"));
+		WriteNameStringMap(MetaData.RootMetaDataMap);
+	}
 
-	// Note that there is no automatic compression here, if you add information to the package summary
-	// you should update ShouldWritePackageSummary
 	Writer->WriteObjectEnd();
+#endif // WITH_METADATA
 }
 
 bool FJsonStringifyImpl::IsDeltaEncoding() const
@@ -985,16 +1078,6 @@ bool FJsonStringifyImpl::IsDeltaEncoding() const
 	return !EnumHasAnyFlags(WriteOptions.Flags, EJsonStringifyFlags::DisableDeltaEncoding);
 }
 
-bool FJsonStringifyImpl::ShouldWritePackageSummary() const
-{
-	return 
-#if WITH_TEXT_ARCHIVE_SUPPORT
-		Versions.Num() > 0;
-#else
-		false;
-#endif
-}
-
 FJsonStringifyImpl::FPendingScope::FPendingScope(FJsonStringifyImpl* To, const TFunction<void()>& Prefix)
 	: Owner(To)
 	, Outer(To->CurrentScope)
diff --git a/Engine/Source/Runtime/Experimental/JsonObjectGraph/Private/JsonStringifyImpl.h b/Engine/Source/Runtime/Experimental/JsonObjectGraph/Private/JsonStringifyImpl.h
--- a/Engine/Source/Runtime/Experimental/JsonObjectGraph/Private/JsonStringifyImpl.h
+++ b/Engine/Source/Runtime/Experimental/JsonObjectGraph/Private/JsonStringifyImpl.h
@@ -73,9 +73,10 @@ private:
 #endif// WITH_TEXT_ARCHIVE_SUPPORT
 
 	void WritePackageSummary();
+	void WriteCustomVersions();
+	void WritePackageMetadata();
 
 	bool IsDeltaEncoding() const;
-	bool ShouldWritePackageSummary() const;
 
 	const FJsonStringifyOptions& WriteOptions;
 
@@ -111,6 +112,7 @@ private:
 	TArray<const UObject*> RootObjects;
 	TSet<const UObject*> ObjectsToExport;
 	TSet<const UObject*> ObjectsExported;
+	const UPackage* RootPackage;
 
 	TArray<FCustomVersion> Versions;
 };
