Build fix: missing forward declaration, removed obsolete module

#Synced-CL 48743446

diff --git a/Engine/Plugins/Experimental/ChaosRigidAsset/Source/ChaosRigidAssetNodes/Private/PhysicsAssetDataflowNodes.cpp b/Engine/Plugins/Experimental/ChaosRigidAsset/Source/ChaosRigidAssetNodes/Private/PhysicsAssetDataflowNodes.cpp
--- a/Engine/Plugins/Experimental/ChaosRigidAsset/Source/ChaosRigidAssetNodes/Private/PhysicsAssetDataflowNodes.cpp
+++ b/Engine/Plugins/Experimental/ChaosRigidAsset/Source/ChaosRigidAssetNodes/Private/PhysicsAssetDataflowNodes.cpp
@@ -21,6 +21,7 @@
 #include "Animation/AnimSequence.h"
 #include "AnimationRuntime.h"
 #include "AnimationUtils.h"
+#include "ChaosVDRuntimeModule.h"
 
 FDataflowPhysicsAssetTerminalNode::FDataflowPhysicsAssetTerminalNode(const UE::Dataflow::FNodeParameters& InParam, FGuid InGuid /*= FGuid::NewGuid()*/)
 	: FDataflowTerminalNode(InParam, InGuid)
diff --git a/Engine/Source/Runtime/Experimental/TraceBasedDebuggers/Public/RemoteSessionsManager.h b/Engine/Source/Runtime/Experimental/TraceBasedDebuggers/Public/RemoteSessionsManager.h
--- a/Engine/Source/Runtime/Experimental/TraceBasedDebuggers/Public/RemoteSessionsManager.h
+++ b/Engine/Source/Runtime/Experimental/TraceBasedDebuggers/Public/RemoteSessionsManager.h
@@ -22,7 +22,7 @@ UE_API DECLARE_LOG_CATEGORY_EXTERN(LogTraceBasedDebuggers, Log, All);
 
 namespace UE::TraceBasedDebuggers
 {
-
+struct FRemoteSessionsManager;
 struct FTraceConnectionDetails;
 
 USTRUCT(MinimalAPI)
@@ -312,7 +312,7 @@ struct IRemoteSessionsHandler : TSharedFromThis<IRemoteSessionsHandler>
 	virtual void OnCreatingMessageEndpoint(const TSharedPtr<FRemoteSessionsManager>&, FMessageEndpointBuilder&) = 0;
 
 	/** Callback to allow registered sessions handlers to subscribe to specific message types */
-	virtual void OnSubscribingMessageTypes(const TSharedPtr<UE::TraceBasedDebuggers::FRemoteSessionsManager>&, FMessageEndpoint&) = 0;
+	virtual void OnSubscribingMessageTypes(const TSharedPtr<FRemoteSessionsManager>&, FMessageEndpoint&) = 0;
 
 	/** Callback to allow registered sessions handlers to build the response for a session information request */
 	virtual void OnBuildingFullSessionInfoResponseMessage(FFullSessionInfoResponseMessage&) = 0;
