Build fix:
- missing header of ChaosVDContextProvider.h
- missing #if WITH_TRACE_BASED_DEBUGGERS
- moved SessionManager creation from EngineEditorBridge ctor to Initialize

#Synced-CL 48743446

diff --git a/Engine/Source/Runtime/Experimental/Chaos/Public/Chaos/Collision/SpatialAccelerationBroadPhase.h b/Engine/Source/Runtime/Experimental/Chaos/Public/Chaos/Collision/SpatialAccelerationBroadPhase.h
--- a/Engine/Source/Runtime/Experimental/Chaos/Public/Chaos/Collision/SpatialAccelerationBroadPhase.h
+++ b/Engine/Source/Runtime/Experimental/Chaos/Public/Chaos/Collision/SpatialAccelerationBroadPhase.h
@@ -14,6 +14,7 @@
 #include "Chaos/EvolutionResimCache.h"
 #include "Chaos/GeometryParticlesfwd.h"
 #include "Chaos/AABBTree.h"
+#include "ChaosVisualDebugger/ChaosVDContextProvider.h"
 #include "Tasks/Task.h"
 
 // Set to 1 to enable some sanity chacks on the filtered broadphase results
diff --git a/Engine/Source/Runtime/Experimental/ChaosSolverEngine/Private/Chaos/ChaosVDRemoteSessionsManager.cpp b/Engine/Source/Runtime/Experimental/ChaosSolverEngine/Private/Chaos/ChaosVDRemoteSessionsManager.cpp
--- a/Engine/Source/Runtime/Experimental/ChaosSolverEngine/Private/Chaos/ChaosVDRemoteSessionsManager.cpp
+++ b/Engine/Source/Runtime/Experimental/ChaosSolverEngine/Private/Chaos/ChaosVDRemoteSessionsManager.cpp
@@ -11,6 +11,7 @@ void FChaosVDRemoteSessionsHandler::OnCreatingMessageEndpoint(
 	const TSharedPtr<UE::TraceBasedDebuggers::FRemoteSessionsManager>& SessionsManager
 	, FMessageEndpointBuilder& InEndPointBuilder)
 {
+#if WITH_TRACE_BASED_DEBUGGERS
 	InEndPointBuilder
 		.Handling<FChaosVDChannelStateChangeCommandMessage>(
 			[](const FChaosVDChannelStateChangeCommandMessage& InMessage, const TSharedRef<IMessageContext>&)
@@ -65,6 +66,7 @@ void FChaosVDRemoteSessionsHandler::OnCreatingMessageEndpoint(
 					}
 				});
 	}
+#endif
 }
 
 void FChaosVDRemoteSessionsHandler::OnSubscribingMessageTypes(
@@ -92,6 +94,7 @@ void FChaosVDRemoteSessionsHandler::OnSubscribingMessageTypes(
 void FChaosVDRemoteSessionsHandler::OnBuildingFullSessionInfoResponseMessage(
 	UE::TraceBasedDebuggers::FFullSessionInfoResponseMessage& InMessage)
 {
+#if WITH_TRACE_BASED_DEBUGGERS
 	using namespace Chaos::VisualDebugger;
 
 	InMessage.bIsRecording = FChaosVDRuntimeModule::Get().IsRecording();
@@ -114,6 +117,7 @@ void FChaosVDRemoteSessionsHandler::OnBuildingFullSessionInfoResponseMessage(
 		});
 
 	InMessage.DebuggerSpecificData.Emplace(FInstancedStruct::Make(MoveTemp(MessageData)));
+#endif
 }
 
 void FChaosVDRemoteSessionsHandler::OnHandlingFullSessionInfoResponseMessage(
diff --git a/Engine/Source/Runtime/Experimental/TraceBasedDebuggers/Private/EngineEditorBridge.cpp b/Engine/Source/Runtime/Experimental/TraceBasedDebuggers/Private/EngineEditorBridge.cpp
--- a/Engine/Source/Runtime/Experimental/TraceBasedDebuggers/Private/EngineEditorBridge.cpp
+++ b/Engine/Source/Runtime/Experimental/TraceBasedDebuggers/Private/EngineEditorBridge.cpp
@@ -72,17 +72,18 @@ void FEngineEditorBridge::SetDataRelayTransportInstance(const TSharedPtr<IDataRe
 }
 
 FEngineEditorBridge::FEngineEditorBridge(const FLogCategoryBase& LogCategory)
-	: RemoteSessionsManager(SharedRemoteSessionsManager.IsValid() ? SharedRemoteSessionsManager.Pin().ToSharedRef() : MakeShared<FRemoteSessionsManager>())
-	, LogCategory(LogCategory)
+	: LogCategory(LogCategory)
 {
+}
+
+void FEngineEditorBridge::Initialize()
+{
+	RemoteSessionsManager = SharedRemoteSessionsManager.IsValid() ? SharedRemoteSessionsManager.Pin().ToSharedRef() : MakeShared<FRemoteSessionsManager>();
 	if (!SharedRemoteSessionsManager.IsValid())
 	{
 		SharedRemoteSessionsManager = RemoteSessionsManager;
 	}
-}
 
-void FEngineEditorBridge::Initialize()
-{
 	// Allow derived classes to perform their initialization
 	OnInitializeInternal();
 
diff --git a/Engine/Source/Runtime/Experimental/TraceBasedDebuggers/Private/TraceBasedDebuggerRuntime.cpp b/Engine/Source/Runtime/Experimental/TraceBasedDebuggers/Private/TraceBasedDebuggerRuntime.cpp
--- a/Engine/Source/Runtime/Experimental/TraceBasedDebuggers/Private/TraceBasedDebuggerRuntime.cpp
+++ b/Engine/Source/Runtime/Experimental/TraceBasedDebuggers/Private/TraceBasedDebuggerRuntime.cpp
@@ -223,6 +223,7 @@ void FRuntimeModule::StartRecording(const FStartRecordingCommandMessage& InRecor
 				break;
 			}
 
+#if WITH_TRACE_BASED_DEBUGGERS
 			if (RelayWriter && !RelayWriter->IsClosed())
 			{
 				UE_LOG_REF(LogCategory, Warning, TEXT("[%hs] Attempting to relay trace, but there is a relay writer still open"), __func__);
@@ -238,6 +239,7 @@ void FRuntimeModule::StartRecording(const FStartRecordingCommandMessage& InRecor
 			RelayWriter->SetMaxAllowedPendingBytes(FMath::Clamp(CVarMaxPendingBytesInRelayWriter->GetInt(), MinBufferSizeLimit, std::numeric_limits<int32>::max()));
 
 			bIsRecording = FTraceAuxiliary::Relay(reinterpret_cast<UPTRINT>(RelayWriter.Get()), FRelayTraceDataWriter::WriteHelper, FRelayTraceDataWriter::CloseHelper, nullptr, &TracingOptions);
+#endif
 
 			if (bIsRecording)
 			{
@@ -283,7 +285,9 @@ void FRuntimeModule::StartRecording(const FStartRecordingCommandMessage& InRecor
 		FMessageDialog::Open(EAppMsgType::Ok, FText::FormatOrdered(LOCTEXT("StartRecordingFailedMessage", "Failed to start recording. \n\n{0}"), FailureReason));
 #endif
 
+#if WITH_TRACE_BASED_DEBUGGERS
 		RelayWriter = nullptr;
+#endif
 
 		// Allow derived classes to perform additional actions on failure
 		OnRecordingStartFailedInternal(FailureReason);
@@ -320,10 +324,12 @@ void FRuntimeModule::StopRecording()
 
 	ClearRelayDataPumpDelegates();
 
+#if WITH_TRACE_BASED_DEBUGGERS
 	if (RelayWriter)
 	{
 		RelayWriter->Close();
 	}
+#endif
 
 	bIsRecording = false;
 	AccumulatedRecordingTime = 0.0f;
@@ -337,10 +343,12 @@ int64 FRuntimeModule::GetBufferedDataBytesSize() const
 {
 	FGuardedInt64 RelaySendBufferSize(0);
 
+#if WITH_TRACE_BASED_DEBUGGERS
 	if (RelayWriter)
 	{
 		RelaySendBufferSize += RelayWriter->GetQueuedBytesNum();
 	}
+#endif
 
 	if (ExternalRelayBufferSizeProviderDelegate.IsBound())
 	{
@@ -459,14 +467,17 @@ void FRuntimeModule::HandleRelayOverflow()
 
 void FRuntimeModule::RelayTraceData() const
 {
+#if WITH_TRACE_BASED_DEBUGGERS
 	if (ensure(ExternalRelayExecutorDelegate.IsBound() && RelayWriter.IsValid()))
 	{
 		ExternalRelayExecutorDelegate.Execute(*RelayWriter);
 	}
+#endif
 }
 
 void FRuntimeModule::SetupRelayDataPumpDelegates()
 {
+#if WITH_TRACE_BASED_DEBUGGERS
 	if (!ensure(RelayWriter))
 	{
 		return;
@@ -474,15 +485,18 @@ void FRuntimeModule::SetupRelayDataPumpDelegates()
 
 	RelayWriter->OnNewDataAvailable().BindRaw(this, &FRuntimeModule::RelayTraceData);
 	RelayWriter->OnDataOverflow().BindRaw(this, &FRuntimeModule::HandleRelayOverflow);
+#endif
 }
 
 void FRuntimeModule::ClearRelayDataPumpDelegates()
 {
+#if WITH_TRACE_BASED_DEBUGGERS
 	if (RelayWriter)
 	{
 		RelayWriter->OnNewDataAvailable().Unbind();
 		RelayWriter->OnDataOverflow().Unbind();
 	}
+#endif
 }
 
 bool FRuntimeModule::RecordingTimerTick(float DeltaTime)
diff --git a/Engine/Source/Runtime/Experimental/TraceBasedDebuggers/Public/EngineEditorBridge.h b/Engine/Source/Runtime/Experimental/TraceBasedDebuggers/Public/EngineEditorBridge.h
--- a/Engine/Source/Runtime/Experimental/TraceBasedDebuggers/Public/EngineEditorBridge.h
+++ b/Engine/Source/Runtime/Experimental/TraceBasedDebuggers/Public/EngineEditorBridge.h
@@ -92,7 +92,7 @@ private:
 	TSharedPtr<IDataRelayTransport> RelayTraceDataTransportInstance;
 
 	static TWeakPtr<FRemoteSessionsManager> SharedRemoteSessionsManager;
-	TSharedRef<FRemoteSessionsManager> RemoteSessionsManager;
+	TSharedPtr<FRemoteSessionsManager> RemoteSessionsManager;
 	/** Handler created by derived classes to handle more message types */
 	TSharedPtr<IRemoteSessionsHandler> RemoteSessionsHandler;
 
diff --git a/Engine/Source/Runtime/Experimental/TraceBasedDebuggers/Public/TraceBasedDebuggerRuntime.h b/Engine/Source/Runtime/Experimental/TraceBasedDebuggers/Public/TraceBasedDebuggerRuntime.h
--- a/Engine/Source/Runtime/Experimental/TraceBasedDebuggers/Public/TraceBasedDebuggerRuntime.h
+++ b/Engine/Source/Runtime/Experimental/TraceBasedDebuggers/Public/TraceBasedDebuggerRuntime.h
@@ -141,7 +141,11 @@ public:
 
 	FRelayTraceDataWriter* GetRelayTraceWriterInstance() const
 	{
+#if WITH_TRACE_BASED_DEBUGGERS
 		return RelayWriter.Get();
+#else
+		return nullptr;
+#endif
 	}
 
 protected:
@@ -206,7 +210,9 @@ private:
 	FDelegateHandle OnConnectionDelegateHandle;
 	FDelegateHandle OnTraceStoppedDelegateHandle;
 
+#if WITH_TRACE_BASED_DEBUGGERS
 	TUniquePtr<FRelayTraceDataWriter> RelayWriter;
+#endif
 
 	FSimpleMulticastDelegate TraceConnectionDetailsUpdatedDelegate;
 	FSimpleDelegate ExternalRelayOverflowHandlerDelegate;
