Split the concept of single threaded physics and single worker physics.
Fix Narrow Phase in single thread.

-- Split the concept of single threaded physics and single worker physics.
single threaded physics : Physics runs on the Game Thread on a single thread
single worker physics : Physics runs on the Physics Thread on a single worker
The latter can be used to avoid thread preemption when few cores are available.

-- Fix Narrow Phase in single thread
Narrow Phase were still creating tasks even in single thread. Now it will stay on only one thread

#Synced-CL 48642165

diff --git a/Engine/Source/Runtime/Experimental/Chaos/Private/Chaos/Collision/CollisionConstraintAllocator.cpp b/Engine/Source/Runtime/Experimental/Chaos/Private/Chaos/Collision/CollisionConstraintAllocator.cpp
--- a/Engine/Source/Runtime/Experimental/Chaos/Private/Chaos/Collision/CollisionConstraintAllocator.cpp
+++ b/Engine/Source/Runtime/Experimental/Chaos/Private/Chaos/Collision/CollisionConstraintAllocator.cpp
@@ -93,7 +93,7 @@ namespace Chaos
 
 		void FCollisionConstraintAllocator::ProcessNewNonCcdConstraints(const TArray<FPBDCollisionConstraint*>& NewConstraints, TArray<FPBDCollisionConstraint*>& AllCCDConstraints)
 		{
-			int32 NumTasks = !!GSingleThreadedPhysics ? 1 : FMath::Min(FTaskGraphInterface::Get().GetNumWorkerThreads(), Chaos::MaxNumWorkers);
+			int32 NumTasks = bSingleWorkerPhysics ? 1 : FMath::Min(FTaskGraphInterface::Get().GetNumWorkerThreads(), Chaos::MaxNumWorkers);
 			NumTasks = FMath::Max(FMath::Min(NumTasks, NewConstraints.Num()), 1);
 
 			const int32 NumConstraints = NewConstraints.Num();
@@ -159,7 +159,7 @@ namespace Chaos
 
 		void FCollisionConstraintAllocator::ProcessNewCcdConstraints()
 		{
-			int32 NumTasks = !!GSingleThreadedPhysics ? 1 : FMath::Min(FTaskGraphInterface::Get().GetNumWorkerThreads(), Chaos::MaxNumWorkers);
+			int32 NumTasks = bSingleWorkerPhysics ? 1 : FMath::Min(FTaskGraphInterface::Get().GetNumWorkerThreads(), Chaos::MaxNumWorkers);
 			NumTasks = FMath::Max(FMath::Min(NumTasks, ActiveCCDConstraints.Num()), 1);
 
 			const int32 NumConstraints = ActiveCCDConstraints.Num();
diff --git a/Engine/Source/Runtime/Experimental/Chaos/Private/Chaos/Framework/Parallel.cpp b/Engine/Source/Runtime/Experimental/Chaos/Private/Chaos/Framework/Parallel.cpp
--- a/Engine/Source/Runtime/Experimental/Chaos/Private/Chaos/Framework/Parallel.cpp
+++ b/Engine/Source/Runtime/Experimental/Chaos/Private/Chaos/Framework/Parallel.cpp
@@ -7,7 +7,10 @@
 
 namespace Chaos
 {
+	// If GSingleThreadedPhysics is 1 physics simulation runs on the Game Thread in single thread 
 	CHAOS_API int32 GSingleThreadedPhysics = 0;
+	// If bSingleWorkerPhysics is true physics simulation run only on one thread on the Physics Thread
+	CHAOS_API bool bSingleWorkerPhysics = false;
 	CHAOS_API int32 InnerParallelForBatchSize = 0;
 	CHAOS_API int32 MinRangeBatchSize = 0;
 	CHAOS_API int32 MaxNumWorkers = 100;
@@ -28,6 +31,7 @@ namespace Chaos
 	FAutoConsoleVariableRef CVarSmallBatchSize(TEXT("p.Chaos.SmallBatchSize"), SmallBatchSize, TEXT("Small batch size for chaos parallel loops"));
 	FAutoConsoleVariableRef CVarLargeBatchSize(TEXT("p.Chaos.LargeBatchSize"), LargeBatchSize, TEXT("Large batch size for chaos parallel loops"));
 	FAutoConsoleVariableRef CVarMinParallelTaskSize(TEXT("p.Chaos.MinParallelTaskSize"), MinParallelTaskSize, TEXT("Minimum number of tasks required to start parallel task execution on worker threads, set this higher to encourage single threaded execution"));
+	FAutoConsoleVariableRef CVarSingleWorkerPhysics(TEXT("p.Chaos.SingleWorkerPhysics"), bSingleWorkerPhysics, TEXT("Run Chaos on a single thread"));
 }
 
 void Chaos::InnerPhysicsParallelFor(int32 Num, TFunctionRef<void(int32)> InCallable, bool bForceSingleThreaded)
@@ -85,12 +89,11 @@ void Chaos::PhysicsParallelFor(int32 InNum, TFunctionRef<void(int32)> InCallable
 		InCallable(Idx);
 	};
 
-	const bool bSingleThreaded = !!GSingleThreadedPhysics || bDisablePhysicsParallelFor || bForceSingleThreaded;
+	const bool bSingleThreaded = bSingleWorkerPhysics || bDisablePhysicsParallelFor || bForceSingleThreaded;
 	const EParallelForFlags Flags = (bSingleThreaded ? EParallelForFlags::ForceSingleThread : EParallelForFlags::None);
 	const int32 MinBatchSize = ((MaxNumWorkers > 0) && (InNum > MaxNumWorkers)) ? FMath::DivideAndRoundUp(InNum, MaxNumWorkers) : 1;
 
 	ParallelFor(TEXT("PhysicsParallelFor"), InNum, MinBatchSize, PassThrough, Flags);
-	//::ParallelFor(InNum, PassThrough, !!GSingleThreadedPhysics || bDisablePhysicsParallelFor || bForceSingleThreaded);
 }
 
 void Chaos::PhysicsParallelForRange(int32 InNum, TFunctionRef<void(int32, int32)> InCallable, const int32 InMinBatchSize, bool bForceSingleThreaded)
@@ -156,7 +159,7 @@ void Chaos::PhysicsParallelForRange(int32 InNum, TFunctionRef<void(int32, int32)
 #endif
 		InCallable(RangeIndex[ThreadId], RangeIndex[ThreadId + 1]);
 	};
-	::ParallelFor(RangeIndex.Num() - 1, PassThrough, !!GSingleThreadedPhysics || bDisablePhysicsParallelFor || bForceSingleThreaded);
+	::ParallelFor(RangeIndex.Num() - 1, PassThrough, bSingleWorkerPhysics || bDisablePhysicsParallelFor || bForceSingleThreaded);
 }
 
 void Chaos::PhysicsParallelForWithContext(int32 InNum, TFunctionRef<int32 (int32, int32)> InContextCreator, TFunctionRef<void(int32, int32)> InCallable, bool bForceSingleThreaded)
@@ -182,7 +185,7 @@ void Chaos::PhysicsParallelForWithContext(int32 InNum, TFunctionRef<int32 (int32
 		InCallable(Idx, ContextIndex);
 	};
 
-	const bool bSingleThreaded = !!GSingleThreadedPhysics || bDisablePhysicsParallelFor || bForceSingleThreaded;
+	const bool bSingleThreaded = bSingleWorkerPhysics || bDisablePhysicsParallelFor || bForceSingleThreaded;
 	const EParallelForFlags Flags = bSingleThreaded ? (EParallelForFlags::ForceSingleThread) : (EParallelForFlags::None);
 	const int32 MinBatchSize = ((MaxNumWorkers > 0) && (InNum > MaxNumWorkers)) ? FMath::DivideAndRoundUp(InNum, MaxNumWorkers) : 1;
 
@@ -198,7 +201,7 @@ void Chaos::PhysicsParallelForWithContext(int32 InNum, TFunctionRef<int32 (int32
 
 bool Chaos::ShouldExecuteTasksInParallel(int32 InTaskNum)
 {
-	if (GSingleThreadedPhysics)
+	if (bSingleWorkerPhysics)
 	{
 		// Single threaded physics should not run in parallel
 		return false;
diff --git a/Engine/Source/Runtime/Experimental/Chaos/Private/Chaos/Island/IslandGroupManager.cpp b/Engine/Source/Runtime/Experimental/Chaos/Private/Chaos/Island/IslandGroupManager.cpp
--- a/Engine/Source/Runtime/Experimental/Chaos/Private/Chaos/Island/IslandGroupManager.cpp
+++ b/Engine/Source/Runtime/Experimental/Chaos/Private/Chaos/Island/IslandGroupManager.cpp
@@ -23,7 +23,6 @@ DECLARE_CYCLE_STAT(TEXT("FPBDRigidsEvolutionGBF::Scatter::Constraints"), STAT_Ev
 
 namespace Chaos
 {
-	extern int32 GSingleThreadedPhysics;
 
 	namespace CVars
 	{
@@ -66,7 +65,7 @@ namespace Chaos
 		{
 			// Check for use of the "-onethread" command line arg, and physics threading disabled (GetNumWorkerThreads() is not affected by these)
 			// @todo(chaos): is the number of worker threads a good indicator of how many threads we get in the solver loop? (Currently uses ParallelFor)
-			NumWorkerThreads = (FApp::ShouldUseThreadingForPerformance() && !GSingleThreadedPhysics) ? FMath::Min(FTaskGraphInterface::Get().GetNumWorkerThreads(), Chaos::MaxNumWorkers) : 0;
+			NumWorkerThreads = (FApp::ShouldUseThreadingForPerformance() && !bSingleWorkerPhysics) ? FMath::Min(FTaskGraphInterface::Get().GetNumWorkerThreads(), Chaos::MaxNumWorkers) : 0;
 			const int32 MaxIslandGroups = (CVars::GIslandGroupsMaxWorkers > 0) ? CVars::GIslandGroupsMaxWorkers : TNumericLimits<int32>::Max();
 			const int32 NumIslandGroups = FMath::Clamp(FMath::CeilToInt32(FReal(NumWorkerThreads) * CVars::GIslandGroupsWorkerMultiplier), 1, MaxIslandGroups);
 
@@ -560,7 +559,7 @@ namespace Chaos
 	#endif
 
 			// @todo(chaos): Remove SolveParallelFor when SolveParallelTasks has been thoroughly tested
-			const bool bSingleThreaded = GSingleThreadedPhysics || (IslandGroups.Num() == 1);
+			const bool bSingleThreaded = bSingleWorkerPhysics || (IslandGroups.Num() == 1);
 			const int32 ParallelMode = bSingleThreaded ? 0 : CVars::GIslandGroupsParallelMode;
 			switch (ParallelMode)
 			{
diff --git a/Engine/Source/Runtime/Experimental/Chaos/Private/Chaos/PBDRigidsEvolution.cpp b/Engine/Source/Runtime/Experimental/Chaos/Private/Chaos/PBDRigidsEvolution.cpp
--- a/Engine/Source/Runtime/Experimental/Chaos/Private/Chaos/PBDRigidsEvolution.cpp
+++ b/Engine/Source/Runtime/Experimental/Chaos/Private/Chaos/PBDRigidsEvolution.cpp
@@ -637,7 +637,7 @@ namespace Chaos
 
 		};
 
-		if (!GSingleThreadedPhysics)
+		if (!bSingleWorkerPhysics)
 		{
 			TaskDispatcher.PruneInternalPendingData(HasToBeUpdated);
 		}
diff --git a/Engine/Source/Runtime/Experimental/Chaos/Private/Chaos/PBDRigidsEvolutionGBF.cpp b/Engine/Source/Runtime/Experimental/Chaos/Private/Chaos/PBDRigidsEvolutionGBF.cpp
--- a/Engine/Source/Runtime/Experimental/Chaos/Private/Chaos/PBDRigidsEvolutionGBF.cpp
+++ b/Engine/Source/Runtime/Experimental/Chaos/Private/Chaos/PBDRigidsEvolutionGBF.cpp
@@ -247,7 +247,7 @@ namespace Chaos
 		const int32 NumCollisionsPerBlock = 2000;
 		const int32 MinCollisionsPerBlock = 100;
 
-		const bool bIsMultithreaded = (FApp::ShouldUseThreadingForPerformance() && !GSingleThreadedPhysics);
+		const bool bIsMultithreaded = (FApp::ShouldUseThreadingForPerformance() && !bSingleWorkerPhysics);
 		if (!bIsMultithreaded)
 		{
 			return NumCollisionsPerBlock;
@@ -1058,7 +1058,7 @@ void FPBDRigidsEvolutionGBF::Integrate(FReal Dt)
 	};
 	
 
-	if (GSingleThreadedPhysics)
+	if (bSingleWorkerPhysics)
 	{
 		// In single threaded mode keep same logic with three batches inside PendingDataArrays: 
 		// 0 Dirty particles from the PushData
@@ -1135,7 +1135,7 @@ void FPBDRigidsEvolutionGBF::ApplyKinematicTargets(const FReal Dt, const FReal S
 		}
 	};
 
-	if (GSingleThreadedPhysics)
+	if (bSingleWorkerPhysics)
 	{
 		DirtyKinAsync();
 		// done with update, let's clear the tracking structures
@@ -1283,7 +1283,7 @@ void FPBDRigidsEvolutionGBF::DispatchKinematicTargetsTasks(int32 DispatchBatchIn
 		}
 	};
 
-	if (GSingleThreadedPhysics)
+	if (bSingleWorkerPhysics)
 	{
 		KinematicTargetsLambda(0, NumParticles, 0);
 	}
diff --git a/Engine/Source/Runtime/Experimental/Chaos/Private/ChaosModule.cpp b/Engine/Source/Runtime/Experimental/Chaos/Private/ChaosModule.cpp
--- a/Engine/Source/Runtime/Experimental/Chaos/Private/ChaosModule.cpp
+++ b/Engine/Source/Runtime/Experimental/Chaos/Private/ChaosModule.cpp
@@ -2,6 +2,7 @@
 
 #include "ChaosModule.h"
 #include "CoreMinimal.h"
+#include "ChaosLog.h"
 #include "ChaosVisualDebugger/ChaosVisualDebuggerTrace.h"
 #include "Modules/ModuleManager.h"
 #include "Misc/CommandLine.h"
@@ -11,6 +12,8 @@ IMPLEMENT_MODULE(FChaosEngineModule, Chaos);
 namespace Chaos
 {
 	extern CHAOS_API int GSingleThreadedPhysics;
+	extern CHAOS_API bool bSingleWorkerPhysics;
+	extern CHAOS_API int32 MinNumberOfWorkerThreadToRunOnSingleWorker;
 }
 
 void FChaosEngineModule::StartupModule()
@@ -23,6 +26,7 @@ void FChaosEngineModule::StartupModule()
 	if(FParse::Param(FCommandLine::Get(), TEXT("SingleThreadedPhysics")))
 	{
 		Chaos::GSingleThreadedPhysics = 1;
+		Chaos::bSingleWorkerPhysics = true;
 	}
 
 #if CHAOS_VISUAL_DEBUGGER_ENABLED
diff --git a/Engine/Source/Runtime/Experimental/Chaos/Public/Chaos/Collision/SpatialAccelerationBroadPhase.h b/Engine/Source/Runtime/Experimental/Chaos/Public/Chaos/Collision/SpatialAccelerationBroadPhase.h
--- a/Engine/Source/Runtime/Experimental/Chaos/Public/Chaos/Collision/SpatialAccelerationBroadPhase.h
+++ b/Engine/Source/Runtime/Experimental/Chaos/Public/Chaos/Collision/SpatialAccelerationBroadPhase.h
@@ -390,20 +390,31 @@ namespace Chaos
 				TArray<UE::Tasks::FTask> NextPendingTasks;
 				const int32 NumContextTask = PendingTasks.Num();
 				check(NumContextTask <= NumActiveBroadphaseContexts);
-				for (int32 ContextIndex = 0; ContextIndex < NumContextTask; ContextIndex++)
+				// If bSingleWorkerPhysics is true, NumActiveBroadphaseContexts shouldn't bigger than 1
+				check(!bSingleWorkerPhysics || NumActiveBroadphaseContexts <= 1);
+				// If bSingleWorkerPhysics is true, PendingTasks should be empty
+				check(!bSingleWorkerPhysics || PendingTasks.IsEmpty());
+				if (bSingleWorkerPhysics)
 				{
-					Private::FBroadPhaseContext& BroadPhaseContext = BroadphaseContexts[ContextIndex];
+					AssignMidPhases(BroadphaseContexts[0]);
+				}
+				else
+				{
+					for (int32 ContextIndex = 0; ContextIndex < NumContextTask; ContextIndex++)
+					{
+						Private::FBroadPhaseContext& BroadPhaseContext = BroadphaseContexts[ContextIndex];
 
-					UE::Tasks::FTask PendingTask = UE::Tasks::Launch(UE_SOURCE_LOCATION, [this, &BroadPhaseContext, CVDContext]()
+						UE::Tasks::FTask PendingTask = UE::Tasks::Launch(UE_SOURCE_LOCATION, [this, &BroadPhaseContext, CVDContext]()
 						{
 							CVD_SCOPE_CONTEXT(CVDContext.Context);
 							AssignMidPhases(BroadPhaseContext);
 						}, PendingTasks[ContextIndex]);
-					NextPendingTasks.Add(PendingTask);
+						NextPendingTasks.Add(PendingTask);
+					}
+					// It is required to wait for the middle phase modifiers, which is a callback to modify middle phases.
+					// This could be removed by adding dependent tasks for the middle phase modifiers.
+					UE::Tasks::Wait(NextPendingTasks);
 				}
-				// It is required to wait for the middle phase modifiers, which is a callback to modify middle phases.
-				// This could be removed by adding dependent tasks for the middle phase modifiers.
-				UE::Tasks::Wait(NextPendingTasks);
 			}
 
 			// Run some error checks in non-shipping builds
@@ -436,154 +447,178 @@ namespace Chaos
 			PendingTasks.Reset();
 			ContextsWithConstraints.Reset();
 			check(NumContextTask <= NumActiveBroadphaseContexts);
-			for (int32 ContextIndex = 0; ContextIndex < NumContextTask; ContextIndex++)
+			// If bSingleWorkerPhysics is true NumActiveBroadphaseContexts should be equal to 1
+			check(!bSingleWorkerPhysics || NumActiveBroadphaseContexts == 1);
+			if (bSingleWorkerPhysics)
 			{
-				Private::FBroadPhaseContext& BroadphaseContext = BroadphaseContexts[ContextIndex];
-				// Accessing the overlaps array is safe here, all overlaps have been finished to be computed at this point
-				if (!BroadphaseContext.Overlaps.IsEmpty())
+				ProcessMidPhases(Dt, BroadphaseContexts[0]);
+			}
+			else
+			{ 
+				for (int32 ContextIndex = 0; ContextIndex < NumContextTask; ContextIndex++)
 				{
-					UE::Tasks::FTask ProcessMidPhaseTask = UE::Tasks::Launch(UE_SOURCE_LOCATION, [this, &BroadphaseContext, Dt, CVDContext]()
+					Private::FBroadPhaseContext& BroadphaseContext = BroadphaseContexts[ContextIndex];
+					// Accessing the overlaps array is safe here, all overlaps have been finished to be computed at this point
+					if (!BroadphaseContext.Overlaps.IsEmpty())
+					{
+						UE::Tasks::FTask ProcessMidPhaseTask = UE::Tasks::Launch(UE_SOURCE_LOCATION, [this, &BroadphaseContext, Dt, CVDContext]()
 						{
 							CVD_SCOPE_CONTEXT(CVDContext.Context);
 							ProcessMidPhases(Dt, BroadphaseContext);
 						}, UE::Tasks::ETaskPriority::High);
-					PendingTasks.Add(ProcessMidPhaseTask);
-					ContextsWithConstraints.Add(ContextIndex);
+						PendingTasks.Add(ProcessMidPhaseTask);
+						ContextsWithConstraints.Add(ContextIndex);
+					}
 				}
 			}
 		}
 
 		void GatherConstraints(bool bIsDeterministic)
 		{
-			// This lambda merge constraints of 2 broad phase context. 
-			// If it is deterministic the lambda assume they are already sorted and merge them in the first given context 
-			// maintaining the sorted data.
-			auto MergeLambda = [this, bIsDeterministic](int32 SmallIndexToMerge, int32 BigIndexToMerge)
+			if (bSingleWorkerPhysics)
 			{
-				QUICK_SCOPE_CYCLE_COUNTER(STAT_MergeConstraints);
-				Private::FBroadPhaseContext& Context = BroadphaseContexts[SmallIndexToMerge];
-				Private::FBroadPhaseContext& OtherContext = BroadphaseContexts[BigIndexToMerge];
-
-				TArray<FPBDCollisionConstraint*>& ThisArray = Context.CollisionContext.Allocator->NewActiveConstraints;
-				TArray<FPBDCollisionConstraint*>& OtherArray = OtherContext.CollisionContext.Allocator->NewActiveConstraints;
-
 				if (bIsDeterministic)
 				{
-					TArray<FPBDCollisionConstraint*> ResultArray;
-					const int32 NumThis = ThisArray.Num();
-					const int32 NumOther = OtherArray.Num();
-					ResultArray.Reserve(NumThis + NumOther);
+					QUICK_SCOPE_CYCLE_COUNTER(STAT_SortConstraint);
+					BroadphaseContexts[0].CollisionContext.Allocator->NewActiveConstraints.Sort(
+					[](const FPBDCollisionConstraintHandle& L, const FPBDCollisionConstraintHandle& R)
+					{
+						return L.GetContact().GetCollisionSortKey() < R.GetContact().GetCollisionSortKey();
+					});
+				}
+			}
+			else
+			{
+				// This lambda merge constraints of 2 broad phase context. 
+				// If it is deterministic the lambda assume they are already sorted and merge them in the first given context 
+				// maintaining the sorted data.
+				auto MergeLambda = [this, bIsDeterministic](int32 SmallIndexToMerge, int32 BigIndexToMerge)
+				{
+					QUICK_SCOPE_CYCLE_COUNTER(STAT_MergeConstraints);
+					Private::FBroadPhaseContext& Context = BroadphaseContexts[SmallIndexToMerge];
+					Private::FBroadPhaseContext& OtherContext = BroadphaseContexts[BigIndexToMerge];
 
-					int32 ThisIndex = 0;
-					int32 OtherIndex = 0;
+					TArray<FPBDCollisionConstraint*>& ThisArray = Context.CollisionContext.Allocator->NewActiveConstraints;
+					TArray<FPBDCollisionConstraint*>& OtherArray = OtherContext.CollisionContext.Allocator->NewActiveConstraints;
 
-					while (ThisIndex < NumThis || OtherIndex < NumOther)
+					if (bIsDeterministic)
 					{
-						if (ThisIndex < NumThis && (OtherIndex >= NumOther || 
-							ThisArray[ThisIndex]->GetContact().GetCollisionSortKey() < OtherArray[OtherIndex]->GetContact().GetCollisionSortKey()))
-						{
-							ResultArray.Add(ThisArray[ThisIndex++]);
-						}
-						else
+						TArray<FPBDCollisionConstraint*> ResultArray;
+						const int32 NumThis = ThisArray.Num();
+						const int32 NumOther = OtherArray.Num();
+						ResultArray.Reserve(NumThis + NumOther);
+
+						int32 ThisIndex = 0;
+						int32 OtherIndex = 0;
+
+						while (ThisIndex < NumThis || OtherIndex < NumOther)
 						{
-							check(OtherIndex < NumOther);
-							ResultArray.Add(OtherArray[OtherIndex++]);
+							if (ThisIndex < NumThis && (OtherIndex >= NumOther || 
+								ThisArray[ThisIndex]->GetContact().GetCollisionSortKey() < OtherArray[OtherIndex]->GetContact().GetCollisionSortKey()))
+							{
+								ResultArray.Add(ThisArray[ThisIndex++]);
+							}
+							else
+							{
+								check(OtherIndex < NumOther);
+								ResultArray.Add(OtherArray[OtherIndex++]);
+							}
 						}
+						ThisArray = MoveTemp(ResultArray);
 					}
-					ThisArray = MoveTemp(ResultArray);
-				}
-				else
-				{
-					ThisArray.Append(OtherArray);
-				}
-				OtherArray.Reset();
-			};
-
-			// This atomic is use to merge the sorted data in a safe multithreaded program. 
-			// When reaching the end of a process, 
-			// - either this atomic is INDEX_NONE, meaning there is no other batch to be merged with. So this atomic take the value of the batch
-			// - either this atomic is the index to be merged with, so the merging is called with this current index, and the one stored in the atomic.
-			std::atomic<int32> NextBatchToMerge = INDEX_NONE;
-			if (bIsDeterministic)
-			{
-				TArray<UE::Tasks::FTask> ProcessSorting;
-				for (int32 ContextIndex = 0; ContextIndex < PendingTasks.Num(); ContextIndex++)
+					else
+					{
+						ThisArray.Append(OtherArray);
+					}
+					OtherArray.Reset();
+				};
+
+				// This atomic is use to merge the sorted data in a safe multithreaded program. 
+				// When reaching the end of a process, 
+				// - either this atomic is INDEX_NONE, meaning there is no other batch to be merged with. So this atomic take the value of the batch
+				// - either this atomic is the index to be merged with, so the merging is called with this current index, and the one stored in the atomic.
+				std::atomic<int32> NextBatchToMerge = INDEX_NONE;
+				if (bIsDeterministic)
 				{
-					int32 CurrentIndex = ContextsWithConstraints[ContextIndex];
-					Private::FBroadPhaseContext& BroadphaseContext = BroadphaseContexts[CurrentIndex];
-					UE::Tasks::FTask SortConstraintTask = UE::Tasks::Launch(UE_SOURCE_LOCATION, [this, &BroadphaseContext, CurrentIndex, &NextBatchToMerge, MergeLambda]()
+					TArray<UE::Tasks::FTask> ProcessSorting;
+					for (int32 ContextIndex = 0; ContextIndex < PendingTasks.Num(); ContextIndex++)
 					{
-						QUICK_SCOPE_CYCLE_COUNTER(STAT_SortConstraintBatch);
-						BroadphaseContext.CollisionContext.Allocator->NewActiveConstraints.Sort(
-							[](const FPBDCollisionConstraintHandle& L, const FPBDCollisionConstraintHandle& R)
-							{
-								return L.GetContact().GetCollisionSortKey() < R.GetContact().GetCollisionSortKey();
-							});
-						int32 NextIndexToMerge = CurrentIndex;
-						while (true) 
+						int32 CurrentIndex = ContextsWithConstraints[ContextIndex];
+						Private::FBroadPhaseContext& BroadphaseContext = BroadphaseContexts[CurrentIndex];
+						UE::Tasks::FTask SortConstraintTask = UE::Tasks::Launch(UE_SOURCE_LOCATION, [this, &BroadphaseContext, CurrentIndex, &NextBatchToMerge, MergeLambda]()
 						{
-							int32 NextExpected = NextBatchToMerge.load(std::memory_order_relaxed);
-							// If there is a sorted batch ready to be merged, lets merge it
-							if (NextExpected != INDEX_NONE)
+							QUICK_SCOPE_CYCLE_COUNTER(STAT_SortConstraintBatch);
+							BroadphaseContext.CollisionContext.Allocator->NewActiveConstraints.Sort(
+								[](const FPBDCollisionConstraintHandle& L, const FPBDCollisionConstraintHandle& R)
+								{
+									return L.GetContact().GetCollisionSortKey() < R.GetContact().GetCollisionSortKey();
+								});
+							int32 NextIndexToMerge = CurrentIndex;
+							while (true) 
 							{
-								// Make sure the value hasn't been overwritten by another thread otherwise, retry
-								if (NextBatchToMerge.compare_exchange_weak(NextExpected, INDEX_NONE, std::memory_order_relaxed))
+								int32 NextExpected = NextBatchToMerge.load(std::memory_order_relaxed);
+								// If there is a sorted batch ready to be merged, lets merge it
+								if (NextExpected != INDEX_NONE)
 								{
-									int32 SmallIndexToMerge = FMath::Min(NextIndexToMerge, NextExpected);
-									int32 BigIndexToMerge = FMath::Max(NextIndexToMerge, NextExpected);
-									check(SmallIndexToMerge != BigIndexToMerge);
-									MergeLambda(SmallIndexToMerge, BigIndexToMerge);
-									NextIndexToMerge = SmallIndexToMerge;
+									// Make sure the value hasn't been overwritten by another thread otherwise, retry
+									if (NextBatchToMerge.compare_exchange_weak(NextExpected, INDEX_NONE, std::memory_order_relaxed))
+									{
+										int32 SmallIndexToMerge = FMath::Min(NextIndexToMerge, NextExpected);
+										int32 BigIndexToMerge = FMath::Max(NextIndexToMerge, NextExpected);
+										check(SmallIndexToMerge != BigIndexToMerge);
+										MergeLambda(SmallIndexToMerge, BigIndexToMerge);
+										NextIndexToMerge = SmallIndexToMerge;
+									}
 								}
-							}
-							else // Otherwise if no batch are available set the current one as the next available for merging
-							{
-								int32 IndexNone = INDEX_NONE;
-								// Make sure there is no available one, otherwise retry
-								if (NextBatchToMerge.compare_exchange_weak(IndexNone, NextIndexToMerge, std::memory_order_relaxed))
+								else // Otherwise if no batch are available set the current one as the next available for merging
 								{
-									break;
+									int32 IndexNone = INDEX_NONE;
+									// Make sure there is no available one, otherwise retry
+									if (NextBatchToMerge.compare_exchange_weak(IndexNone, NextIndexToMerge, std::memory_order_relaxed))
+									{
+										break;
+									}
 								}
 							}
-						}
-					}, PendingTasks[ContextIndex]);
-					ProcessSorting.Add(SortConstraintTask);
+						}, PendingTasks[ContextIndex]);
+						ProcessSorting.Add(SortConstraintTask);
+					}
+					PendingTasks = MoveTemp(ProcessSorting);
 				}
-				PendingTasks = MoveTemp(ProcessSorting);
-			}
-			else // For no deterministic program merge all context in new tasks, without sorting them
-			{
-				// Create a binary tree of task to merge all sorted constraints
-				int32 NeighborStep = 1;
-				while (PendingTasks.Num() >= 2)
+				else // For no deterministic program merge all context in new tasks, without sorting them
 				{
-					int32 NumProcessSorting = PendingTasks.Num();
-					TArray<UE::Tasks::FTask> ProcessMerging;
-					// We merge data in the context that have a power of 2 index, so for first layer 0, 2, 4..., for the second layer 0, 4, 8... and so on.
-					for (int32 ProcessIndex = 0, ContextIndex = 0; ProcessIndex < NumProcessSorting - 1; ProcessIndex += 2, ContextIndex += 2 * NeighborStep)
+					// Create a binary tree of task to merge all sorted constraints
+					int32 NeighborStep = 1;
+					while (PendingTasks.Num() >= 2)
 					{
-						TArray<UE::Tasks::FTask> Preriq{ PendingTasks[ProcessIndex], PendingTasks[ProcessIndex + 1] };
-						int32 SmallIndexToMerge = ContextsWithConstraints[ContextIndex];
-						int32 BigIndexToMerge = ContextsWithConstraints[ContextIndex + NeighborStep];
-						UE::Tasks::FTask MergeConstraintTask = UE::Tasks::Launch(UE_SOURCE_LOCATION, [SmallIndexToMerge, BigIndexToMerge, MergeLambda]()
+						int32 NumProcessSorting = PendingTasks.Num();
+						TArray<UE::Tasks::FTask> ProcessMerging;
+						// We merge data in the context that have a power of 2 index, so for first layer 0, 2, 4..., for the second layer 0, 4, 8... and so on.
+						for (int32 ProcessIndex = 0, ContextIndex = 0; ProcessIndex < NumProcessSorting - 1; ProcessIndex += 2, ContextIndex += 2 * NeighborStep)
 						{
-							QUICK_SCOPE_CYCLE_COUNTER(STAT_MergeConstraints);
-							MergeLambda(SmallIndexToMerge, BigIndexToMerge);
+							TArray<UE::Tasks::FTask> Preriq{ PendingTasks[ProcessIndex], PendingTasks[ProcessIndex + 1] };
+							int32 SmallIndexToMerge = ContextsWithConstraints[ContextIndex];
+							int32 BigIndexToMerge = ContextsWithConstraints[ContextIndex + NeighborStep];
+							UE::Tasks::FTask MergeConstraintTask = UE::Tasks::Launch(UE_SOURCE_LOCATION, [SmallIndexToMerge, BigIndexToMerge, MergeLambda]()
+							{
+								QUICK_SCOPE_CYCLE_COUNTER(STAT_MergeConstraints);
+								MergeLambda(SmallIndexToMerge, BigIndexToMerge);
 								
-						}, Preriq);
-						ProcessMerging.Add(MergeConstraintTask);
-					}
-					// Add the last one for odd process number
-					if (NumProcessSorting % 2 == 1)
-					{
-						ProcessMerging.Add(PendingTasks[NumProcessSorting - 1]);
+							}, Preriq);
+							ProcessMerging.Add(MergeConstraintTask);
+						}
+						// Add the last one for odd process number
+						if (NumProcessSorting % 2 == 1)
+						{
+							ProcessMerging.Add(PendingTasks[NumProcessSorting - 1]);
+						}
+						PendingTasks = MoveTemp(ProcessMerging);
+						NeighborStep *= 2;
 					}
-					PendingTasks = MoveTemp(ProcessMerging);
-					NeighborStep *= 2;
+					ensure(PendingTasks.Num() <= 1);
 				}
-				ensure(PendingTasks.Num() <= 1);
+				UE::Tasks::Wait(PendingTasks);
 			}
-			UE::Tasks::Wait(PendingTasks);
 
 			// Lets move all constraints in the first allocator to be picked up by the CollisionConstraintAllocator
 			if (!ContextsWithConstraints.IsEmpty() && ContextsWithConstraints[0] != 0)
@@ -615,7 +650,7 @@ namespace Chaos
 			}
 
 			// The number of contexts that may have new overlaps in them
-			NumActiveBroadphaseContexts = !!GSingleThreadedPhysics ? 1 : FMath::Max(FMath::Min(Chaos::CVars::NumWorkerCollisionFactor * FTaskGraphInterface::Get().GetNumWorkerThreads(), Chaos::MaxNumWorkers), 1);
+			NumActiveBroadphaseContexts = bSingleWorkerPhysics ? 1 : FMath::Max(FMath::Min(Chaos::CVars::NumWorkerCollisionFactor * FTaskGraphInterface::Get().GetNumWorkerThreads(), Chaos::MaxNumWorkers), 1);
 
 			BroadphaseContexts.SetNum(NumActiveBroadphaseContexts);
 			Allocator->SetMaxContexts(NumActiveBroadphaseContexts);
@@ -652,13 +687,27 @@ namespace Chaos
 					};
 					using LambdaType = decltype(ReturnHandle);
 					const int32 NumHandles = CurHandlesArray->Num();
-					DispatchTasks<bOnlyRigid, SpatialAccelerationType, THandle, LambdaType>(ContextIndex, NumHandles, InSpatialAcceleration, Dt, ReturnHandle);
+					if (bSingleWorkerPhysics)
+					{
+						ProduceParticleOverlapsSingleThreaded<bOnlyRigid, SpatialAccelerationType, THandle, LambdaType>(ContextIndex, NumHandles, InSpatialAcceleration, Dt, ReturnHandle);
+					}
+					else
+					{
+						DispatchTasks<bOnlyRigid, SpatialAccelerationType, THandle, LambdaType>(ContextIndex, NumHandles, InSpatialAcceleration, Dt, ReturnHandle);
+					}
 				}
 				else if (TSOA* Soa = SOAView.SOA)
 				{
 					auto ReturnHandle = [Soa](int32 Index) { return THandleBase(Soa, Index); };
 					using LambdaType = decltype(ReturnHandle);
-					DispatchTasks<bOnlyRigid, SpatialAccelerationType, THandle, LambdaType>(ContextIndex, NumParticles, InSpatialAcceleration, Dt, ReturnHandle);
+					if (bSingleWorkerPhysics)
+					{
+						ProduceParticleOverlapsSingleThreaded<bOnlyRigid, SpatialAccelerationType, THandle, LambdaType>(ContextIndex, NumParticles, InSpatialAcceleration, Dt, ReturnHandle);
+					}
+					else
+					{
+						DispatchTasks<bOnlyRigid, SpatialAccelerationType, THandle, LambdaType>(ContextIndex, NumParticles, InSpatialAcceleration, Dt, ReturnHandle);
+					}
 				}
 			}
 		}
@@ -779,6 +828,24 @@ namespace Chaos
 			}
 		}
 
+		template<bool bOnlyRigid, typename SpatialAccelerationType, typename THandle, typename LambdaType>
+		void ProduceParticleOverlapsSingleThreaded(int32& ContextIndex, int32 NumParticles, const SpatialAccelerationType& InSpatialAcceleration, FReal Dt, LambdaType ReturnHandle)
+		{
+			using THandleBase = typename THandle::THandleBase;
+			using TTransientHandle = typename THandle::TTransientHandle;
+
+			SCOPE_CYCLE_COUNTER(STAT_Collisions_ParticlePairBroadPhase);
+			for (int32 Index = 0; Index < NumParticles; Index++)
+			{
+				THandleBase Handle = ReturnHandle(Index);
+				TTransientHandle& TransientHandle = static_cast<TTransientHandle&>(Handle);
+				if (!TransientHandle.LightWeightDisabled())
+				{
+					ProduceParticleOverlaps<bOnlyRigid>(Dt, TransientHandle.Handle(), InSpatialAcceleration, BroadphaseContexts[ContextIndex]);
+				}
+			}
+		}
+
 		// Generate the set of particles that overlap the specified particle and are allowed to collide with it
 		template<bool bOnlyRigid, typename T_SPATIALACCELERATION>
 		void ProduceParticleOverlaps(
@@ -943,18 +1010,18 @@ namespace Chaos
 			// Prefetch initial set of MidPhases
 			const int32 PrefetchLookahead = 4;
 			const int32 NumContextMidPhases = BroadphaseContext.MidPhases.Num();
-			for (int32 Index = 0; Index < NumContextMidPhases && Index < PrefetchLookahead; Index++)
-			{
-				BroadphaseContext.MidPhases[Index]->CachePrefetch();
-			}
+				for (int32 Index = 0; Index < NumContextMidPhases && Index < PrefetchLookahead; Index++)
+				{
+					BroadphaseContext.MidPhases[Index]->CachePrefetch();
+				}
 
 			for (int32 Index = 0; Index < NumContextMidPhases; Index++)
 			{
-				// Prefetch next MidPhase
-				if (Index + PrefetchLookahead < NumContextMidPhases)
-				{
-					BroadphaseContext.MidPhases[Index + PrefetchLookahead]->CachePrefetch();
-				}
+					// Prefetch next MidPhase
+					if (Index + PrefetchLookahead < NumContextMidPhases)
+					{
+						BroadphaseContext.MidPhases[Index + PrefetchLookahead]->CachePrefetch();
+					}
 
 				FParticlePairMidPhase* MidPhase = BroadphaseContext.MidPhases[Index];
 				const FGeometryParticleHandle* Particle0 = MidPhase->GetParticle0();
diff --git a/Engine/Source/Runtime/Experimental/Chaos/Public/Chaos/Framework/Parallel.h b/Engine/Source/Runtime/Experimental/Chaos/Public/Chaos/Framework/Parallel.h
--- a/Engine/Source/Runtime/Experimental/Chaos/Public/Chaos/Framework/Parallel.h
+++ b/Engine/Source/Runtime/Experimental/Chaos/Public/Chaos/Framework/Parallel.h
@@ -20,6 +20,7 @@ namespace Chaos
 	 */
 	bool CHAOS_API ShouldExecuteTasksInParallel(int32 InTaskNum);
 
+	CHAOS_API extern bool bSingleWorkerPhysics;
 	CHAOS_API extern int32 MaxNumWorkers;
 	CHAOS_API extern int32 SmallBatchSize;
 	CHAOS_API extern int32 LargeBatchSize;
