[TraceBasedDebuggers] fix handler registration

#Synced-CL 48748307

diff --git a/Engine/Source/Runtime/Experimental/TraceBasedDebuggers/Private/EngineEditorBridge.cpp b/Engine/Source/Runtime/Experimental/TraceBasedDebuggers/Private/EngineEditorBridge.cpp
--- a/Engine/Source/Runtime/Experimental/TraceBasedDebuggers/Private/EngineEditorBridge.cpp
+++ b/Engine/Source/Runtime/Experimental/TraceBasedDebuggers/Private/EngineEditorBridge.cpp
@@ -96,7 +96,10 @@ void FEngineEditorBridge::RegisterRemoteSessionsHandler(TSharedRef<IRemoteSessio
 {
 	ensureMsgf(!RemoteSessionsHandler.IsValid(), TEXT("Handler already registered."));
 	RemoteSessionsHandler = Handler;
-	GetSessionsManager()->RegisterExternalHandler(Handler);
+	if (TSharedPtr<FRemoteSessionsManager> SessionsManager = GetSessionsManager())
+	{
+		SessionsManager->RegisterExternalHandler(Handler);
+	}
 }
 
 void FEngineEditorBridge::TearDown()
