[TraceBasedDebuggers] fix handler unregistration

#Synced-CL 48750733

diff --git a/Engine/Source/Runtime/Experimental/TraceBasedDebuggers/Private/EngineEditorBridge.cpp b/Engine/Source/Runtime/Experimental/TraceBasedDebuggers/Private/EngineEditorBridge.cpp
--- a/Engine/Source/Runtime/Experimental/TraceBasedDebuggers/Private/EngineEditorBridge.cpp
+++ b/Engine/Source/Runtime/Experimental/TraceBasedDebuggers/Private/EngineEditorBridge.cpp
@@ -117,7 +117,10 @@ void FEngineEditorBridge::TearDown()
 
 	if (RemoteSessionsHandler)
 	{
-		GetSessionsManager()->UnregisterExternalHandler(RemoteSessionsHandler.ToSharedRef());
+		if (TSharedPtr<FRemoteSessionsManager> SessionsManager = GetSessionsManager())
+		{
+			SessionsManager->UnregisterExternalHandler(RemoteSessionsHandler.ToSharedRef());
+		}
 	}
 
 	bIsInitialized = false;
