[TraceBasedDebuggers] backing out previous quick fixes and properly fixed the shutdown issue (i.e.,prevent creation of the session manager during shutdown)

#Synced-CL 48751102

diff --git a/Engine/Source/Runtime/Experimental/ChaosSolverEngine/Private/Chaos/ChaosSolverEnginePlugin.cpp b/Engine/Source/Runtime/Experimental/ChaosSolverEngine/Private/Chaos/ChaosSolverEnginePlugin.cpp
--- a/Engine/Source/Runtime/Experimental/ChaosSolverEngine/Private/Chaos/ChaosSolverEnginePlugin.cpp
+++ b/Engine/Source/Runtime/Experimental/ChaosSolverEngine/Private/Chaos/ChaosSolverEnginePlugin.cpp
@@ -1,4 +1,4 @@
-// Copyright Epic Games, Inc. All Rights Reserved.
+﻿// Copyright Epic Games, Inc. All Rights Reserved.
 
 #include "Chaos/ChaosSolverEnginePlugin.h"
 
@@ -38,7 +38,7 @@ void FChaosSolverEnginePlugin::StartupModule()
 	{
 		FChaosVDEngineEditorBridge::Get().Initialize();
 	});
-#endif
+#endif // WITH_CHAOS_VISUAL_DEBUGGER
 
 }
 
@@ -49,10 +49,11 @@ void FChaosSolverEnginePlugin::ShutdownModule()
 #endif
 
 #if WITH_CHAOS_VISUAL_DEBUGGER
-	FChaosVDEngineEditorBridge::Get().TearDown();
-#endif
-
+	// Calling Teardown only if the bridge was instantiated since we don't
+	// want calling 'Get()' which will create the global instance during the shutdown.
+	if (FChaosVDEngineEditorBridge::IsInstantiated())
+	{
+		FChaosVDEngineEditorBridge::Get().TearDown();
+	}
+#endif // WITH_CHAOS_VISUAL_DEBUGGER
 }
-
-
-
diff --git a/Engine/Source/Runtime/Experimental/ChaosSolverEngine/Private/Chaos/ChaosVDEngineEditorBridge.cpp b/Engine/Source/Runtime/Experimental/ChaosSolverEngine/Private/Chaos/ChaosVDEngineEditorBridge.cpp
--- a/Engine/Source/Runtime/Experimental/ChaosSolverEngine/Private/Chaos/ChaosVDEngineEditorBridge.cpp
+++ b/Engine/Source/Runtime/Experimental/ChaosSolverEngine/Private/Chaos/ChaosVDEngineEditorBridge.cpp
@@ -1,4 +1,4 @@
-// Copyright Epic Games, Inc. All Rights Reserved.
+﻿// Copyright Epic Games, Inc. All Rights Reserved.
 
 #include "Chaos/ChaosVDEngineEditorBridge.h"
 
@@ -22,6 +22,8 @@
 #include "UObject/Package.h"
 
 
+bool FChaosVDEngineEditorBridge::bIsInstantiated = false;
+
 bool FChaosVDEngineEditorBridge::AddOnScreenRecordingMessage(float DummyDeltaTime)
 {
 	constexpr bool bContinueLooping = false;
@@ -184,9 +186,15 @@ FChaosVDEngineEditorBridge::FChaosVDEngineEditorBridge(): FEngineEditorBridge(Lo
 FChaosVDEngineEditorBridge& FChaosVDEngineEditorBridge::Get()
 {
 	static FChaosVDEngineEditorBridge CVDEngineEditorBridge;
+	bIsInstantiated = true;
 	return CVDEngineEditorBridge;
 }
 
+bool FChaosVDEngineEditorBridge::IsInstantiated()
+{
+	return bIsInstantiated;
+}
+
 void FChaosVDEngineEditorBridge::BuildRecordingStatusInternal(UE::TraceBasedDebuggers::FRecordingStatusMessage& OutStatusMessage) const
 {
 	const FChaosVDRuntimeModule& RuntimeModule = FChaosVDRuntimeModule::Get();
diff --git a/Engine/Source/Runtime/Experimental/ChaosSolverEngine/Public/Chaos/ChaosVDEngineEditorBridge.h b/Engine/Source/Runtime/Experimental/ChaosSolverEngine/Public/Chaos/ChaosVDEngineEditorBridge.h
--- a/Engine/Source/Runtime/Experimental/ChaosSolverEngine/Public/Chaos/ChaosVDEngineEditorBridge.h
+++ b/Engine/Source/Runtime/Experimental/ChaosSolverEngine/Public/Chaos/ChaosVDEngineEditorBridge.h
@@ -1,4 +1,4 @@
-// Copyright Epic Games, Inc. All Rights Reserved.
+﻿// Copyright Epic Games, Inc. All Rights Reserved.
 
 #pragma once
 
@@ -52,6 +52,9 @@ public:
 
 	CHAOSSOLVERENGINE_API static FChaosVDEngineEditorBridge& Get();
 
+	/** @return Whether the static instance has been created (i.e., Get() was called and teardown is required) */
+	CHAOSSOLVERENGINE_API static bool IsInstantiated();
+
 PRAGMA_DISABLE_DEPRECATION_WARNINGS
 	UE_DEPRECATED(5.8, "Use GetSessionsManager instead")
 	TSharedPtr<FChaosVDRemoteSessionsManager> GetRemoteSessionsManager()
@@ -104,6 +107,7 @@ private:
 #endif
 
 	FTSTicker::FDelegateHandle DeferredShowMessageOnScreenHandle;
+	static bool bIsInstantiated;
 };
 
 #else
diff --git a/Engine/Source/Runtime/Experimental/TraceBasedDebuggers/Private/EngineEditorBridge.cpp b/Engine/Source/Runtime/Experimental/TraceBasedDebuggers/Private/EngineEditorBridge.cpp
--- a/Engine/Source/Runtime/Experimental/TraceBasedDebuggers/Private/EngineEditorBridge.cpp
+++ b/Engine/Source/Runtime/Experimental/TraceBasedDebuggers/Private/EngineEditorBridge.cpp
@@ -1,4 +1,4 @@
-// Copyright Epic Games, Inc. All Rights Reserved.
+﻿// Copyright Epic Games, Inc. All Rights Reserved.
 
 #include "EngineEditorBridge.h"
 
@@ -72,18 +72,17 @@ void FEngineEditorBridge::SetDataRelayTransportInstance(const TSharedPtr<IDataRe
 }
 
 FEngineEditorBridge::FEngineEditorBridge(const FLogCategoryBase& LogCategory)
-	: LogCategory(LogCategory)
+	: RemoteSessionsManager(SharedRemoteSessionsManager.IsValid() ? SharedRemoteSessionsManager.Pin().ToSharedRef() : MakeShared<FRemoteSessionsManager>())
+	, LogCategory(LogCategory)
 {
-}
-
-void FEngineEditorBridge::Initialize()
-{
-	RemoteSessionsManager = SharedRemoteSessionsManager.IsValid() ? SharedRemoteSessionsManager.Pin().ToSharedRef() : MakeShared<FRemoteSessionsManager>();
 	if (!SharedRemoteSessionsManager.IsValid())
 	{
 		SharedRemoteSessionsManager = RemoteSessionsManager;
 	}
+}
 
+void FEngineEditorBridge::Initialize()
+{
 	// Allow derived classes to perform their initialization
 	OnInitializeInternal();
 
@@ -96,10 +95,7 @@ void FEngineEditorBridge::RegisterRemoteSessionsHandler(TSharedRef<IRemoteSessio
 {
 	ensureMsgf(!RemoteSessionsHandler.IsValid(), TEXT("Handler already registered."));
 	RemoteSessionsHandler = Handler;
-	if (TSharedPtr<FRemoteSessionsManager> SessionsManager = GetSessionsManager())
-	{
-		SessionsManager->RegisterExternalHandler(Handler);
-	}
+	GetSessionsManager()->RegisterExternalHandler(Handler);
 }
 
 void FEngineEditorBridge::TearDown()
@@ -117,10 +113,7 @@ void FEngineEditorBridge::TearDown()
 
 	if (RemoteSessionsHandler)
 	{
-		if (TSharedPtr<FRemoteSessionsManager> SessionsManager = GetSessionsManager())
-		{
-			SessionsManager->UnregisterExternalHandler(RemoteSessionsHandler.ToSharedRef());
-		}
+		GetSessionsManager()->UnregisterExternalHandler(RemoteSessionsHandler.ToSharedRef());
 	}
 
 	bIsInitialized = false;
diff --git a/Engine/Source/Runtime/Experimental/TraceBasedDebuggers/Public/EngineEditorBridge.h b/Engine/Source/Runtime/Experimental/TraceBasedDebuggers/Public/EngineEditorBridge.h
--- a/Engine/Source/Runtime/Experimental/TraceBasedDebuggers/Public/EngineEditorBridge.h
+++ b/Engine/Source/Runtime/Experimental/TraceBasedDebuggers/Public/EngineEditorBridge.h
@@ -1,4 +1,4 @@
-// Copyright Epic Games, Inc. All Rights Reserved.
+﻿// Copyright Epic Games, Inc. All Rights Reserved.
 
 #pragma once
 
@@ -92,7 +92,7 @@ private:
 	TSharedPtr<IDataRelayTransport> RelayTraceDataTransportInstance;
 
 	static TWeakPtr<FRemoteSessionsManager> SharedRemoteSessionsManager;
-	TSharedPtr<FRemoteSessionsManager> RemoteSessionsManager;
+	TSharedRef<FRemoteSessionsManager> RemoteSessionsManager;
 	/** Handler created by derived classes to handle more message types */
 	TSharedPtr<IRemoteSessionsHandler> RemoteSessionsHandler;
 
