Fix warning V792: The 'IsInRelevancyVolumeAtPosition' function located to the right of the operator '|' will be called regardless of the value of the left operand. Perhaps, it is better to use '||'.
#rnx
#Synced-CL 48746020

diff --git a/Engine/Source/Runtime/Experimental/Chaos/Private/ChaosVisualDebugger/ChaosVisualDebuggerTrace.cpp b/Engine/Source/Runtime/Experimental/Chaos/Private/ChaosVisualDebugger/ChaosVisualDebuggerTrace.cpp
--- a/Engine/Source/Runtime/Experimental/Chaos/Private/ChaosVisualDebugger/ChaosVisualDebuggerTrace.cpp
+++ b/Engine/Source/Runtime/Experimental/Chaos/Private/ChaosVisualDebugger/ChaosVisualDebuggerTrace.cpp
@@ -1024,7 +1024,7 @@ void FChaosVisualDebuggerTrace::TraceSceneQueryStart(const Chaos::FImplicitObjec
 
 	// In overlap queries, only the start position is valid
 	const bool bEvaluateBothPositions = QueryType != EChaosVDSceneQueryType::Overlap;
-	bool bIsRelevant = bEvaluateBothPositions ? (IsInRelevancyVolumeAtPosition(Start, GeometryOrientation) | IsInRelevancyVolumeAtPosition(End, GeometryOrientation))
+	bool bIsRelevant = bEvaluateBothPositions ? (IsInRelevancyVolumeAtPosition(Start, GeometryOrientation) || IsInRelevancyVolumeAtPosition(End, GeometryOrientation))
 												: IsInRelevancyVolumeAtPosition(Start, GeometryOrientation);
 
 	if (!bIsRelevant)
