[TraceBasedDebuggers] improved uses of WITH_TRACE_BASED_DEBUGGERS

#Synced-CL 48762536

diff --git a/Engine/Source/Runtime/Experimental/ChaosSolverEngine/Private/Chaos/ChaosVDRemoteSessionsManager.cpp b/Engine/Source/Runtime/Experimental/ChaosSolverEngine/Private/Chaos/ChaosVDRemoteSessionsManager.cpp
--- a/Engine/Source/Runtime/Experimental/ChaosSolverEngine/Private/Chaos/ChaosVDRemoteSessionsManager.cpp
+++ b/Engine/Source/Runtime/Experimental/ChaosSolverEngine/Private/Chaos/ChaosVDRemoteSessionsManager.cpp
@@ -2,16 +2,16 @@
 
 #include "Chaos/ChaosVDRemoteSessionsManager.h"
 
+#if WITH_TRACE_BASED_DEBUGGERS
 #include "ChaosVDRuntimeModule.h"
-#include "MessageEndpointBuilder.h"
 #include "ChaosVisualDebugger/ChaosVDOptionalDataChannel.h"
 #include "ChaosVisualDebugger/ChaosVisualDebuggerTrace.h"
+#include "MessageEndpointBuilder.h"
 
 void FChaosVDRemoteSessionsHandler::OnCreatingMessageEndpoint(
 	const TSharedPtr<UE::TraceBasedDebuggers::FRemoteSessionsManager>& SessionsManager
 	, FMessageEndpointBuilder& InEndPointBuilder)
 {
-#if WITH_TRACE_BASED_DEBUGGERS
 	InEndPointBuilder
 		.Handling<FChaosVDChannelStateChangeCommandMessage>(
 			[](const FChaosVDChannelStateChangeCommandMessage& InMessage, const TSharedRef<IMessageContext>&)
@@ -66,7 +66,6 @@ void FChaosVDRemoteSessionsHandler::OnCreatingMessageEndpoint(
 					}
 				});
 	}
-#endif
 }
 
 void FChaosVDRemoteSessionsHandler::OnSubscribingMessageTypes(
@@ -94,7 +93,6 @@ void FChaosVDRemoteSessionsHandler::OnSubscribingMessageTypes(
 void FChaosVDRemoteSessionsHandler::OnBuildingFullSessionInfoResponseMessage(
 	UE::TraceBasedDebuggers::FFullSessionInfoResponseMessage& InMessage)
 {
-#if WITH_TRACE_BASED_DEBUGGERS
 	using namespace Chaos::VisualDebugger;
 
 	InMessage.bIsRecording = FChaosVDRuntimeModule::Get().IsRecording();
@@ -117,7 +115,6 @@ void FChaosVDRemoteSessionsHandler::OnBuildingFullSessionInfoResponseMessage(
 		});
 
 	InMessage.DebuggerSpecificData.Emplace(FInstancedStruct::Make(MoveTemp(MessageData)));
-#endif
 }
 
 void FChaosVDRemoteSessionsHandler::OnHandlingFullSessionInfoResponseMessage(
@@ -142,6 +139,8 @@ void FChaosVDRemoteSessionsHandler::OnHandlingFullSessionInfoResponseMessage(
 	}
 }
 
+#endif // WITH_TRACE_BASED_DEBUGGERS
+
 PRAGMA_DISABLE_DEPRECATION_WARNINGS
 
 const FChaosVDTraceDetails& FChaosVDSessionInfo::GetConnectionDetails()
diff --git a/Engine/Source/Runtime/Experimental/ChaosSolverEngine/Public/Chaos/ChaosVDRemoteSessionsManager.h b/Engine/Source/Runtime/Experimental/ChaosSolverEngine/Public/Chaos/ChaosVDRemoteSessionsManager.h
--- a/Engine/Source/Runtime/Experimental/ChaosSolverEngine/Public/Chaos/ChaosVDRemoteSessionsManager.h
+++ b/Engine/Source/Runtime/Experimental/ChaosSolverEngine/Public/Chaos/ChaosVDRemoteSessionsManager.h
@@ -160,6 +160,7 @@ struct FChaosVDSessionData
 	TMap<FString, FChaosVDDataChannelState> DataChannelsStatesByName;
 };
 
+#if WITH_TRACE_BASED_DEBUGGERS
 /**
  * Remote sessions handler to handle ChaosVD specific message types.
  */
@@ -170,6 +171,7 @@ struct FChaosVDRemoteSessionsHandler : UE::TraceBasedDebuggers::IRemoteSessionsH
 	UE_API virtual void OnBuildingFullSessionInfoResponseMessage(UE::TraceBasedDebuggers::FFullSessionInfoResponseMessage&) override;
 	UE_API virtual void OnHandlingFullSessionInfoResponseMessage(const UE::TraceBasedDebuggers::FFullSessionInfoResponseMessage&, const TSharedPtr<UE::TraceBasedDebuggers::FSessionInfo>&) override;
 };
+#endif //WITH_TRACE_BASED_DEBUGGERS
 
 PRAGMA_DISABLE_DEPRECATION_WARNINGS
 
diff --git a/Engine/Source/Runtime/Experimental/TraceBasedDebuggers/Private/TraceBasedDebuggerRuntime.cpp b/Engine/Source/Runtime/Experimental/TraceBasedDebuggers/Private/TraceBasedDebuggerRuntime.cpp
--- a/Engine/Source/Runtime/Experimental/TraceBasedDebuggers/Private/TraceBasedDebuggerRuntime.cpp
+++ b/Engine/Source/Runtime/Experimental/TraceBasedDebuggers/Private/TraceBasedDebuggerRuntime.cpp
@@ -2,6 +2,8 @@
 
 #include "TraceBasedDebuggerRuntime.h"
 
+#if WITH_TRACE_BASED_DEBUGGERS
+
 #include "Misc/CommandLine.h"
 #include "RelayTraceDataWriter.h"
 #include "RemoteSessionsManager.h"
@@ -223,7 +225,6 @@ void FRuntimeModule::StartRecording(const FStartRecordingCommandMessage& InRecor
 				break;
 			}
 
-#if WITH_TRACE_BASED_DEBUGGERS
 			if (RelayWriter && !RelayWriter->IsClosed())
 			{
 				UE_LOG_REF(LogCategory, Warning, TEXT("[%hs] Attempting to relay trace, but there is a relay writer still open"), __func__);
@@ -239,7 +240,6 @@ void FRuntimeModule::StartRecording(const FStartRecordingCommandMessage& InRecor
 			RelayWriter->SetMaxAllowedPendingBytes(FMath::Clamp(CVarMaxPendingBytesInRelayWriter->GetInt(), MinBufferSizeLimit, std::numeric_limits<int32>::max()));
 
 			bIsRecording = FTraceAuxiliary::Relay(reinterpret_cast<UPTRINT>(RelayWriter.Get()), FRelayTraceDataWriter::WriteHelper, FRelayTraceDataWriter::CloseHelper, nullptr, &TracingOptions);
-#endif
 
 			if (bIsRecording)
 			{
@@ -285,9 +285,7 @@ void FRuntimeModule::StartRecording(const FStartRecordingCommandMessage& InRecor
 		FMessageDialog::Open(EAppMsgType::Ok, FText::FormatOrdered(LOCTEXT("StartRecordingFailedMessage", "Failed to start recording. \n\n{0}"), FailureReason));
 #endif
 
-#if WITH_TRACE_BASED_DEBUGGERS
 		RelayWriter = nullptr;
-#endif
 
 		// Allow derived classes to perform additional actions on failure
 		OnRecordingStartFailedInternal(FailureReason);
@@ -324,12 +322,11 @@ void FRuntimeModule::StopRecording()
 
 	ClearRelayDataPumpDelegates();
 
-#if WITH_TRACE_BASED_DEBUGGERS
+
 	if (RelayWriter)
 	{
 		RelayWriter->Close();
 	}
-#endif
 
 	bIsRecording = false;
 	AccumulatedRecordingTime = 0.0f;
@@ -343,12 +340,10 @@ int64 FRuntimeModule::GetBufferedDataBytesSize() const
 {
 	FGuardedInt64 RelaySendBufferSize(0);
 
-#if WITH_TRACE_BASED_DEBUGGERS
 	if (RelayWriter)
 	{
 		RelaySendBufferSize += RelayWriter->GetQueuedBytesNum();
 	}
-#endif
 
 	if (ExternalRelayBufferSizeProviderDelegate.IsBound())
 	{
@@ -467,17 +462,14 @@ void FRuntimeModule::HandleRelayOverflow()
 
 void FRuntimeModule::RelayTraceData() const
 {
-#if WITH_TRACE_BASED_DEBUGGERS
 	if (ensure(ExternalRelayExecutorDelegate.IsBound() && RelayWriter.IsValid()))
 	{
 		ExternalRelayExecutorDelegate.Execute(*RelayWriter);
 	}
-#endif
 }
 
 void FRuntimeModule::SetupRelayDataPumpDelegates()
 {
-#if WITH_TRACE_BASED_DEBUGGERS
 	if (!ensure(RelayWriter))
 	{
 		return;
@@ -485,18 +477,15 @@ void FRuntimeModule::SetupRelayDataPumpDelegates()
 
 	RelayWriter->OnNewDataAvailable().BindRaw(this, &FRuntimeModule::RelayTraceData);
 	RelayWriter->OnDataOverflow().BindRaw(this, &FRuntimeModule::HandleRelayOverflow);
-#endif
 }
 
 void FRuntimeModule::ClearRelayDataPumpDelegates()
 {
-#if WITH_TRACE_BASED_DEBUGGERS
 	if (RelayWriter)
 	{
 		RelayWriter->OnNewDataAvailable().Unbind();
 		RelayWriter->OnDataOverflow().Unbind();
 	}
-#endif
 }
 
 bool FRuntimeModule::RecordingTimerTick(float DeltaTime)
@@ -654,4 +643,5 @@ void FRuntimeModule::OnEnableRequiredTraceChannelsInternal()
 
 } // UE::TraceBasedDebuggers
 
-#undef LOCTEXT_NAMESPACE
\ No newline at end of file
+#undef LOCTEXT_NAMESPACE
+#endif // WITH_TRACE_BASED_DEBUGGERS
\ No newline at end of file
diff --git a/Engine/Source/Runtime/Experimental/TraceBasedDebuggers/Public/TraceBasedDebuggerRuntime.h b/Engine/Source/Runtime/Experimental/TraceBasedDebuggers/Public/TraceBasedDebuggerRuntime.h
--- a/Engine/Source/Runtime/Experimental/TraceBasedDebuggers/Public/TraceBasedDebuggerRuntime.h
+++ b/Engine/Source/Runtime/Experimental/TraceBasedDebuggers/Public/TraceBasedDebuggerRuntime.h
@@ -2,6 +2,7 @@
 
 #pragma once
 
+#if WITH_TRACE_BASED_DEBUGGERS
 #include "Containers/Ticker.h"
 #include "ProfilingDebugging/TraceAuxiliary.h"
 #include "TraceDataRelayTransport.h"
@@ -141,11 +142,7 @@ public:
 
 	FRelayTraceDataWriter* GetRelayTraceWriterInstance() const
 	{
-#if WITH_TRACE_BASED_DEBUGGERS
 		return RelayWriter.Get();
-#else
-		return nullptr;
-#endif
 	}
 
 protected:
@@ -210,9 +207,7 @@ private:
 	FDelegateHandle OnConnectionDelegateHandle;
 	FDelegateHandle OnTraceStoppedDelegateHandle;
 
-#if WITH_TRACE_BASED_DEBUGGERS
 	TUniquePtr<FRelayTraceDataWriter> RelayWriter;
-#endif
 
 	FSimpleMulticastDelegate TraceConnectionDetailsUpdatedDelegate;
 	FSimpleDelegate ExternalRelayOverflowHandlerDelegate;
@@ -233,4 +228,5 @@ private:
 
 } // UE::TraceBasedDebuggers
 
-#undef UE_API
\ No newline at end of file
+#undef UE_API
+#endif //WITH_TRACE_BASED_DEBUGGERS
\ No newline at end of file
