Chaos API - cleanup, deduplication, and simplification
- remove internal lock from Object Pointers, instead require user to obtain a scene lock
- fix all natvis

#rnx

#Synced-CL 48564344

diff --git a/Engine/Plugins/Experimental/ChaosRigidPhysicsAsync/Source/ChaosRigidPhysicsAsync/Private/ChaosRigidPhysicsAsync/RigidBodyAsync.cpp b/Engine/Plugins/Experimental/ChaosRigidPhysicsAsync/Source/ChaosRigidPhysicsAsync/Private/ChaosRigidPhysicsAsync/RigidBodyAsync.cpp
--- a/Engine/Plugins/Experimental/ChaosRigidPhysicsAsync/Source/ChaosRigidPhysicsAsync/Private/ChaosRigidPhysicsAsync/RigidBodyAsync.cpp
+++ b/Engine/Plugins/Experimental/ChaosRigidPhysicsAsync/Source/ChaosRigidPhysicsAsync/Private/ChaosRigidPhysicsAsync/RigidBodyAsync.cpp
@@ -3,64 +3,143 @@
 
 #include "ChaosRigidPhysicsAsync/RigidBodyContainerAsync.h"
 #include "ChaosRigidPhysicsAsync/RigidSceneAsync.h"
-#include "RigidPhysics/RigidAccessor.h"
 #include "RigidPhysics/RigidBody.h"
 
 #if UE_RIGIDPHYSICS_API_ENABLED
 
 namespace Chaos::Rigids::Async
 {
-	FRigidBodyAsyncGT::FRigidBodyAsyncGT(
-		FRigidBodyContainerAsyncGT* InContainer, 
-		FSingleParticlePhysicsProxy* InProxy, 
-		const FRigidBodyIdAsync& InBodyId)
-		: Container(InContainer)
-		, Proxy(InProxy)
-		, BodyId(InBodyId)
+	UE_RIGIDPHYSICS_RIGIDTYPED_IMPL(FRigidBodyAsyncGT, UE::Physics::IRigidBody);
+
+	FRigidBodyAsyncGT::FRigidBodyAsyncGT(FRigidBodyPoolAsyncGT* InContainer, const UE::Physics::FRigidDebugName& InName, UE::Physics::ERigidMovementType InMovementType)
+		: BodyId()
+		, BodyContainer(InContainer)
+		, Proxy(nullptr)
 	{
+		using namespace UE::Physics;
+		UE_RIGIDPHYSICS_CHECK(BodyContainer != nullptr);
+
+		TUniquePtr<FGeometryParticle> Particle = nullptr;
+		if (InMovementType == ERigidMovementType::Static)
+		{
+			Particle = FGeometryParticle::CreateParticle();
+		}
+		else
+		{
+			TUniquePtr<FPBDRigidParticle> RigidParticle = FPBDRigidParticle::CreateParticle();
+			if (InMovementType == ERigidMovementType::Kinematic)
+			{
+				RigidParticle->SetObjectState(EObjectStateType::Kinematic);
+			}
+			Particle = MoveTemp(RigidParticle);
+		}
+
+		UE_RIGIDPHYSICS_CHECK(Particle != nullptr);
+		if (Particle != nullptr)
+		{
+#if CHAOS_DEBUG_NAME
+			Particle->SetDebugName(InName.SharedValue());
+#endif
+
+			// Create a default shape for testing
+			// TODO_CHAOS_API: Pass body settings into factory functions
+			FImplicitObjectPtr Geom = FImplicitObjectPtr(new FImplicitBox3(FVec3(-50), FVec3(50)));
+			FImplicitObjectPtr Union = new FImplicitObjectUnion(TArray<FImplicitObjectPtr>({ Geom }));
+			Particle->SetGeometry(Union);
+
+			Proxy = FSingleParticlePhysicsProxy::Create(MoveTemp(Particle));
+		}
 	}
 
 	FRigidBodyAsyncGT::~FRigidBodyAsyncGT()
 	{
+		UE_RIGIDPHYSICS_CHECK(!IsActive());
+
+		// If we were never deactivated, our proxy will not have been destroyed. See Deactivate().
+		if (Proxy != nullptr)
+		{
+			delete Proxy;
+		}
 	}
 
-	const char* FRigidBodyAsyncGT::GetTypeName() const
+	UE::Physics::FRigidBodyHandle FRigidBodyAsyncGT::GetHandle() const
 	{
-		return "RigidBodyAsyncGT";
+		using namespace UE::Physics;
+
+		return FRigidBodyHandle(BodyContainer->GetHandle(), BodyId);
 	}
 
-	UE::Physics::FRigidBodyId FRigidBodyAsyncGT::GetBodyId() const
+	void FRigidBodyAsyncGT::SetId(const UE::Physics::FRigidObjectId& InId)
 	{
-		CHAOSRIGIDPHYSICSASYNC_CHECK(IsValid());
-		// NOTE: Cannot use Proxy UniqueIdx - not assigned until RegisterObject
-		// so we must create yet another ID!
-		return GetBodyIdAsync().AsRigidBodyId();
+		BodyId = InId;
 	}
 
-	UE::Physics::FRigidBodyHandle FRigidBodyAsyncGT::GetBodyHandle() const
+	bool FRigidBodyAsyncGT::IsValid() const
 	{
-		check(IsValid());
-		return UE::Physics::FRigidBodyHandle(Container->GetContainerHandle(), GetBodyId());
+		return (Proxy != nullptr) && (Proxy->GetParticle_LowLevel() != nullptr);
 	}
 
-	FRigidSceneAsyncGT* FRigidBodyAsyncGT::GetSceneAsync() const
+	const FSingleParticlePhysicsProxy* FRigidBodyAsyncGT::GetProxy() const
 	{
-		return Container->GetSceneAsync();
+		return Proxy;
 	}
 
-	bool FRigidBodyAsyncGT::IsValid() const
+	FSingleParticlePhysicsProxy* FRigidBodyAsyncGT::GetProxy()
 	{
-		return (Proxy != nullptr) && (Proxy->GetParticle_LowLevel() != nullptr);
+		return Proxy;
+	}
+
+	const FRigidBodyHandle_External& FRigidBodyAsyncGT::GetProxyAPI() const
+	{
+		UE_RIGIDPHYSICS_CHECK(IsValid());
+		return Proxy->GetGameThreadAPI();
+	}
+
+	FRigidBodyHandle_External& FRigidBodyAsyncGT::GetProxyAPI()
+	{
+		UE_RIGIDPHYSICS_CHECK(IsValid());
+		return Proxy->GetGameThreadAPI();
 	}
 
-	bool FRigidBodyAsyncGT::IsInScene() const
+	FRigidSceneAsyncGT* FRigidBodyAsyncGT::GetScene() const
 	{
-		// TODO_CHAOSAPI: need a better way to know if we're in the scene
+		UE_RIGIDPHYSICS_CHECK(BodyContainer != nullptr);
+		return BodyContainer->GetScene();
+	}
+
+	bool FRigidBodyAsyncGT::IsActive() const
+	{
+		// TODO_CHAOSAPI: need a better way to know if we're active in the scene
 		// This works for now because UniqueIdx is set in FPBDRigidsSolver::RegisterObject 
 		// and GetMarkDeleted() is true after FPBDRigidsSolver::UnregisterObject.
 		return (Proxy != nullptr) && !Proxy->GetMarkedDeleted() && (Proxy->GetParticle_LowLevel() != nullptr) && Proxy->GetParticle_LowLevel()->UniqueIdx().IsValid();
 	}
 
+	void FRigidBodyAsyncGT::Activate()
+	{
+		// Cannot re-activate after deactivate because the proxy is destroyed by UnregisterObject
+		// TODO_CHAOSAPI: Fix this - Proxy lifetime should be shared and PT should only control
+		// lifetime of PT elements.
+		if (Proxy != nullptr)
+		{
+			if (FRigidSceneAsyncGT* Scene = GetScene())
+			{
+				Scene->ActivateBody(this);
+			}
+		}
+	}
+
+	void FRigidBodyAsyncGT::Deactivate()
+	{
+		if (FRigidSceneAsyncGT* Scene = GetScene())
+		{
+			Scene->DeactivateBody(this);
+
+			// Cannot re-activate after deactivate because the proxy is destroyed by UnregisterObject.
+			Proxy = nullptr;
+		}
+	}
+
 	bool FRigidBodyAsyncGT::IsStatic() const
 	{
 		return (GetProxyAPI().ObjectState() == EObjectStateType::Static);
@@ -90,7 +169,7 @@ namespace Chaos::Rigids::Async
 
 	void FRigidBodyAsyncGT::UpdateTransform(const FTransform& InTransform)
 	{
-		CHAOSRIGIDPHYSICSASYNC_CHECK(IsValid());
+		UE_RIGIDPHYSICS_CHECK(IsValid());
 
 		FRigidBodyHandle_External& ProxyAPI = GetProxyAPI();
 		if (!IsKinematic() && !IsSleeping() && FVec3::IsNearlyEqual(InTransform.GetLocation(), ProxyAPI.X(), SMALL_NUMBER) && FRotation3::IsNearlyEqual(InTransform.GetRotation(), ProxyAPI.R(), SMALL_NUMBER))
@@ -111,29 +190,23 @@ namespace Chaos::Rigids::Async
 		ProxyAPI.SetX(InTransform.GetTranslation());
 		ProxyAPI.SetR(InTransform.GetRotation());
 		ProxyAPI.UpdateShapeBounds();
-
-		// Notify container (e.g., update the spatial acceleration or other dependents)
-		if (Container != nullptr)
-		{
-			Container->OnRigidBodyTransformUpdated(this);
-		}
 	}
 
 	void FRigidBodyAsyncGT::SetKinematicTarget(const FTransform& InTransform)
 	{
-		CHAOSRIGIDPHYSICSASYNC_CHECK(IsValid());
+		UE_RIGIDPHYSICS_CHECK(IsValid());
 		GetProxyAPI().SetKinematicTarget(FKinematicTarget::MakePositionTarget(InTransform.GetTranslation(), InTransform.GetRotation()));
 	}
 
 	FTransform FRigidBodyAsyncGT::GetTransform() const
 	{
-		CHAOSRIGIDPHYSICSASYNC_CHECK(IsValid());
+		UE_RIGIDPHYSICS_CHECK(IsValid());
 		return FTransform(GetProxyAPI().R(), GetProxyAPI().X());
 	}
 
 	FBounds3d FRigidBodyAsyncGT::GetBounds() const
 	{
-		CHAOSRIGIDPHYSICSASYNC_CHECK(IsValid());
+		UE_RIGIDPHYSICS_CHECK(IsValid());
 		FConstGenericParticleHandle ParticleHandle = Proxy->GetHandle_LowLevel();
 		FAABB3 Bounds;
 		if (GetProxyAPI().GetGeometry()->HasBoundingBox())
@@ -149,13 +222,13 @@ namespace Chaos::Rigids::Async
 
 	double FRigidBodyAsyncGT::GetMass() const
 	{
-		CHAOSRIGIDPHYSICSASYNC_CHECK(IsValid());
+		UE_RIGIDPHYSICS_CHECK(IsValid());
 		return GetProxyAPI().M();
 	}
 
 	void FRigidBodyAsyncGT::SetMass(double InMass)
 	{
-		CHAOSRIGIDPHYSICSASYNC_CHECK(IsValid());
+		UE_RIGIDPHYSICS_CHECK(IsValid());
 		GetProxyAPI().SetM(InMass);
 	}
 
@@ -170,6 +243,7 @@ namespace Chaos::Rigids::Async
 
 	void FRigidBodyAsyncGT::SetGravityScale(float InScale)
 	{ 
+		UE_RIGIDPHYSICS_CHECK(IsValid());
 		GetProxyAPI().SetGravityEnabled(InScale > 0.0f);
 	}
 
@@ -182,13 +256,13 @@ namespace Chaos::Rigids::Async
 
 namespace Chaos::Rigids::Async
 {
+	UE_RIGIDPHYSICS_RIGIDTYPED_IMPL(FRigidBodyAsyncPT, UE::Physics::IRigidBody);
+
 	FRigidBodyAsyncPT::FRigidBodyAsyncPT(
-		FRigidBodyContainerAsyncPT* InContainer, 
-		FSingleParticlePhysicsProxy* InProxy, 
-		const FRigidBodyIdAsync& InBodyId)
-		: Container(InContainer)
+		const UE::Physics::FRigidBodyHandle& InBodyHandle,
+		FSingleParticlePhysicsProxy* InProxy)
+		: Handle(InBodyHandle)
 		, Proxy(InProxy)
-		, BodyId(InBodyId)
 	{
 	}
 
@@ -196,21 +270,10 @@ namespace Chaos::Rigids::Async
 	{
 	}
 
-	const char* FRigidBodyAsyncPT::GetTypeName() const
-	{
-		return "RigidBodyAsyncPT";
-	}
-
-	UE::Physics::FRigidBodyId FRigidBodyAsyncPT::GetBodyId() const
+	UE::Physics::FRigidBodyHandle FRigidBodyAsyncPT::GetHandle() const
 	{
-		CHAOSRIGIDPHYSICSASYNC_CHECK(IsValid());
-		return GetBodyIdAsync().AsRigidBodyId();
-	}
-
-	UE::Physics::FRigidBodyHandle FRigidBodyAsyncPT::GetBodyHandle() const
-	{
-		CHAOSRIGIDPHYSICSASYNC_CHECK(IsValid());
-		return UE::Physics::FRigidBodyHandle(Container->GetContainerHandle(), GetBodyId());
+		UE_RIGIDPHYSICS_CHECK(IsValid());
+		return Handle;
 	}
 
 	bool FRigidBodyAsyncPT::IsValid() const
@@ -218,14 +281,6 @@ namespace Chaos::Rigids::Async
 		return (Proxy != nullptr) && (Proxy->GetHandle_LowLevel() != nullptr);
 	}
 
-	bool FRigidBodyAsyncPT::IsInScene() const
-	{
-		// TODO_CHAOSAPI: Asymmetry - proxies are destroyed on remove from scene
-		// so once we are in the scene we will always be in the scene until destroyed
-		// TODO_CHAOSAPI: need a better way to know if we're in the scene
-		return (Proxy != nullptr) && (Proxy->GetHandle_LowLevel() != nullptr) && (Proxy->GetHandle_LowLevel()->PhysicsProxy() != nullptr);
-	}
-
 	bool FRigidBodyAsyncPT::IsStatic() const
 	{
 		return (GetProxyAPI()->ObjectState() == EObjectStateType::Static);
@@ -248,35 +303,43 @@ namespace Chaos::Rigids::Async
 
 	void FRigidBodyAsyncPT::InitTransform(const FTransform& InTransform)
 	{
-		CHAOSRIGIDPHYSICSASYNC_CHECK(IsValid());
-		if (FPBDRigidsSolver* Solver = Proxy->GetSolver<FPBDRigidsSolver>())
-		{
-			// NOTE: The Chaos Proxy API does not do everything required when changing
-			// properties on the particle. E.g., marking dirty, etc.
-			// Here'what should be in GetProxyAPI()->SetX()/SetR()
-			FGeometryParticleHandle* ParticleHandle = Proxy->GetHandle_LowLevel();
-			Solver->GetEvolution()->SetParticleTransform(ParticleHandle, InTransform.GetTranslation(), InTransform.GetRotation(), /*bIsTeleport*/true);
-			Solver->GetEvolution()->DirtyParticle(*ParticleHandle);
-		}
+		SetTransform(InTransform,  /*bIsTeleport*/true);
 	}
 
 	void FRigidBodyAsyncPT::UpdateTransform(const FTransform& InTransform)
 	{
-		CHAOSRIGIDPHYSICSASYNC_CHECK(IsValid());
+		SetTransform(InTransform,  /*bIsTeleport*/false);
+	}
+
+	void FRigidBodyAsyncPT::SetTransform(const FTransform& InTransform, bool bIsTeleport)
+	{
+		UE_RIGIDPHYSICS_CHECK(IsValid());
+		if (!IsDynamic())
+		{
+			// TODO_CHAOSAPI: Statics are assumed to be controlled only by the GT. We should track changes that need 
+			// to be sent back to the GT (this is also true for various properties like EtherDrag, and not just on Statics).
+			UE_LOGFMT(
+				LogRigidPhysics, Error, 
+				"InitTransform and UpdateTransform on Static or Kinematic Body '{0}' not supported on the Physics Thread. Transform would diverge from Game Thread.", 
+				GetProxyAPI()->DebugName().IsValid()? **GetProxyAPI()->DebugName() : TEXT("Unknown")
+				);
+			return;
+		}
+
 		if (FPBDRigidsSolver* Solver = Proxy->GetSolver<FPBDRigidsSolver>())
 		{
 			// NOTE: The Chaos Proxy API does not do everything required when changing
 			// properties on the particle. E.g., marking dirty, etc.
 			// Here'what should be in GetProxyAPI()->SetX()/SetR()
 			FGeometryParticleHandle* ParticleHandle = Proxy->GetHandle_LowLevel();
-			Solver->GetEvolution()->SetParticleTransform(ParticleHandle, InTransform.GetTranslation(), InTransform.GetRotation(), /*bIsTeleport*/false);
+			Solver->GetEvolution()->SetParticleTransform(ParticleHandle, InTransform.GetTranslation(), InTransform.GetRotation(), bIsTeleport);
 			Solver->GetEvolution()->DirtyParticle(*ParticleHandle);
 		}
 	}
 
 	void FRigidBodyAsyncPT::SetKinematicTarget(const FTransform& InTransform)
 	{
-		CHAOSRIGIDPHYSICSASYNC_CHECK(IsValid());
+		UE_RIGIDPHYSICS_CHECK(IsValid());
 		if (FPBDRigidsSolver* Solver = Proxy->GetSolver<FPBDRigidsSolver>())
 		{
 			FGeometryParticleHandle* ParticleHandle = Proxy->GetHandle_LowLevel();
@@ -287,28 +350,28 @@ namespace Chaos::Rigids::Async
 
 	FTransform FRigidBodyAsyncPT::GetTransform() const
 	{
-		CHAOSRIGIDPHYSICSASYNC_CHECK(IsValid());
+		UE_RIGIDPHYSICS_CHECK(IsValid());
 		FConstGenericParticleHandle ParticleHandle = Proxy->GetHandle_LowLevel();
 		return FTransform(ParticleHandle->Q(), ParticleHandle->P());
 	}
 
 	FBounds3d FRigidBodyAsyncPT::GetBounds() const
 	{
-		CHAOSRIGIDPHYSICSASYNC_CHECK(IsValid());
+		UE_RIGIDPHYSICS_CHECK(IsValid());
 		FConstGenericParticleHandle ParticleHandle = Proxy->GetHandle_LowLevel();
 		return FBounds3d(ParticleHandle->WorldSpaceInflatedBounds().Min(), ParticleHandle->WorldSpaceInflatedBounds().Max());
 	}
 
 	double FRigidBodyAsyncPT::GetMass() const
 	{
-		CHAOSRIGIDPHYSICSASYNC_CHECK(IsValid());
+		UE_RIGIDPHYSICS_CHECK(IsValid());
 		FConstGenericParticleHandle ParticleHandle = Proxy->GetHandle_LowLevel();
 		return ParticleHandle->M();
 	}
 
 	void FRigidBodyAsyncPT::SetMass(double InMass)
 	{
-		CHAOSRIGIDPHYSICSASYNC_CHECK(IsValid());
+		UE_RIGIDPHYSICS_CHECK(IsValid());
 		ensure(false);
 		// TODO_CHAOSAPI: Not yet supported on the PT because the mass will not make it back to GT.
 		//FGenericParticleHandle ParticleHandle = Proxy->GetHandle_LowLevel();
diff --git a/Engine/Plugins/Experimental/ChaosRigidPhysicsAsync/Source/ChaosRigidPhysicsAsync/Private/ChaosRigidPhysicsAsync/RigidBodyContainerAsync.cpp b/Engine/Plugins/Experimental/ChaosRigidPhysicsAsync/Source/ChaosRigidPhysicsAsync/Private/ChaosRigidPhysicsAsync/RigidBodyContainerAsync.cpp
--- a/Engine/Plugins/Experimental/ChaosRigidPhysicsAsync/Source/ChaosRigidPhysicsAsync/Private/ChaosRigidPhysicsAsync/RigidBodyContainerAsync.cpp
+++ b/Engine/Plugins/Experimental/ChaosRigidPhysicsAsync/Source/ChaosRigidPhysicsAsync/Private/ChaosRigidPhysicsAsync/RigidBodyContainerAsync.cpp
@@ -3,334 +3,142 @@
 
 #include "ChaosRigidPhysicsAsync/RigidBodyAsync.h"
 #include "ChaosRigidPhysicsAsync/RigidSceneAsync.h"
-#include "RigidPhysics/RigidAccessor.h"
-#include "RigidPhysics/RigidBody.h"
+#include "RigidPhysics/RigidScene.h"
 
 #if UE_RIGIDPHYSICS_API_ENABLED
 
 namespace Chaos::Rigids::Async
 {
-	UE::Physics::FRigidBodyContainerTypeId FRigidBodyContainerAsyncGT::StaticRigidBodyContainerTypeId()
-	{
-		using namespace UE::Physics;
-
-		static const TCHAR* StaticTypeName = TEXT("RigidBodyContainerAsyncGT");
-		static FRigidBodyContainerTypeId StaticTypeId = FRigidBodyContainerTypeId(StaticTypeName);
-		return StaticTypeId;
-	}
+	UE_RIGIDPHYSICS_RIGIDTYPED_IMPL(FRigidBodyPoolAsyncGT, UE::Physics::IRigidBodyContainer);
 
-	FRigidBodyContainerAsyncGT::FRigidBodyContainerAsyncGT(const UE::Physics::FRigidBodyContainerId& InContainerId)
-		: UE::Physics::IRigidBodyContainer(StaticRigidBodyContainerTypeId())
-		, Scene(nullptr)
-		, ContainerId(InContainerId)
-		, Solver(nullptr)
-		, BodyPool(1000, 0)
-		, Bodies()
-		, BodiesFreeIndices()
-		, NextBodyIdIndex(0)
+	FRigidBodyPoolAsyncGT::FRigidBodyPoolAsyncGT(FRigidSceneAsyncGT* InScene, const UE::Physics::FRigidDebugName& InName)
+		: FRigidBodyContainer()
+		, Name(InName)
+		, BodyContainerId()
+		, Scene(InScene)
+		, BodyPool(10)
 	{
 	}
 
-	FRigidBodyContainerAsyncGT::~FRigidBodyContainerAsyncGT()
+	FRigidBodyPoolAsyncGT::~FRigidBodyPoolAsyncGT()
 	{
-		check(BodyPool.GetNumAllocated() == 0);
-		check(BodiesFreeIndices.Num() == Bodies.Num());
+		UE_RIGIDPHYSICS_CHECK(GetBodyRegistry()->IsEmpty());
+		UE_RIGIDPHYSICS_CHECK(BodyPool.GetNumAllocated() == 0);
 	}
 
-	UE::Physics::IRigidScene* FRigidBodyContainerAsyncGT::GetScene() const
+	void FRigidBodyPoolAsyncGT::CheckIsLocked() const
 	{
-		return GetSceneAsync();
-	}
 
-	UE::Physics::FRigidBodyContainerId FRigidBodyContainerAsyncGT::GetContainerId() const
-	{
-		return ContainerId;
-	}
-
-	FRigidBodyAsyncGT* FRigidBodyContainerAsyncGT::GetRigidBody(const UE::Physics::FRigidBodyId& BodyId) const
-	{
-		// Extract the AsyncId data from the BodyId
-		FRigidBodyIdAsync AsyncId = FRigidBodyIdAsync(BodyId);
-		if (Bodies.IsValidIndex(AsyncId.GetBodyIndex()))
-		{
-			// If the Epochs match the BodyId is current and we can return the body
-			const FRigidBodyArrayEntry& BodyEntry = Bodies[AsyncId.GetBodyIndex()];
-			if (BodyEntry.Epoch == AsyncId.GetEpoch())
-			{
-				return BodyEntry.Body;
-			}
-		}
-
-		// If the Epochs don't match, the Body with BodyId was destroyed
-		return nullptr;
-	}
-
-	UE::Physics::FRigidBodyContainerHandle FRigidBodyContainerAsyncGT::GetContainerHandle() const
-	{
-		check(Scene != nullptr);
-		return UE::Physics::FRigidBodyContainerHandle(Scene->GetSceneHandle(), GetContainerId());
-	}
-
-	void FRigidBodyContainerAsyncGT::Init(FRigidSceneAsyncGT* InScene, FPBDRigidsSolver* InSolver)
-	{
-		Scene = InScene;
-		Solver = InSolver;
 	}
 
-	FRigidBodyAsyncGT* FRigidBodyContainerAsyncGT::CreateRigidBody(const UE::Physics::FRigidNamePtr& InName, UE::Physics::ERigidMovementType InMovementType)
+	UE::Physics::FRigidBodyContainerHandle FRigidBodyPoolAsyncGT::GetHandle() const
 	{
 		using namespace UE::Physics;
-
-		TUniquePtr<FGeometryParticle> Particle;
-		if (InMovementType == ERigidMovementType::Static)
-		{
-			Particle = FGeometryParticle::CreateParticle();
-		}
-		else
-		{
-			TUniquePtr<FPBDRigidParticle> RigidParticle = FPBDRigidParticle::CreateParticle();
-			if (InMovementType == ERigidMovementType::Kinematic)
-			{
-				RigidParticle->SetObjectState(EObjectStateType::Kinematic);
-			}
-			Particle = MoveTemp(RigidParticle);
-		}
-
-
-#if CHAOS_DEBUG_NAME
-		Particle->SetDebugName(InName);
-#endif
-
-		// Create a default shape for testing
-		FImplicitObjectPtr Geom = FImplicitObjectPtr(new FImplicitBox3(FVec3(-50), FVec3(50)));
-		FImplicitObjectPtr Union = new FImplicitObjectUnion(TArray<FImplicitObjectPtr>({ Geom }));
-
-		Particle->SetGeometry(Union);
-
-		return CreateRigidBodyForProxy(FSingleParticlePhysicsProxy::Create(MoveTemp(Particle)));
-	}
-
-	FRigidBodyAsyncGT* FRigidBodyContainerAsyncGT::CreateRigidBodyForProxy(FSingleParticlePhysicsProxy* InProxy)
-	{
-		// Store the proxy with a unique Id (unique within this container)
-		return AllocBody(InProxy);
-	}
-
-	void FRigidBodyContainerAsyncGT::DestroyRigidBody(const UE::Physics::FRigidBodyId& BodyId)
-	{
-		if (FRigidBodyAsyncGT* BodyGT = GetRigidBody(BodyId))
-		{
-			// TODO_CHAOSAPI: Register/Unregister asymmetry.
-			// NOTE: This is the correct place for removing the body from the pool, but see RemoveRigidBodyFromScene 
-			// which is where it usually happens atm. It still needs to be here as well to destroy particles that were 
-			// never added to the scene. If the particle is destroyed in RemoveRigidBodyFromScene,  then GetRigidBody() 
-			// will have returned null so there's no danger of duplicate destruction.
-			FreeBody(BodyGT);
-		}
-	}
-
-	void FRigidBodyContainerAsyncGT::AddRigidBodyToScene(const UE::Physics::FRigidBodyId& BodyId)
-	{
-		// This ensure will fire if you try to re-add a body to a scene after removing it. This should be supported
-		// but FPBDRigidsSolver::UnregisterObject destroys the proxy so we would have to recreate it and we would
-		// lose any changes the user may have made. For now just treat this as an error.
-		FRigidBodyAsyncGT* BodyGT = GetRigidBody(BodyId);
-		if (ensure(BodyGT != nullptr))
-		{
-			if (!BodyGT->IsInScene())
-			{
-				// NOTE: Handle GT and PT
-				Solver->RegisterObject(BodyGT->GetProxy());
-			}
-		}
+		return FRigidBodyContainerHandle(Scene->GetHandle(), BodyContainerId);
 	}
 
-	void FRigidBodyContainerAsyncGT::RemoveRigidBodyFromScene(const UE::Physics::FRigidBodyId& BodyId)
+	void FRigidBodyPoolAsyncGT::SetId(const UE::Physics::FRigidObjectId& InId)
 	{
-		if (FRigidBodyAsyncGT* BodyGT = GetRigidBody(BodyId))
-		{
-			if (BodyGT->IsInScene())
-			{
-				// NOTE: Handles GT and PT
-				Solver->UnregisterObject(BodyGT->GetProxy());
-			}
-
-			// TODO_CHAOSAPI: Register/Unregister asymmetry.
-			// This should be called from DestroyRigidBody but FPBDRigidsSolver::UnregisterObject 
-			// enqueues a command that destroys the proxy from the physics thread. That command may be executed
-			// before we call DestroyRigidBody which would leave the body with a dangling proxy pointer
-			// unless we also free the body here.
-			FreeBody(BodyGT);
-		}
+		BodyContainerId = InId;
 	}
 
-	FRigidSceneAsyncGT* FRigidBodyContainerAsyncGT::GetSceneAsync() const
+	FRigidSceneAsyncGT* FRigidBodyPoolAsyncGT::GetScene() const
 	{
 		return Scene;
 	}
 
-	void FRigidBodyContainerAsyncGT::OnRigidBodyTransformUpdated(FRigidBodyAsyncGT* Body)
+	int32 FRigidBodyPoolAsyncGT::GetNumBodies() const
 	{
-		if (Scene != nullptr)
-		{
-			Scene->OnRigidBodyTransformUpdated(Body);
-		}
+		CheckIsLocked();
+
+		return GetBodyRegistry()->GetNumObjects();
 	}
 
-	FRigidBodyAsyncGT* FRigidBodyContainerAsyncGT::AllocBody(FSingleParticlePhysicsProxy* InProxy)
+	FRigidBodyAsyncGT* FRigidBodyPoolAsyncGT::CreateBody(const UE::Physics::FRigidDebugName& InName, UE::Physics::ERigidMovementType InMovementType)
 	{
-		// Make sure we have free array entries
-		if (BodiesFreeIndices.IsEmpty())
-		{
-			ReserveBodyArray(Bodies.Num() + 1000);
-		}
-
-		// Assign an array index
-		int32 BodyIndex = BodiesFreeIndices.Pop(EAllowShrinking::No);
-		FRigidBodyArrayEntry& BodyEntry = Bodies[BodyIndex];
-		check(BodyEntry.Body == nullptr);
-
-		// Build the Id from the BodyIndex and Epoch
-		FRigidBodyIdAsync AsyncId = FRigidBodyIdAsync(BodyIndex, BodyEntry.Epoch);
+		using namespace UE::Physics;
 
-		// Allocate the body from the pool
-		FRigidBodyAsyncGT* Body = BodyPool.Alloc(this, InProxy, AsyncId);
+		FRigidBodyAsyncGT* Body = BodyPool.Alloc(this, InName, InMovementType);
 
-		// Store the body in the ID->Body Map
-		Bodies[BodyIndex].Body = Body;
+		FRigidObjectId BodyId = GetBodyRegistry()->AddObject(Body);
 
+		Body->SetId(BodyId);
+		
 		return Body;
 	}
 
-	void FRigidBodyContainerAsyncGT::FreeBody(FRigidBodyAsyncGT* BodyGT)
+	void FRigidBodyPoolAsyncGT::DestroyBody(UE::Physics::IRigidBody* InBody)
 	{
-		check(BodyGT != nullptr);
-
-		// Get the entry in the Id->Body map
-		FRigidBodyIdAsync AsyncId = BodyGT->GetBodyIdAsync();
-		FRigidBodyArrayEntry& BodyEntry = Bodies[AsyncId.GetBodyIndex()];
-
-		// When FreeBody is called we must know we have a body from this container
-		check(BodyEntry.Epoch == AsyncId.GetEpoch());
-		check(BodyEntry.Body == BodyGT);
-
-		// Reset the entry and increment the epoch to invalidate outstanding handles
-		BodyEntry.Body = nullptr;
-		BodyEntry.Epoch = BodyEntry.Epoch + 1;
+		using namespace UE::Physics;
+		CheckIsLocked();
+		UE_RIGIDPHYSICS_CHECK(InBody != nullptr);
+		UE_RIGIDPHYSICS_CHECK(GetBodyRegistry()->FindObject(InBody->GetHandle().GetId()) == InBody);
 
-		// Add the array item back into the free pool
-		BodiesFreeIndices.Push(AsyncId.GetBodyIndex());
+		InBody->Deactivate();
 
-		// Return the body to the pool
-		BodyPool.Free(BodyGT);
-	}
+		GetBodyRegistry()->RemoveObject(InBody->GetHandle().GetId());
 
-	void FRigidBodyContainerAsyncGT::ReserveBodyArray(int32 MaxBodies)
-	{
-		if (MaxBodies > Bodies.Num())
+		if (FRigidBodyAsyncGT* Body = InBody->AsAChecked<FRigidBodyAsyncGT>())
 		{
-			int32 PrevNumBodies = Bodies.Num();
-			Bodies.SetNum(MaxBodies);
-			BodiesFreeIndices.Reserve(MaxBodies);
-
-			// Loop backawards so we allocate lowest indices first by popping from BodiesFreeIndices
-			for (int32 BodyIndex = MaxBodies - 1; BodyIndex >= PrevNumBodies; --BodyIndex)
-			{
-				Bodies[BodyIndex].Body = nullptr;
-				Bodies[BodyIndex].Epoch = 0;
-				BodiesFreeIndices.Push(BodyIndex);
-			}
+			BodyPool.Free(Body);
 		}
 	}
-}
 
+}
 
 namespace Chaos::Rigids::Async
 {
-	UE::Physics::FRigidBodyContainerTypeId FRigidBodyContainerAsyncPT::StaticRigidBodyContainerTypeId()
-	{
-		using namespace UE::Physics;
+	UE_RIGIDPHYSICS_RIGIDTYPED_IMPL(FRigidBodyPoolAsyncPT, UE::Physics::IRigidBodyContainer);
 
-		static const TCHAR* StaticTypeName = TEXT("RigidBodyContainerAsyncPT");
-		static FRigidBodyContainerTypeId StaticTypeId = FRigidBodyContainerTypeId(StaticTypeName);
-		return StaticTypeId;
-	}
-
-	FRigidBodyContainerAsyncPT::FRigidBodyContainerAsyncPT(const UE::Physics::FRigidBodyContainerId& InContainerId)
-		: IRigidBodyContainer(StaticRigidBodyContainerTypeId())
-		, Scene(nullptr)
-		, ContainerId(InContainerId)
-		, Solver(nullptr)
+	FRigidBodyPoolAsyncPT::FRigidBodyPoolAsyncPT(const UE::Physics::FRigidDebugName& InName, const UE::Physics::FRigidBodyContainerHandle& InContainerHandle)
+		: IRigidBodyContainer()
+		, Handle(InContainerHandle)
 		, BodyPool(1000, 0)
 		, IdToBodyMap()
 	{
 	}
 
-	FRigidBodyContainerAsyncPT::~FRigidBodyContainerAsyncPT()
+	FRigidBodyPoolAsyncPT::~FRigidBodyPoolAsyncPT()
 	{
 		check(BodyPool.GetNumAllocated() == 0);
 		check(IdToBodyMap.IsEmpty());
 	}
 
-	UE::Physics::IRigidScene* FRigidBodyContainerAsyncPT::GetScene() const
+	UE::Physics::FRigidBodyContainerHandle FRigidBodyPoolAsyncPT::GetHandle() const
 	{
-		return Scene;
+		return Handle;
 	}
 
-	UE::Physics::FRigidBodyContainerId FRigidBodyContainerAsyncPT::GetContainerId() const
+	int32 FRigidBodyPoolAsyncPT::GetNumBodies() const
 	{
-		return ContainerId;
+		return IdToBodyMap.Num();
 	}
 
-	FRigidBodyAsyncPT* FRigidBodyContainerAsyncPT::GetRigidBody(const UE::Physics::FRigidBodyId& BodyId) const
+	FRigidBodyAsyncPT* FRigidBodyPoolAsyncPT::GetBody(const UE::Physics::FRigidObjectId& InBodyId) const
 	{
-		const FRigidBodyAsyncPT* const* PBody = IdToBodyMap.Find(BodyId);
-		if (PBody != nullptr)
+		if (FRigidBodyAsyncPT*const* PBody = IdToBodyMap.Find(InBodyId))
 		{
-			return const_cast<FRigidBodyAsyncPT*>(*PBody);
+			return *PBody;
 		}
 		return nullptr;
 	}
 
-	void FRigidBodyContainerAsyncPT::Init(FRigidSceneAsyncPT* InScene, FPBDRigidsSolver* InSolver)
-	{
-		Scene = InScene;
-		Solver = InSolver;
-	}
-
-	UE::Physics::FRigidBodyContainerHandle FRigidBodyContainerAsyncPT::GetContainerHandle() const
+	void FRigidBodyPoolAsyncPT::CreateBodyPT(const UE::Physics::FRigidBodyHandle& InBodyHandle, FSingleParticlePhysicsProxy* InProxy)
 	{
-		check(Scene != nullptr);
-		return UE::Physics::FRigidBodyContainerHandle(Scene->GetSceneHandle(), GetContainerId());
-	}
+		UE_RIGIDPHYSICS_CHECK(GetBody(InBodyHandle.GetId()) == nullptr);
 
-	void FRigidBodyContainerAsyncPT::CreateRigidBody(const FRigidBodyIdAsync& InBodyId, FSingleParticlePhysicsProxy* InProxy)
-	{
-		FRigidBodyAsyncPT* Body = BodyPool.Alloc(this, InProxy, InBodyId);
-		IdToBodyMap.Add(InBodyId.AsRigidBodyId(), Body);
+		FRigidBodyAsyncPT* Body = BodyPool.Alloc(InBodyHandle, InProxy);
+		IdToBodyMap.Add(InBodyHandle.GetId(), Body);
 	}
 
-	void FRigidBodyContainerAsyncPT::DestroyRigidBody(const UE::Physics::FRigidBodyId& BodyId)
+	void FRigidBodyPoolAsyncPT::DestroyBodyPT(const UE::Physics::FRigidBodyHandle& InBodyHandle)
 	{
-		if (FRigidBodyAsyncPT* Body = GetRigidBody(BodyId))
+		if (FRigidBodyAsyncPT* Body = GetBody(InBodyHandle.GetId()))
 		{
-			IdToBodyMap.Remove(BodyId);
+			IdToBodyMap.Remove(InBodyHandle.GetId());
 			BodyPool.Free(Body);
 		}
 	}
-
-	void FRigidBodyContainerAsyncPT::AddRigidBodyToScene(const UE::Physics::FRigidBodyId& BodyId)
-	{
-		// TODO_CHAOSAPI: nothing to do for now because RegisterObject does GT and PT
-		// See FRigidBodyContainerAsyncGT::AddRigidBodyToScene
-	}
-
-	void FRigidBodyContainerAsyncPT::RemoveRigidBodyFromScene(const UE::Physics::FRigidBodyId& BodyId)
-	{
-		// TODO_CHAOSAPI: nothing to do for now because UnregisterObject does GT and PT
-		// See FRigidBodyContainerAsyncGT::RemoveRigidBodyFromScene
-	}
-
 }
 
 #endif // UE_RIGIDPHYSICS_API_ENABLED
\ No newline at end of file
diff --git a/Engine/Plugins/Experimental/ChaosRigidPhysicsAsync/Source/ChaosRigidPhysicsAsync/Private/ChaosRigidPhysicsAsync/RigidFactoryAsync.cpp b/Engine/Plugins/Experimental/ChaosRigidPhysicsAsync/Source/ChaosRigidPhysicsAsync/Private/ChaosRigidPhysicsAsync/RigidFactoryAsync.cpp
--- a/Engine/Plugins/Experimental/ChaosRigidPhysicsAsync/Source/ChaosRigidPhysicsAsync/Private/ChaosRigidPhysicsAsync/RigidFactoryAsync.cpp
+++ b/Engine/Plugins/Experimental/ChaosRigidPhysicsAsync/Source/ChaosRigidPhysicsAsync/Private/ChaosRigidPhysicsAsync/RigidFactoryAsync.cpp
@@ -4,113 +4,67 @@
 #include "ChaosRigidPhysicsAsync/RigidBodyAsync.h"
 #include "ChaosRigidPhysicsAsync/RigidBodyContainerAsync.h"
 #include "ChaosRigidPhysicsAsync/RigidSceneAsync.h"
-#include "RigidPhysics/RigidAccessor.h"
 #include "RigidPhysics/RigidBody.h"
 
 #if UE_RIGIDPHYSICS_API_ENABLED
 
 namespace Chaos::Rigids::Async
 {
-	UE::Physics::FRigidSceneHandle FRigidFactoryAsync::CreateRigidScene(const FRigidSceneSettingsAsync& InSettings)
+	FRigidFactoryAsync::FRigidFactoryAsync()
 	{
-		// Create a scene and store a hard reference to prevent its destruction
-		TSharedPtr<FRigidSceneAsyncGT, ESPMode::ThreadSafe> Scene = MakeShared<FRigidSceneAsyncGT>(InSettings);
-		Scenes.Add(Scene);
 
-		Scene->Startup();
-
-		// Return a weak/unlocked reference to the scene
-		return UE::Physics::FRigidSceneHandle(Scene);
 	}
 
-	UE::Physics::FRigidSceneHandle FRigidFactoryAsync::CreateRigidSceneFromSolver(FPBDRigidsSolver* InSolver)
+	FRigidFactoryAsync::~FRigidFactoryAsync()
 	{
-		// Create a scene and store a hard reference to prevent its destruction
-		TSharedPtr<FRigidSceneAsyncGT, ESPMode::ThreadSafe> Scene = MakeShared<FRigidSceneAsyncGT>(InSolver);
-		Scenes.Add(Scene);
-
-		Scene->Startup();
-
-		// Return a weak/unlocked reference to the scene
-		return UE::Physics::FRigidSceneHandle(Scene);
+		// Make sure all scenes were destroyed
+		UE_RIGIDPHYSICS_CHECK(Scenes.IsEmpty());
 	}
 
-	void FRigidFactoryAsync::DestroyRigidScene(const UE::Physics::FRigidSceneHandle& InSceneHandle)
+	UE::Physics::FRigidSceneHandle FRigidFactoryAsync::CreateScene(const UE::Physics::FRigidDebugName& InName)
 	{
-		FRigidSceneAsyncGT* Scene = static_cast<FRigidSceneAsyncGT*>(Private::FRigidAccessor::GetIRigidScene(InSceneHandle.LockRW()));
+		using namespace UE::Physics;
 
-		Scene->Shutdown();
+		FRigidSceneSettingsAsync Settings;
 
-		Scenes.RemoveAll([Scene](const UE::Physics::FRigidSceneSharedPtr& Item) { return !Item.IsValid() || (Item.Get() == Scene); });
-	}
-}
+		// Create a scene and store a hard reference to prevent its destruction
+		FRigidSceneAsyncGT* Scene = new FRigidSceneAsyncGT(InName, Settings);
+		Scenes.Add(Scene);
 
-namespace Chaos::Rigids::Async::Private
-{
-	class FRigidAccessorAsync
-	{
-	public:
-		static UE::Physics::FRigidSceneHandle CreateSceneFromSolver(FRigidFactoryAsync& InFactory, FPBDRigidsSolver* InSolver)
-		{
-			return InFactory.CreateRigidSceneFromSolver(InSolver);
-		}
-	};
+		Scene->Startup();
 
-	FRigidSceneAsyncGT* GetAsyncScene(const UE::Physics::TRigidScenePtr<UE::Physics::FGameContextRW>& InScenePtr)
-	{
-		return static_cast<FRigidSceneAsyncGT*>(Private::FRigidAccessor::GetIRigidScene(InScenePtr));
+		return FRigidSceneHandle(Scene);
 	}
 
-	FRigidBodyContainerAsyncGT* GetAsyncBodyContainer(const UE::Physics::TRigidScenePtr<UE::Physics::FGameContextRW>& InScenePtr)
+	UE::Physics::IRigidScene* FRigidFactoryAsync::CreateSceneFromSolver(FPBDRigidsSolver* InSolver)
 	{
-		if (FRigidSceneAsyncGT* AsyncScene = GetAsyncScene(InScenePtr))
-		{
-			return static_cast<FRigidBodyContainerAsyncGT*>(AsyncScene->GetRigidBodyContainer(UE::Physics::FRigidBodyContainerId(0)));
-		}
-		return nullptr;
-	}
+		using namespace UE::Physics;
 
-	FRigidBodyAsyncGT* GetAsyncBody(const UE::Physics::TRigidBodyPtr<UE::Physics::FGameContextRO>& InBodyPtr)
-	{
-		return static_cast<FRigidBodyAsyncGT*>(Private::FRigidAccessor::GetIRigidBody(InBodyPtr));
-	}
+		FRigidSceneAsyncGT* Scene = new FRigidSceneAsyncGT(FRigidDebugName::Make(*InSolver->GetDebugName().ToString()), InSolver);
+		Scenes.Add(Scene);
 
-	UE::Physics::FRigidSceneHandle CreateSceneFromSolver(FRigidFactoryAsync& InFactory, FPBDRigidsSolver* InSolver)
-	{
-		check(InSolver != nullptr);
-		return FRigidAccessorAsync::CreateSceneFromSolver(InFactory, InSolver);
+		Scene->Startup();
+
+		return Scene;
 	}
 
-	UE::Physics::FRigidBodyHandle CreateBodyFromProxy(const UE::Physics::FRigidSceneHandle& InSceneHandle, FSingleParticlePhysicsProxy* InProxy)
+	void FRigidFactoryAsync::DestroyScene(const UE::Physics::FRigidSceneHandle& InSceneHandle)
 	{
 		using namespace UE::Physics;
-		check(InProxy != nullptr);
 
-		if (TRigidScenePtr<FGameContextRW> ScenePtr = InSceneHandle.LockRW())
+		if (IRigidScene* Scene = InSceneHandle.Get())
 		{
-			if (FRigidBodyContainerAsyncGT* BodyContainer = GetAsyncBodyContainer(ScenePtr))
+			if (FRigidSceneAsyncGT* SceneGT = Scene->AsA<FRigidSceneAsyncGT>())
 			{
-				if (FRigidBodyAsyncGT* Body = BodyContainer->CreateRigidBodyForProxy(InProxy))
-				{
-					return FRigidBodyHandle(FRigidBodyContainerHandle(InSceneHandle, BodyContainer->GetContainerId()), Body->GetBodyId());
-				}
-			}
-		}
-		return FRigidBodyHandle();
-	}
+				check(Scenes.Contains(SceneGT));
+			
+				SceneGT->Shutdown();
 
-	FSingleParticlePhysicsProxy* GetProxyFromBody(const UE::Physics::FRigidBodyHandle& InBodyHandle)
-	{
-		using namespace UE::Physics;
+				Scenes.Remove(SceneGT);
 
-		if (TRigidBodyPtr<FGameContextRO> BodyPtr = InBodyHandle.LockRO())
-		{
-			if (FRigidBodyAsyncGT* Body = GetAsyncBody(BodyPtr))
-			{
-				return Body->GetProxy();
+				delete SceneGT;
 			}
 		}
-		return nullptr;
 	}
 }
 
diff --git a/Engine/Plugins/Experimental/ChaosRigidPhysicsAsync/Source/ChaosRigidPhysicsAsync/Private/ChaosRigidPhysicsAsync/RigidSceneAsync.cpp b/Engine/Plugins/Experimental/ChaosRigidPhysicsAsync/Source/ChaosRigidPhysicsAsync/Private/ChaosRigidPhysicsAsync/RigidSceneAsync.cpp
--- a/Engine/Plugins/Experimental/ChaosRigidPhysicsAsync/Source/ChaosRigidPhysicsAsync/Private/ChaosRigidPhysicsAsync/RigidSceneAsync.cpp
+++ b/Engine/Plugins/Experimental/ChaosRigidPhysicsAsync/Source/ChaosRigidPhysicsAsync/Private/ChaosRigidPhysicsAsync/RigidSceneAsync.cpp
@@ -1,62 +1,88 @@
 // Copyright Epic Games, Inc. All Rights Reserved.
 #include "ChaosRigidPhysicsAsync/RigidSceneAsync.h"
+#include "RigidPhysics/RigidContext.h"
 #include "RigidPhysics/RigidModifier.h"
 
 #if UE_RIGIDPHYSICS_API_ENABLED
 
 namespace Chaos::Rigids::Async
 {
-	FRigidSceneAsyncGT::FRigidSceneAsyncGT(const FRigidSceneSettingsAsync& InSceneSettings)
-		: SceneSettings(InSceneSettings)
+	FRigidSceneAsyncGT::FRigidSceneAsyncGT(
+		const UE::Physics::FRigidDebugName& InName,
+		const FRigidSceneSettingsAsync& InSceneSettings)
+		: UE::Physics::FRigidScene()
+		, Name(InName)
+		, SceneSettings(InSceneSettings)
 		, Solver(nullptr)
 		, SpatialAcceleration(nullptr)
-		, SingleBodyContainer(UE::Physics::FRigidBodyContainerId(0))
+		, BodyPool()
 		, SceneModifiers()
 		, ScenePT()
 		, bOwnsSolver(false)
 	{
-		ScenePT = MakeUnique<FRigidSceneAsyncPT>(UE::Physics::FRigidBodyContainerId(0));
+		ConstructScene();
 	}
 
-	FRigidSceneAsyncGT::FRigidSceneAsyncGT(FPBDRigidsSolver* InSolver)
-		: SceneSettings()
+	FRigidSceneAsyncGT::FRigidSceneAsyncGT(
+		const UE::Physics::FRigidDebugName& InName,
+		FPBDRigidsSolver* InSolver)
+		: UE::Physics::FRigidScene()
+		, Name(InName)
+		, SceneSettings()
 		, Solver(InSolver)
 		, SpatialAcceleration(nullptr)
-		, SingleBodyContainer(UE::Physics::FRigidBodyContainerId(0))
+		, BodyPool()
 		, SceneModifiers()
 		, ScenePT()
 		, bOwnsSolver(false)
 	{
-		ScenePT = MakeUnique<FRigidSceneAsyncPT>(UE::Physics::FRigidBodyContainerId(0));
+		ConstructScene();
 	}
 
 	FRigidSceneAsyncGT::~FRigidSceneAsyncGT()
 	{
 		check(Solver == nullptr);
+
+		DestructScene();
 	}
 
-	UE::Physics::FRigidSceneHandle FRigidSceneAsyncGT::GetSceneHandle() const
+	void FRigidSceneAsyncGT::ConstructScene()
 	{
-		return SceneHandle;
+		using namespace UE::Physics;
+
+		BodyPool = MakeUnique<FRigidBodyPoolAsyncGT>(this, Name);
+
+		FRigidObjectId BodyPoolId = GetBodyContainerRegistry()->AddObject(BodyPool.Get());
+		BodyPool->SetId(BodyPoolId);
+
+		ScenePT = MakeUnique<FRigidSceneAsyncPT>(Name, BodyPool->GetHandle());
 	}
 
-	void FRigidSceneAsyncGT::Startup()
+	void FRigidSceneAsyncGT::DestructScene()
 	{
 		using namespace UE::Physics;
 
-		SceneHandle = FRigidSceneHandle(AsShared());
+		ScenePT.Reset();
+
+		GetBodyContainerRegistry()->RemoveObject(BodyPool->GetHandle().GetId());
+		BodyPool->SetId(FRigidObjectId());
+
+		BodyPool.Reset();
+	}
+
+	void FRigidSceneAsyncGT::Startup()
+	{
+		using namespace UE::Physics;
 
 		if (Solver == nullptr)
 		{
 			// NOTE: Not using full Async for now - it works but it complicates testing (time shift and interpolation)
 			FChaosSolversModule* Module = FChaosSolversModule::GetModule();
-			Solver = Module->CreateSolver(/*Owner=*/nullptr, /*AsyncDt*/-1, EThreadingMode::TaskGraph, SceneSettings.SceneName);
+			Solver = Module->CreateSolver(/*Owner=*/nullptr, /*AsyncDt*/-1, EThreadingMode::TaskGraph, FName(*Name.Value()));
 			bOwnsSolver = true;
 		}
 
-		SingleBodyContainer.Init(this, Solver);
-
-		ScenePT->Startup(SceneHandle, Solver);
+		ScenePT->Startup(GetHandle(), Solver);
 	}
 
 	void FRigidSceneAsyncGT::Shutdown()
@@ -82,120 +108,115 @@ namespace Chaos::Rigids::Async
 		}
 	}
 
-	void FRigidSceneAsyncGT::Lock(UE::Physics::EContextLockType InLockType)
+	UE::Physics::FRigidSceneHandle FRigidSceneAsyncGT::GetHandle() const
+	{
+		return UE::Physics::FRigidSceneHandle(const_cast<FRigidSceneAsyncGT*>(this));
+	}
+
+	void FRigidSceneAsyncGT::Lock(UE::Physics::ERigidLockType InLockType) const
 	{
 		using namespace UE::Physics;
 
-		CHAOSRIGIDPHYSICSASYNC_CHECK(IsInGameThreadContext());
+		UE_RIGIDPHYSICS_CHECK(IsInGameThreadContext());
 
 		if (Solver != nullptr)
 		{
 			switch (InLockType)
 			{
-			case EContextLockType::ReadOnly:
+			case ERigidLockType::ReadOnly:
 				Solver->GetExternalDataLock_External().ReadLock();
 				break;
-			case EContextLockType::ReadWrite:
+			case ERigidLockType::ReadWrite:
 				Solver->GetExternalDataLock_External().WriteLock();
 				break;
 			}
 		}
 	}
 
-	void FRigidSceneAsyncGT::Unlock(UE::Physics::EContextLockType InLockType)
+	void FRigidSceneAsyncGT::Unlock(UE::Physics::ERigidLockType InLockType) const
 	{
 		using namespace UE::Physics;
 
-		CHAOSRIGIDPHYSICSASYNC_CHECK(IsInGameThreadContext());
+		UE_RIGIDPHYSICS_CHECK(IsInGameThreadContext());
 
 		if (Solver != nullptr)
 		{
 			switch (InLockType)
 			{
-			case EContextLockType::ReadOnly:
+			case ERigidLockType::ReadOnly:
 				Solver->GetExternalDataLock_External().ReadUnlock();
 				break;
-			case EContextLockType::ReadWrite:
+			case ERigidLockType::ReadWrite:
 				Solver->GetExternalDataLock_External().WriteUnlock();
 				break;
 			}
 		}
 	}
 
-	FRigidBodyAsyncGT* FRigidSceneAsyncGT::CreateRigidBody(const UE::Physics::FRigidNamePtr& InName, UE::Physics::ERigidMovementType InMovementType)
-	{
-		CheckGameWriteLocked();
-
-		FRigidBodyAsyncGT* BodyGT = SingleBodyContainer.CreateRigidBody(InName, InMovementType);
-
-		Solver->EnqueueCommandImmediate(
-			[this, BodyId = BodyGT->GetBodyIdAsync(), Proxy = BodyGT->Proxy]()
-			{
-				ScenePT->CreateRigidBody(BodyId, Proxy);
-			});
-
-		return BodyGT;
-	}
-
-	void FRigidSceneAsyncGT::DestroyRigidBody(const UE::Physics::FRigidBodyHandle& BodyHandle)
+	void FRigidSceneAsyncGT::ActivateBody(FRigidBodyAsyncGT* InBody)
 	{
-		CheckGameWriteLocked();
-
-		if (BodyHandle.GetContainerHandle().GetContainerId() == SingleBodyContainer.GetContainerId())
+		if (!InBody->IsActive())
 		{
-			UE::Physics::FRigidBodyId BodyId = BodyHandle.GetBodyId();
+			UE_RIGIDPHYSICS_CHECK(InBody->GetProxy() != nullptr);
 
-			SingleBodyContainer.DestroyRigidBody(BodyId);
+			// NOTE: Solver handles GT and PT Particles but we need to create a wrapper on the PT.
+			Solver->RegisterObject(InBody->GetProxy());
 
+			// NOTE: Currently this does nothing because Solver->RegisterObject will activate the particle
+			// via the dirty proxy mechanism.
 			Solver->EnqueueCommandImmediate(
-				[BodyId, this]()
+				[BodyHandle = InBody->GetHandle(), this]()
 				{
-					ScenePT->DestroyRigidBody(BodyId);
+					ScenePT->ActivateBody(BodyHandle);
 				});
 		}
 	}
 
-	void FRigidSceneAsyncGT::RegisterRigidBody(const UE::Physics::FRigidBodyHandle& BodyHandle)
+	void FRigidSceneAsyncGT::DeactivateBody(FRigidBodyAsyncGT* InBody)
 	{
-		CheckGameWriteLocked();
-
-		if (BodyHandle.GetContainerHandle().GetContainerId() == SingleBodyContainer.GetContainerId())
+		if (InBody->IsActive())
 		{
-			SingleBodyContainer.AddRigidBodyToScene(BodyHandle.GetBodyId());
+			UE_RIGIDPHYSICS_CHECK(InBody->GetProxy() != nullptr);
 
 			Solver->EnqueueCommandImmediate(
-				[BodyId = BodyHandle.GetBodyId(), this]()
+				[BodyHandle = InBody->GetHandle(), this]()
 				{
-					ScenePT->RegisterRigidBody(BodyId);
+					ScenePT->DeactivateBody(BodyHandle);
 				});
+
+			// NOTE: Solver handles GT and PT
+			Solver->UnregisterObject(InBody->GetProxy());
 		}
 	}
 
-	void FRigidSceneAsyncGT::UnregisterRigidBody(const UE::Physics::FRigidBodyHandle& BodyHandle)
+	FRigidBodyAsyncGT* FRigidSceneAsyncGT::CreateBody(const UE::Physics::FRigidDebugName& InName, UE::Physics::ERigidMovementType InMovementType)
 	{
+		using namespace UE::Physics;
 		CheckGameWriteLocked();
 
-		if (BodyHandle.GetContainerHandle().GetContainerId() == SingleBodyContainer.GetContainerId())
-		{
-			SingleBodyContainer.RemoveRigidBodyFromScene(BodyHandle.GetBodyId());
+		FRigidBodyAsyncGT* BodyGT = BodyPool->CreateBody(InName, InMovementType);
 
-			Solver->EnqueueCommandImmediate(
-				[BodyId = BodyHandle.GetBodyId(), this]()
-				{
-					ScenePT->UnregisterRigidBody(BodyId);
-				});
-		}
+		Solver->EnqueueCommandImmediate(
+			[this, BodyHandle = BodyGT->GetHandle(), Proxy = BodyGT->GetProxy()]()
+			{
+				ScenePT->CreateBodyPT(BodyHandle, Proxy);
+			});
+
+		return BodyGT;
 	}
 
-	FRigidBodyContainerAsyncGT* FRigidSceneAsyncGT::GetRigidBodyContainer(const UE::Physics::FRigidBodyContainerId& InBodyContainerId)
+	void FRigidSceneAsyncGT::DestroyBody(UE::Physics::IRigidBody* InBody)
 	{
-		CheckGameReadLocked();
+		using namespace UE::Physics;
+		CheckGameWriteLocked();
 
-		if (InBodyContainerId == SingleBodyContainer.GetContainerId())
-		{
-			return &SingleBodyContainer;
-		}
-		return nullptr;
+		Solver->EnqueueCommandImmediate(
+			[this, BodyHandle = InBody->GetHandle()]()
+			{
+				ScenePT->DestroyBodyPT(BodyHandle);
+			});
+
+		BodyPool->DestroyBody(InBody);
 	}
 
 	void FRigidSceneAsyncGT::RegisterModifier(UE::Physics::IRigidSceneModifier* Modifier)
@@ -248,11 +269,6 @@ namespace Chaos::Rigids::Async
 		}
 	}
 
-	const char* FRigidSceneAsyncGT::GetTypeName() const
-	{ 
-		return "FRigidSceneAsyncGT";
-	}
-
 	void FRigidSceneAsyncGT::OnRigidBodyTransformUpdated(FRigidBodyAsyncGT* Body)
 	{
 		UpdateRigidBodyInAccelerationStructure(Body);
@@ -266,35 +282,31 @@ namespace Chaos::Rigids::Async
 			return;
 		}
 
-		if (Body->IsInScene())
+		if (Body->GetProxy() != nullptr)
 		{
-			FGeometryParticle* Particle = Body->GetProxy()->GetParticle_LowLevel();
-
-			if (SpatialAcceleration != nullptr)
+			if (FGeometryParticle* Particle = Body->GetProxy()->GetParticle_LowLevel())
 			{
-				const FBounds3d Bounds = Body->GetBounds();
-				const FAABB3 WorldBounds = FAABB3(Bounds.Min, Bounds.Max);
-				const bool bHasBounds = !WorldBounds.IsFull();
+				if (SpatialAcceleration != nullptr)
+				{
+					const FBounds3d Bounds = Body->GetBounds();
+					const FAABB3 WorldBounds = FAABB3(Bounds.Min, Bounds.Max);
+					const bool bHasBounds = !WorldBounds.IsFull();
 
-				FAccelerationStructureHandle AccelerationHandle(Particle);
-				SpatialAcceleration->UpdateElementIn(AccelerationHandle, WorldBounds, bHasBounds, Body->GetProxyAPI().SpatialIdx());
-			}
+					FAccelerationStructureHandle AccelerationHandle(Particle);
+					SpatialAcceleration->UpdateElementIn(AccelerationHandle, WorldBounds, bHasBounds, Body->GetProxyAPI().SpatialIdx());
+				}
 
-			if (Solver != nullptr)
-			{
-				Solver->UpdateParticleInAccelerationStructure_External(Particle, EPendingSpatialDataOperation::Update);
+				if (Solver != nullptr)
+				{
+					Solver->UpdateParticleInAccelerationStructure_External(Particle, EPendingSpatialDataOperation::Update);
+				}
 			}
 		}
 	}
 
-	UE::Physics::FRigidSceneHandle FRigidSceneAsyncGT::GetHandle()
-	{
-		return UE::Physics::FRigidSceneHandle(AsShared());
-	}
-
 	void FRigidSceneAsyncGT::CheckGameReadLocked() const
 	{
-		CHAOSRIGIDPHYSICSASYNC_CHECK(IsInGameThreadContext());
+		UE_RIGIDPHYSICS_CHECK(IsInGameThreadContext());
 
 		// TODO_CHAOSAPI: Check read or write lock
 		//if (Solver != nullptr)
@@ -305,7 +317,7 @@ namespace Chaos::Rigids::Async
 
 	void FRigidSceneAsyncGT::CheckGameWriteLocked() const
 	{
-		CHAOSRIGIDPHYSICSASYNC_CHECK(IsInGameThreadContext());
+		UE_RIGIDPHYSICS_CHECK(IsInGameThreadContext());
 
 		// TODO_CHAOSAPI: Check write lock
 		//if (Solver != nullptr)
@@ -320,12 +332,10 @@ namespace Chaos::Rigids::Async
 
 		CheckGameWriteLocked();
 
-		FGameContextRW Context = FGameContextRW(EContextGamePhase::PreSimulate);
-		TRigidScenePtr<FGameContextRW> ScenePtr = UE::Physics::Private::FRigidAccessor::MakeRigidScenePtr(Context, this);
-
+		FRigidContextGameRW Context = FRigidContextGameRW(this);
 		for (IRigidSceneModifier* SceneModifier : SceneModifiers)
 		{
-			SceneModifier->PreSimulate(ScenePtr);
+			SceneModifier->PreSimulate(Context);
 		}
 	}
 
@@ -335,12 +345,10 @@ namespace Chaos::Rigids::Async
 
 		CheckGameWriteLocked();
 
-		FGameContextRW Context = FGameContextRW(EContextGamePhase::PostSimulate);
-		TRigidScenePtr<FGameContextRW> ScenePtr = UE::Physics::Private::FRigidAccessor::MakeRigidScenePtr(Context, this);
-
+		FRigidContextGameRW Context = FRigidContextGameRW(this);
 		for (IRigidSceneModifier* SceneModifier : SceneModifiers)
 		{
-			SceneModifier->PostSimulate(ScenePtr);
+			SceneModifier->PostSimulate(Context);
 		}
 	}
 }
@@ -348,63 +356,25 @@ namespace Chaos::Rigids::Async
 
 namespace Chaos::Rigids::Async
 {
-	FRigidSceneAsyncPT::FRigidSceneAsyncPT(const UE::Physics::FRigidBodyContainerId& InSingleBodyContainerId)
-		: Solver(nullptr)
-		, SingleBodyContainer(InSingleBodyContainerId)
+	FRigidSceneAsyncPT::FRigidSceneAsyncPT(const UE::Physics::FRigidDebugName& InName, const UE::Physics::FRigidBodyContainerHandle& InBodyPoolHandle)
+		: Name(InName)
+		, Solver(nullptr)
 		, SceneModifiers()
-		, Context()
 	{
+		BodyPool = MakeUnique<FRigidBodyPoolAsyncPT>(InName, InBodyPoolHandle);
 	}
 
-	void FRigidSceneAsyncPT::Lock(UE::Physics::EContextLockType InLockType)
+	void FRigidSceneAsyncPT::Lock(UE::Physics::ERigidLockType InLockType) const
 	{
-		CheckSimContext();
-		// Locking not required on the physics thread.
-		// All parallel access patterns are controlled by the physics thread itself.
-		// No callbacks are called in parallel (yet).
 	}
 
-	void FRigidSceneAsyncPT::Unlock(UE::Physics::EContextLockType InLockType)
+	void FRigidSceneAsyncPT::Unlock(UE::Physics::ERigidLockType InLockType) const
 	{
-		CheckSimContext();
-		// No lock - see Lock() for comments.
-	}
-
-	// TODO_CHAOSAPI: PT and GT APIs are diverging - perhaps IRigidScene needs to be split
-	// Or we need to handle particle creation and registeration from the PT...
-	FRigidBodyAsyncPT* FRigidSceneAsyncPT::CreateRigidBody(const UE::Physics::FRigidNamePtr& InName, UE::Physics::ERigidMovementType InMovementType)
-	{
-		CheckSimContext();
-		checkf(false, TEXT("Not supported from physics thread"));
-		return nullptr;
-	}
-
-	void FRigidSceneAsyncPT::DestroyRigidBody(const UE::Physics::FRigidBodyHandle& BodyHandle)
-	{
-		CheckSimContext();
-		checkf(false, TEXT("Not supported from physics thread"));
-	}
-
-	void FRigidSceneAsyncPT::RegisterRigidBody(const UE::Physics::FRigidBodyHandle& BodyHandle)
-	{
-		CheckSimContext();
-		checkf(false, TEXT("Not supported from physics thread"));
 	}
 
-	void FRigidSceneAsyncPT::UnregisterRigidBody(const UE::Physics::FRigidBodyHandle& BodyHandle)
+	UE::Physics::FRigidSceneHandle FRigidSceneAsyncPT::GetHandle() const
 	{
-		CheckSimContext();
-		checkf(false, TEXT("Not supported from physics thread"));
-	}
-
-	FRigidBodyContainerAsyncPT* FRigidSceneAsyncPT::GetRigidBodyContainer(const UE::Physics::FRigidBodyContainerId& InBodyContainerId)
-	{
-		CheckSimContext();
-		if (InBodyContainerId == SingleBodyContainer.GetContainerId())
-		{
-			return &SingleBodyContainer;
-		}
-		return nullptr;
+		return SceneHandle;
 	}
 
 	void FRigidSceneAsyncPT::RegisterModifier(UE::Physics::IRigidSceneModifier* Modifier)
@@ -433,11 +403,6 @@ namespace Chaos::Rigids::Async
 		checkf(false, TEXT("Not supported from physics thread"));
 	}
 
-	const char* FRigidSceneAsyncPT::GetTypeName() const
-	{ 
-		return "FRigidSceneAsyncPT";
-	}
-
 	UE::Physics::FRigidSceneHandle FRigidSceneAsyncPT::GetSceneHandle() const
 	{
 		CheckSimContext();
@@ -452,8 +417,6 @@ namespace Chaos::Rigids::Async
 
 		Solver = InSolver;
 
-		SingleBodyContainer.Init(this, Solver);
-
 		Solver->GetEvolution()->SetPreIntegrateCallback(
 			[this](FReal Dt)
 			{
@@ -478,40 +441,48 @@ namespace Chaos::Rigids::Async
 		}
 	}
 
-	void FRigidSceneAsyncPT::CreateRigidBody(const FRigidBodyIdAsync& InBodyId, FSingleParticlePhysicsProxy* Proxy)
+	void FRigidSceneAsyncPT::ActivateBody(const UE::Physics::FRigidBodyHandle& InBodyHandle)
 	{
 		CheckSimContext();
-
-		SingleBodyContainer.CreateRigidBody(InBodyId, Proxy);
 	}
 
-	void FRigidSceneAsyncPT::DestroyRigidBody(const UE::Physics::FRigidBodyId& InBodyId)
+	void FRigidSceneAsyncPT::DeactivateBody(const UE::Physics::FRigidBodyHandle& InBodyHandle)
 	{
-		SingleBodyContainer.DestroyRigidBody(InBodyId);
+		CheckSimContext();
 	}
 
-	void FRigidSceneAsyncPT::RegisterRigidBody(const UE::Physics::FRigidBodyId& InBodyId)
+	void FRigidSceneAsyncPT::CreateBodyPT(const UE::Physics::FRigidBodyHandle& InBodyHandle, FSingleParticlePhysicsProxy* Proxy)
 	{
 		CheckSimContext();
-		// Nothing to do because RegisterObject handle GT and PT
-		// See FRigidSceneAsyncGT::RegisterRigidBody
+
+		BodyPool->CreateBodyPT(InBodyHandle, Proxy);
 	}
 
-	void FRigidSceneAsyncPT::UnregisterRigidBody(const UE::Physics::FRigidBodyId& InBodyId)
+	void FRigidSceneAsyncPT::DestroyBodyPT(const UE::Physics::FRigidBodyHandle& InBodyHandle)
+	{
+		BodyPool->DestroyBodyPT(InBodyHandle);
+	}
+
+	UE::Physics::IRigidBodyContainer* FRigidSceneAsyncPT::GetBodyContainer(const UE::Physics::FRigidObjectId& InBodyContainerId) const
 	{
 		CheckSimContext();
-		// Nothing to do because UnegisterObject handle GT and PT
-		// See FRigidSceneAsyncGT::UnregisterRigidBody
+
+		if (InBodyContainerId == BodyPool->GetHandle().GetId())
+		{
+			return BodyPool.Get();
+		}
+
+		return nullptr;
 	}
 
 	void FRigidSceneAsyncPT::CheckGameContext() const
 	{
-		CHAOSRIGIDPHYSICSASYNC_CHECK(IsInGameThreadContext());
+		UE_RIGIDPHYSICS_CHECK(IsInGameThreadContext());
 	}
 
 	void FRigidSceneAsyncPT::CheckSimContext() const
 	{
-		CHAOSRIGIDPHYSICSASYNC_CHECK(IsInPhysicsThreadContext());
+		UE_RIGIDPHYSICS_CHECK(IsInPhysicsThreadContext());
 	}
 
 	void FRigidSceneAsyncPT::OnPreTick(FReal Dt)
@@ -520,18 +491,14 @@ namespace Chaos::Rigids::Async
 
 		CheckSimContext();
 
-		Context.SetSimPhase(EContextSimPhase::PreTick);
-
-		TRigidScenePtr<FSimContextRW> ScenePtr = Private::FRigidAccessor::MakeRigidScenePtr(Context, this);
+		FRigidContextSimRW Context = FRigidContextSimRW(this);
 		for (IRigidSceneModifier* SceneModifier : SceneModifiers)
 		{
-			SceneModifier->SimPreTick(ScenePtr);
+			SceneModifier->PreTick(Context);
 		}
 
 		// Handle legacy modifiers
 		Solver->OnEvolutionPreIntegrate(Dt);
-
-		Context.SetSimPhase(EContextSimPhase::None);
 	}
 
 	void FRigidSceneAsyncPT::OnPostTick(FReal Dt)
@@ -540,18 +507,14 @@ namespace Chaos::Rigids::Async
 
 		CheckSimContext();
 
-		Context.SetSimPhase(EContextSimPhase::PostTick);
-
-		TRigidScenePtr<FSimContextRW> ScenePtr = Private::FRigidAccessor::MakeRigidScenePtr(Context, this);
+		FRigidContextSimRW Context = FRigidContextSimRW(this);
 		for (IRigidSceneModifier* SceneModifier : SceneModifiers)
 		{
-			SceneModifier->SimPostTick(ScenePtr);
+			SceneModifier->PostTick(Context);
 		}
 
 		// Handle legacy modifiers
 		Solver->OnEvolutionPostIntegrate(Dt);
-
-		Context.SetSimPhase(EContextSimPhase::None);
 	}
 
 }
diff --git a/Engine/Plugins/Experimental/ChaosRigidPhysicsAsync/Source/ChaosRigidPhysicsAsync/Public/ChaosRigidPhysicsAsync/RigidBodyAsync.h b/Engine/Plugins/Experimental/ChaosRigidPhysicsAsync/Source/ChaosRigidPhysicsAsync/Public/ChaosRigidPhysicsAsync/RigidBodyAsync.h
--- a/Engine/Plugins/Experimental/ChaosRigidPhysicsAsync/Source/ChaosRigidPhysicsAsync/Public/ChaosRigidPhysicsAsync/RigidBodyAsync.h
+++ b/Engine/Plugins/Experimental/ChaosRigidPhysicsAsync/Source/ChaosRigidPhysicsAsync/Public/ChaosRigidPhysicsAsync/RigidBodyAsync.h
@@ -10,52 +10,10 @@
 #include "PBDRigidsSolver.h"
 #include "Physics/Experimental/PhysScene_Chaos.h"
 #include "PhysicsProxy/SingleParticlePhysicsProxy.h"
-#include "RigidPhysics/RigidAccessor.h"
 #include "RigidPhysics/RigidBody.h"
 
 namespace Chaos::Rigids::Async
 {
-	// Utility class to standardize conversion from FRigidBodyID (uint64) to 
-	// ints for BodyIndex and Epoch, as used by the Async Body Container.
-	class UE_INTERNAL FRigidBodyIdAsync
-	{
-	public:
-		FRigidBodyIdAsync(int32 InBodyIndex, int32 InEpoch)
-			: BodyIndex(InBodyIndex)
-			, Epoch(InEpoch)
-		{
-		}
-
-		explicit FRigidBodyIdAsync(const UE::Physics::FRigidBodyId& BodyId)
-			: BodyIndex(INDEX_NONE)
-			, Epoch(INDEX_NONE)
-		{
-			uint64 Id = Private::FRigidAccessor::GetRigidBodyIdInt(BodyId);
-			BodyIndex = (int32)(Id & 0xFFFFFFFF);
-			Epoch = (int32)((Id >> 32) & 0xFFFFFFFF);
-		}
-
-		UE::Physics::FRigidBodyId AsRigidBodyId() const
-		{
-			uint64 Id = ((uint64)BodyIndex | ((uint64)Epoch << 32));
-			return UE::Physics::FRigidBodyId(Id);
-		}
-
-		int32 GetBodyIndex() const
-		{
-			return BodyIndex;
-		}
-
-		int32 GetEpoch() const
-		{
-			return Epoch;
-		}
-
-	private:
-		int32 BodyIndex;
-		int32 Epoch;
-	};
-
 	//
 	//
 	// RIGID BODY GT/PT
@@ -73,90 +31,62 @@ namespace Chaos::Rigids::Async
 
 	class UE_INTERNAL FRigidBodyAsyncGT : public UE::Physics::IRigidBody
 	{
-	private:
-		inline FRigidBodyHandle_External& GetProxyAPI()
-		{ 
-			return Proxy->GetGameThreadAPI();
-		}
-		
-		inline const FRigidBodyHandle_External& GetProxyAPI() const
-		{
-			return Proxy->GetGameThreadAPI();
-		}
-
 	public:
+		UE_RIGIDPHYSICS_RIGIDTYPED_DECL(CHAOSRIGIDPHYSICSASYNC_API, FRigidBodyAsyncGT);
 
-		UE_INTERNAL CHAOSRIGIDPHYSICSASYNC_API FRigidBodyAsyncGT(
-			FRigidBodyContainerAsyncGT* InContainer, 
-			FSingleParticlePhysicsProxy* InProxy, 
-			const FRigidBodyIdAsync& InBodyId);
+		UE_INTERNAL CHAOSRIGIDPHYSICSASYNC_API FRigidBodyAsyncGT(FRigidBodyPoolAsyncGT* InContainer, const UE::Physics::FRigidDebugName& InName, UE::Physics::ERigidMovementType InMovementType);
 
 		virtual ~FRigidBodyAsyncGT();
 
 		//
-		// IRIGIDBODY BEGIN
+		// Begin IRigidBody
 		//
-
-		virtual const char* GetTypeName() const override final;
-		virtual UE::Physics::FRigidBodyId GetBodyId() const override final;
-		virtual UE::Physics::FRigidBodyHandle GetBodyHandle() const override final;
-		virtual bool IsValid() const override final;
-		virtual bool IsInScene() const override final;
-		virtual bool IsStatic() const override final;
-		virtual bool IsKinematic() const override final;
-		virtual bool IsDynamic() const override final;
-		virtual bool IsSleeping() const override final;
-		virtual void InitTransform(const FTransform& InTransform) override final;
-		virtual void UpdateTransform(const FTransform& InTransform) override final;
-		virtual void SetKinematicTarget(const FTransform& InTransform) override final;
-		virtual FTransform GetTransform() const override final;
-		virtual FBounds3d GetBounds() const override final;
-		virtual double GetMass() const override final;
-		virtual void SetMass(double InMass) override final;
-		//virtual void SetCCDEnabled(bool bInEnabled) override final {  }
-		//virtual void SetMACDEnabled(bool bInEnabled) override final {  }
-		//virtual void SetSmoothEdgeCollisionsEnabled(bool bInEnabled) override final {  }
-		//virtual void SetMaxLinearVelocity(float InMax) override final {  }
-		//virtual void SetMaxAngularVelocity(float InMax) override final {  }
-		//virtual void SetInertiaConditioningEnabled(bool bInEnabled) override final {  }
-		//virtual void SetGyroscopicTorqueEnabled(bool bInEnabled) override final {  }
-		//virtual void SetGravityGroup(int32 InId) override final {  }
-		virtual void SetGravityScale(float InScale) override final;
-		//virtual void SetWakeEventsEnabled(bool bInEnabled) override final {  }
-		//virtual void SetNumPositionIterations(int32 InNum) override final {  }
-		//virtual void SetNumVelocityIterations(int32 InNum) override final {  }
-		//virtual void SetNumProjectionIterations(int32 InNum) override final {  }
-		//virtual void SetIsSleeping(bool bInSleeping) override final {  }
-
+		UE_INTERNAL virtual UE::Physics::FRigidBodyHandle GetHandle() const override final;
+		UE_INTERNAL virtual bool IsActive() const override final;
+		UE_INTERNAL virtual void Activate() override final;
+		UE_INTERNAL virtual void Deactivate() override final;
+		UE_INTERNAL virtual bool IsStatic() const override final;
+		UE_INTERNAL virtual bool IsKinematic() const override final;
+		UE_INTERNAL virtual bool IsDynamic() const override final;
+		UE_INTERNAL virtual bool IsSleeping() const override final;
+		UE_INTERNAL virtual void InitTransform(const FTransform& InTransform) override final;
+		UE_INTERNAL virtual void UpdateTransform(const FTransform& InTransform) override final;
+		UE_INTERNAL virtual void SetKinematicTarget(const FTransform& InTransform) override final;
+		UE_INTERNAL virtual FTransform GetTransform() const override final;
+		UE_INTERNAL virtual FBounds3d GetBounds() const override final;
+		UE_INTERNAL virtual double GetMass() const override final;
+		UE_INTERNAL virtual void SetMass(double InMass) override final;
+		//UE_INTERNAL virtual void SetCCDEnabled(bool bInEnabled) override final {  }
+		//UE_INTERNAL virtual void SetMACDEnabled(bool bInEnabled) override final {  }
+		//UE_INTERNAL virtual void SetSmoothEdgeCollisionsEnabled(bool bInEnabled) override final {  }
+		//UE_INTERNAL virtual void SetMaxLinearVelocity(float InMax) override final {  }
+		//UE_INTERNAL virtual void SetMaxAngularVelocity(float InMax) override final {  }
+		//UE_INTERNAL virtual void SetInertiaConditioningEnabled(bool bInEnabled) override final {  }
+		//UE_INTERNAL virtual void SetGyroscopicTorqueEnabled(bool bInEnabled) override final {  }
+		//UE_INTERNAL virtual void SetGravityGroup(int32 InId) override final {  }
+		UE_INTERNAL virtual void SetGravityScale(float InScale) override final;
+		//UE_INTERNAL virtual void SetWakeEventsEnabled(bool bInEnabled) override final {  }
+		//UE_INTERNAL virtual void SetNumPositionIterations(int32 InNum) override final {  }
+		//UE_INTERNAL virtual void SetNumVelocityIterations(int32 InNum) override final {  }
+		//UE_INTERNAL virtual void SetNumProjectionIterations(int32 InNum) override final {  }
+		//UE_INTERNAL virtual void SetIsSleeping(bool bInSleeping) override final {  }
 		//
-		// IRIGIDBODY END
+		// End IRigidBody
 		//
 
-		UE_INTERNAL FRigidBodyIdAsync GetBodyIdAsync() const
-		{
-			return BodyId;
-		}
+		UE_INTERNAL CHAOSRIGIDPHYSICSASYNC_API bool IsValid() const;
+		UE_INTERNAL CHAOSRIGIDPHYSICSASYNC_API void SetId(const UE::Physics::FRigidObjectId& InId);
+		UE_INTERNAL CHAOSRIGIDPHYSICSASYNC_API const FSingleParticlePhysicsProxy* GetProxy() const;
+		UE_INTERNAL CHAOSRIGIDPHYSICSASYNC_API FSingleParticlePhysicsProxy* GetProxy();
+		UE_INTERNAL CHAOSRIGIDPHYSICSASYNC_API const FRigidBodyHandle_External& GetProxyAPI() const;
+		UE_INTERNAL CHAOSRIGIDPHYSICSASYNC_API FRigidBodyHandle_External& GetProxyAPI();
+		UE_INTERNAL CHAOSRIGIDPHYSICSASYNC_API FRigidSceneAsyncGT* GetScene() const;
 
-		UE_INTERNAL const FSingleParticlePhysicsProxy* GetProxy() const
-		{
-			return Proxy;
-		}
-		
-		UE_INTERNAL FSingleParticlePhysicsProxy* GetProxy()
-		{
-			return Proxy;
-		}
 
 	private:
-		friend class FRigidSceneAsyncGT;
-		friend class FRigidBodyContainerAsyncGT;
-
-		FRigidSceneAsyncGT* GetSceneAsync() const;
-
-		FRigidBodyContainerAsyncGT* Container;
+		UE::Physics::FRigidObjectId BodyId;
+		FRigidBodyPoolAsyncGT* BodyContainer;
 		FSingleParticlePhysicsProxy* Proxy;
-		
-		FRigidBodyIdAsync BodyId;
 	};
 
 
@@ -174,50 +104,40 @@ namespace Chaos::Rigids::Async
 		}
 
 	public:
+		UE_RIGIDPHYSICS_RIGIDTYPED_DECL(CHAOSRIGIDPHYSICSASYNC_API, FRigidBodyAsyncPT);
+
 		UE_INTERNAL CHAOSRIGIDPHYSICSASYNC_API FRigidBodyAsyncPT(
-			FRigidBodyContainerAsyncPT* InContainer, 
-			FSingleParticlePhysicsProxy* InProxy, 
-			const FRigidBodyIdAsync& InBodyId);
+			const UE::Physics::FRigidBodyHandle& InBodyHandle,
+			FSingleParticlePhysicsProxy* InProxy);
 
 		virtual ~FRigidBodyAsyncPT();
 
 		//
-		// IRIGIDBODY BEGIN
+		// Begin IRigidBody
 		//
-
-		virtual const char* GetTypeName() const override final;
-		virtual UE::Physics::FRigidBodyId GetBodyId() const override final;
-		virtual UE::Physics::FRigidBodyHandle GetBodyHandle() const override final;
-		virtual bool IsValid() const override final;
-		virtual bool IsInScene() const override final;
-		virtual bool IsStatic() const override final;
-		virtual bool IsKinematic() const override final;
-		virtual bool IsDynamic() const override final;
-		virtual bool IsSleeping() const override final;
-		virtual void InitTransform(const FTransform& InTransform) override final;
-		virtual void UpdateTransform(const FTransform& InTransform) override final;
-		virtual void SetKinematicTarget(const FTransform& InTransform) override final;
-		virtual FTransform GetTransform() const override final;
-		virtual FBounds3d GetBounds() const override final;
-		virtual double GetMass() const override final;
-		virtual void SetMass(double InMass) override final;
-
+		UE_INTERNAL virtual UE::Physics::FRigidBodyHandle GetHandle() const override final;
+		UE_INTERNAL virtual bool IsStatic() const override final;
+		UE_INTERNAL virtual bool IsKinematic() const override final;
+		UE_INTERNAL virtual bool IsDynamic() const override final;
+		UE_INTERNAL virtual bool IsSleeping() const override final;
+		UE_INTERNAL virtual void InitTransform(const FTransform& InTransform) override final;
+		UE_INTERNAL virtual void UpdateTransform(const FTransform& InTransform) override final;
+		UE_INTERNAL virtual void SetKinematicTarget(const FTransform& InTransform) override final;
+		UE_INTERNAL virtual FTransform GetTransform() const override final;
+		UE_INTERNAL virtual FBounds3d GetBounds() const override final;
+		UE_INTERNAL virtual double GetMass() const override final;
+		UE_INTERNAL virtual void SetMass(double InMass) override final;
 		//
-		// IRIGIDBODY END
+		// End IRigidBody
 		//
 
-		UE_INTERNAL FRigidBodyIdAsync GetBodyIdAsync() const
-		{
-			return BodyId;
-		}
+		UE_INTERNAL bool IsValid() const;
 
 	private:
-		friend class FRigidBodyContainerAsyncPT;
+		UE_INTERNAL void SetTransform(const FTransform& InTransform, bool bIsTeleport);
 
-		FRigidBodyContainerAsyncPT* Container;
+		UE::Physics::FRigidBodyHandle Handle;
 		FSingleParticlePhysicsProxy* Proxy;
-
-		FRigidBodyIdAsync BodyId;
 	};
 
 }
diff --git a/Engine/Plugins/Experimental/ChaosRigidPhysicsAsync/Source/ChaosRigidPhysicsAsync/Public/ChaosRigidPhysicsAsync/RigidBodyContainerAsync.h b/Engine/Plugins/Experimental/ChaosRigidPhysicsAsync/Source/ChaosRigidPhysicsAsync/Public/ChaosRigidPhysicsAsync/RigidBodyContainerAsync.h
--- a/Engine/Plugins/Experimental/ChaosRigidPhysicsAsync/Source/ChaosRigidPhysicsAsync/Public/ChaosRigidPhysicsAsync/RigidBodyContainerAsync.h
+++ b/Engine/Plugins/Experimental/ChaosRigidPhysicsAsync/Source/ChaosRigidPhysicsAsync/Public/ChaosRigidPhysicsAsync/RigidBodyContainerAsync.h
@@ -12,124 +12,69 @@
 #include "Containers/Map.h"
 #include "PBDRigidsSolver.h"
 #include "PhysicsProxy/SingleParticlePhysicsProxy.h"
-#include "RigidPhysics/RigidAccessor.h"
-#include "RigidPhysics/RigidBody.h"
+#include "RigidPhysics/RigidBodyContainer.h"
 
 namespace Chaos::Rigids::Async
 {
-	//
-	//
-	// RIGID BODY CONTAINER GT/PT
-	//
-	// The GT and PT containers hold the set of GT and PT particles respectively.
-	//
-	//
-
-	class UE_INTERNAL FRigidBodyContainerAsyncGT : public UE::Physics::IRigidBodyContainer
+	class UE_INTERNAL FRigidBodyPoolAsyncGT : public UE::Physics::FRigidBodyContainer
 	{
 	public:
-		UE_INTERNAL CHAOSRIGIDPHYSICSASYNC_API static UE::Physics::FRigidBodyContainerTypeId StaticRigidBodyContainerTypeId();
+		UE_RIGIDPHYSICS_RIGIDTYPED_DECL(CHAOSRIGIDPHYSICSASYNC_API, FRigidBodyPoolAsyncGT);
 
-		UE_INTERNAL CHAOSRIGIDPHYSICSASYNC_API FRigidBodyContainerAsyncGT(const UE::Physics::FRigidBodyContainerId& InContainerId);
-
-		virtual ~FRigidBodyContainerAsyncGT();
+		UE_INTERNAL FRigidBodyPoolAsyncGT(FRigidSceneAsyncGT* InScene, const UE::Physics::FRigidDebugName& InName);
+		virtual ~FRigidBodyPoolAsyncGT();
 
 		//
-		// IRIGIDBODYCONTAINER BEGIN
+		// Begin IRigidBodyContainer
 		//
-
-		virtual UE::Physics::IRigidScene* GetScene() const override final;
-		virtual UE::Physics::FRigidBodyContainerId GetContainerId() const override final;
-		virtual FRigidBodyAsyncGT* GetRigidBody(const UE::Physics::FRigidBodyId& BodyId) const override final;
-
+		UE_INTERNAL virtual UE::Physics::FRigidBodyContainerHandle GetHandle() const override final;
+		UE_INTERNAL virtual int32 GetNumBodies() const override final;
+		UE_INTERNAL virtual FRigidBodyAsyncGT* CreateBody(const UE::Physics::FRigidDebugName& InName, UE::Physics::ERigidMovementType InMovementType) override final;
+		UE_INTERNAL virtual void DestroyBody(UE::Physics::IRigidBody* InBody) override final;
 		//
-		// IRIGIDBODYCONTAINER END
+		// End IRigidBodyContainer
 		//
 
-		UE_INTERNAL CHAOSRIGIDPHYSICSASYNC_API UE::Physics::FRigidBodyContainerHandle GetContainerHandle() const;
-
-		UE_INTERNAL CHAOSRIGIDPHYSICSASYNC_API void Init(FRigidSceneAsyncGT* InScene, FPBDRigidsSolver* InSolver);
-
-		UE_INTERNAL CHAOSRIGIDPHYSICSASYNC_API FRigidBodyAsyncGT* CreateRigidBody(const UE::Physics::FRigidNamePtr& InName, UE::Physics::ERigidMovementType InMovementType);
-
-		UE_INTERNAL CHAOSRIGIDPHYSICSASYNC_API FRigidBodyAsyncGT* CreateRigidBodyForProxy(FSingleParticlePhysicsProxy* InProxy);
-
-		UE_INTERNAL CHAOSRIGIDPHYSICSASYNC_API void DestroyRigidBody(const UE::Physics::FRigidBodyId& BodyId);
-
-		UE_INTERNAL CHAOSRIGIDPHYSICSASYNC_API void AddRigidBodyToScene(const UE::Physics::FRigidBodyId& BodyId);
-
-		UE_INTERNAL CHAOSRIGIDPHYSICSASYNC_API void RemoveRigidBodyFromScene(const UE::Physics::FRigidBodyId& BodyId);
-
-		UE_INTERNAL CHAOSRIGIDPHYSICSASYNC_API void OnRigidBodyTransformUpdated(FRigidBodyAsyncGT* Body);
-
-		UE_INTERNAL CHAOSRIGIDPHYSICSASYNC_API FRigidSceneAsyncGT* GetSceneAsync() const;
+		UE_INTERNAL CHAOSRIGIDPHYSICSASYNC_API void SetId(const UE::Physics::FRigidObjectId& InId);
+		UE_INTERNAL CHAOSRIGIDPHYSICSASYNC_API FRigidSceneAsyncGT* GetScene() const;
+		UE_INTERNAL CHAOSRIGIDPHYSICSASYNC_API void VisitBodies(const TFunctionRef<UE::Physics::ERigidVisitorResponse(UE::Physics::IRigidBody*)>& InVisitor);
 
 	private:
+		void CheckIsLocked() const;
 
-		// Store the Epoch with the body to validate handles
-		struct FRigidBodyArrayEntry
-		{
-			FRigidBodyAsyncGT* Body;
-			int32 Epoch;
-		};
-
-		FRigidBodyAsyncGT* AllocBody(FSingleParticlePhysicsProxy* InProxy);
-		void FreeBody(FRigidBodyAsyncGT* BodyGT);
-		void ReserveBodyArray(int32 MaxBodies);
-
+		UE::Physics::FRigidDebugName Name;
+		UE::Physics::FRigidObjectId BodyContainerId;
 		FRigidSceneAsyncGT* Scene;
-		UE::Physics::FRigidBodyContainerId ContainerId;
-		FPBDRigidsSolver* Solver;
-
-		// Body allocator
 		TObjectPool<FRigidBodyAsyncGT> BodyPool;
-
-		// BodyId->Body mapping
-		TArray<FRigidBodyArrayEntry> Bodies;
-		TArray<int32> BodiesFreeIndices;
-		int64 NextBodyIdIndex;
 	};
 
-
-	class UE_INTERNAL FRigidBodyContainerAsyncPT : public UE::Physics::IRigidBodyContainer
+	class UE_INTERNAL FRigidBodyPoolAsyncPT : public UE::Physics::IRigidBodyContainer
 	{
 	public:
-		UE_INTERNAL static UE::Physics::FRigidBodyContainerTypeId StaticRigidBodyContainerTypeId();
+		UE_RIGIDPHYSICS_RIGIDTYPED_DECL(CHAOSRIGIDPHYSICSASYNC_API, FRigidBodyPoolAsyncPT);
 
-		UE_INTERNAL FRigidBodyContainerAsyncPT(const UE::Physics::FRigidBodyContainerId& InContainerId);
+		UE_INTERNAL FRigidBodyPoolAsyncPT(const UE::Physics::FRigidDebugName& InName, const UE::Physics::FRigidBodyContainerHandle& InContainerHandle);
 
-		virtual ~FRigidBodyContainerAsyncPT();
+		virtual ~FRigidBodyPoolAsyncPT();
 
 		//
-		// IRIGIDBODYCONTAINER BEGIN
+		// Begin IRigidBodyContainer
 		//
-
-		virtual UE::Physics::IRigidScene* GetScene() const override final;
-		virtual UE::Physics::FRigidBodyContainerId GetContainerId() const override final;
-		virtual FRigidBodyAsyncPT* GetRigidBody(const UE::Physics::FRigidBodyId& BodyId) const override final;
-
+		UE_INTERNAL virtual UE::Physics::FRigidBodyContainerHandle GetHandle() const override final;
+		UE_INTERNAL virtual int32 GetNumBodies() const override final;
+		UE_INTERNAL virtual FRigidBodyAsyncPT* GetBody(const UE::Physics::FRigidObjectId& InBodyId) const override final;
 		//
-		// IRIGIDBODYCONTAINER END
+		// End IRigidBodyContainer
 		//
 
-		UE_INTERNAL void Init(FRigidSceneAsyncPT* InScene, FPBDRigidsSolver* InSolver);
-
-		UE_INTERNAL UE::Physics::FRigidBodyContainerHandle GetContainerHandle() const;
-
-		UE_INTERNAL void CreateRigidBody(const FRigidBodyIdAsync& InBodyId, FSingleParticlePhysicsProxy* InProxy);
-
-		UE_INTERNAL void DestroyRigidBody(const UE::Physics::FRigidBodyId& BodyId);
-
-		UE_INTERNAL void AddRigidBodyToScene(const UE::Physics::FRigidBodyId& BodyId);
-
-		UE_INTERNAL void RemoveRigidBodyFromScene(const UE::Physics::FRigidBodyId& BodyId);
+		UE_INTERNAL CHAOSRIGIDPHYSICSASYNC_API void CreateBodyPT(const UE::Physics::FRigidBodyHandle& InBodyHandle, FSingleParticlePhysicsProxy* InProxy);
+		UE_INTERNAL CHAOSRIGIDPHYSICSASYNC_API void DestroyBodyPT(const UE::Physics::FRigidBodyHandle& InBodyHandle);
 
 	private:
-		FRigidSceneAsyncPT* Scene;
-		UE::Physics::FRigidBodyContainerId ContainerId;
+		UE::Physics::FRigidBodyContainerHandle Handle;
 		FPBDRigidsSolver* Solver;
 		TObjectPool<FRigidBodyAsyncPT> BodyPool;
-		TMap<UE::Physics::FRigidBodyId, FRigidBodyAsyncPT*> IdToBodyMap;
+		TMap<UE::Physics::FRigidObjectId, FRigidBodyAsyncPT*> IdToBodyMap;
 	};
 
 }
diff --git a/Engine/Plugins/Experimental/ChaosRigidPhysicsAsync/Source/ChaosRigidPhysicsAsync/Public/ChaosRigidPhysicsAsync/RigidFactoryAsync.h b/Engine/Plugins/Experimental/ChaosRigidPhysicsAsync/Source/ChaosRigidPhysicsAsync/Public/ChaosRigidPhysicsAsync/RigidFactoryAsync.h
--- a/Engine/Plugins/Experimental/ChaosRigidPhysicsAsync/Source/ChaosRigidPhysicsAsync/Public/ChaosRigidPhysicsAsync/RigidFactoryAsync.h
+++ b/Engine/Plugins/Experimental/ChaosRigidPhysicsAsync/Source/ChaosRigidPhysicsAsync/Public/ChaosRigidPhysicsAsync/RigidFactoryAsync.h
@@ -5,6 +5,7 @@
 
 #if UE_RIGIDPHYSICS_API_ENABLED
 
+#include "RigidPhysics/RigidScene.h"
 #include "Templates/SharedPointer.h"
 #include "Templates/UniquePtr.h"
 
@@ -14,47 +15,29 @@ namespace Chaos
 	class FSingleParticlePhysicsProxy;
 }
 
-namespace Chaos::Rigids::Async::Private
-{
-	using FRigidAccessor = UE::Physics::Private::FRigidAccessor;
-	class FRigidAccessorAsync;
-}
-
 namespace Chaos::Rigids::Async
 {
 	class FRigidSceneSettingsAsync
 	{
 	public:
-		FName SceneName = NAME_None;
 		float AsyncDt = (1.0f / 60.0f);
 	};
 
-	class FRigidFactoryAsync
+	class FRigidFactoryAsync : public UE::Physics::IRigidSceneFactory
 	{
 	public:
-		CHAOSRIGIDPHYSICSASYNC_API UE::Physics::FRigidSceneHandle CreateRigidScene(const FRigidSceneSettingsAsync& InSettings);
-		CHAOSRIGIDPHYSICSASYNC_API void DestroyRigidScene(const UE::Physics::FRigidSceneHandle& InSceneHandle);
+		FRigidFactoryAsync();
+		virtual ~FRigidFactoryAsync();
+
+		UE_INTERNAL CHAOSRIGIDPHYSICSASYNC_API virtual UE::Physics::FRigidSceneHandle CreateScene(const UE::Physics::FRigidDebugName& InName) override final;
+		UE_INTERNAL CHAOSRIGIDPHYSICSASYNC_API virtual void DestroyScene(const UE::Physics::FRigidSceneHandle& InSceneHandle) override final;
 
-	private:
 		// TODO_CHAOSAPI: Temp until solver is created through new API
-		CHAOSRIGIDPHYSICSASYNC_API UE::Physics::FRigidSceneHandle CreateRigidSceneFromSolver(Chaos::FPBDRigidsSolver* InSolver);
+		UE_INTERNAL CHAOSRIGIDPHYSICSASYNC_API UE::Physics::IRigidScene* CreateSceneFromSolver(Chaos::FPBDRigidsSolver* InSolver);
 
-		friend class Private::FRigidAccessorAsync;
-		TArray<UE::Physics::FRigidSceneSharedPtr> Scenes;
+	private:
+		TArray<UE::Physics::IRigidScene*> Scenes;
 	};
 }
 
-// TODO_CHAOSAPI: Temp until solver is created through new API
-namespace Chaos::Rigids::Async::Private
-{
-	// Create a scene for and externally-created solver (simulation proxy)
-	UE_INTERNAL CHAOSRIGIDPHYSICSASYNC_API UE::Physics::FRigidSceneHandle CreateSceneFromSolver(FRigidFactoryAsync& InFactory, FPBDRigidsSolver* InSolver);
-
-	// Create a body handle for an externally-created single particle proxy
-	UE_INTERNAL CHAOSRIGIDPHYSICSASYNC_API UE::Physics::FRigidBodyHandle CreateBodyFromProxy(const UE::Physics::FRigidSceneHandle& InSceneHandle, FSingleParticlePhysicsProxy* InProxy);
-
-	// Get the proxy from the body handle
-	UE_INTERNAL CHAOSRIGIDPHYSICSASYNC_API FSingleParticlePhysicsProxy* GetProxyFromBody(const UE::Physics::FRigidBodyHandle& InBodyHandle);
-}
-
 #endif // UE_RIGIDPHYSICS_API_ENABLED
\ No newline at end of file
diff --git a/Engine/Plugins/Experimental/ChaosRigidPhysicsAsync/Source/ChaosRigidPhysicsAsync/Public/ChaosRigidPhysicsAsync/RigidFwdAsync.h b/Engine/Plugins/Experimental/ChaosRigidPhysicsAsync/Source/ChaosRigidPhysicsAsync/Public/ChaosRigidPhysicsAsync/RigidFwdAsync.h
--- a/Engine/Plugins/Experimental/ChaosRigidPhysicsAsync/Source/ChaosRigidPhysicsAsync/Public/ChaosRigidPhysicsAsync/RigidFwdAsync.h
+++ b/Engine/Plugins/Experimental/ChaosRigidPhysicsAsync/Source/ChaosRigidPhysicsAsync/Public/ChaosRigidPhysicsAsync/RigidFwdAsync.h
@@ -11,22 +11,10 @@ namespace Chaos::Rigids::Async
 {
 	class UE_INTERNAL FRigidBodyAsyncGT;
 	class UE_INTERNAL FRigidBodyAsyncPT;
-	class UE_INTERNAL FRigidBodyContainerAsyncGT;
-	class UE_INTERNAL FRigidBodyContainerAsyncPT;
+	class UE_INTERNAL FRigidBodyPoolAsyncGT;
+	class UE_INTERNAL FRigidBodyPoolAsyncPT;
 	class UE_INTERNAL FRigidSceneAsyncGT;
 	class UE_INTERNAL FRigidSceneAsyncPT;
 }
 
-namespace Chaos::Rigids::Async::Private
-{
-	using UE_INTERNAL FRigidAccessor = UE::Physics::Private::FRigidAccessor;
-}
-
-// Enable asserts if we have a context to check (not TEST or SHIPPING builds)
-#if PHYSICS_THREAD_CONTEXT && DO_CHECK
-#define CHAOSRIGIDPHYSICSASYNC_CHECK(X) check(X)
-#else
-#define CHAOSRIGIDPHYSICSASYNC_CHECK(X)
-#endif
-
 #endif // UE_RIGIDPHYSICS_API_ENABLED
\ No newline at end of file
diff --git a/Engine/Plugins/Experimental/ChaosRigidPhysicsAsync/Source/ChaosRigidPhysicsAsync/Public/ChaosRigidPhysicsAsync/RigidSceneAsync.h b/Engine/Plugins/Experimental/ChaosRigidPhysicsAsync/Source/ChaosRigidPhysicsAsync/Public/ChaosRigidPhysicsAsync/RigidSceneAsync.h
--- a/Engine/Plugins/Experimental/ChaosRigidPhysicsAsync/Source/ChaosRigidPhysicsAsync/Public/ChaosRigidPhysicsAsync/RigidSceneAsync.h
+++ b/Engine/Plugins/Experimental/ChaosRigidPhysicsAsync/Source/ChaosRigidPhysicsAsync/Public/ChaosRigidPhysicsAsync/RigidSceneAsync.h
@@ -5,7 +5,6 @@
 
 #if UE_RIGIDPHYSICS_API_ENABLED
 
-#include "Async/TaskGraphInterfaces.h"
 #include "Chaos/ObjectPool.h"
 #include "ChaosRigidPhysicsAsync/RigidBodyAsync.h"
 #include "ChaosRigidPhysicsAsync/RigidBodyContainerAsync.h"
@@ -16,8 +15,7 @@
 #include "PBDRigidsSolver.h"
 #include "Physics/Experimental/PhysScene_Chaos.h"
 #include "PhysicsProxy/SingleParticlePhysicsProxy.h"
-#include "RigidPhysics/RigidAccessor.h"
-#include "RigidPhysics/RigidBody.h"
+#include "RigidPhysics/RigidScene.h"
 
 namespace Chaos::Rigids::Async
 {
@@ -32,57 +30,66 @@ namespace Chaos::Rigids::Async
 	//
 	//
 
-	class UE_INTERNAL FRigidSceneAsyncGT : public UE::Physics::IRigidScene, public TSharedFromThis<FRigidSceneAsyncGT, ESPMode::ThreadSafe>
+	class UE_INTERNAL FRigidSceneAsyncGT : public UE::Physics::FRigidScene, public TSharedFromThis<FRigidSceneAsyncGT, ESPMode::ThreadSafe>
 	{
 	public:
-		UE_INTERNAL CHAOSRIGIDPHYSICSASYNC_API FRigidSceneAsyncGT(const FRigidSceneSettingsAsync& InSceneSettings);
-		UE_INTERNAL CHAOSRIGIDPHYSICSASYNC_API FRigidSceneAsyncGT(FPBDRigidsSolver* InSolver);
+		UE_INTERNAL CHAOSRIGIDPHYSICSASYNC_API FRigidSceneAsyncGT(
+			const UE::Physics::FRigidDebugName& InName,
+			const FRigidSceneSettingsAsync& InSceneSettings);
+		UE_INTERNAL CHAOSRIGIDPHYSICSASYNC_API FRigidSceneAsyncGT(
+			const UE::Physics::FRigidDebugName& InName,
+			FPBDRigidsSolver* InSolver);
 		virtual ~FRigidSceneAsyncGT();
 
 		//
-		// IRIGIDSCENE BEGIN
+		// Begin IRigidLockable
 		//
-	public:
-		virtual void Lock(UE::Physics::EContextLockType InLockType) override final;
-		virtual void Unlock(UE::Physics::EContextLockType InLockType) override final;
-		virtual FRigidBodyAsyncGT* CreateRigidBody(const UE::Physics::FRigidNamePtr& InName, UE::Physics::ERigidMovementType InMovementType) override final;
-		virtual void DestroyRigidBody(const UE::Physics::FRigidBodyHandle& BodyHandle) override final;
-		virtual void RegisterRigidBody(const UE::Physics::FRigidBodyHandle& BodyHandle) override final;
-		virtual void UnregisterRigidBody(const UE::Physics::FRigidBodyHandle& BodyHandle) override final;
-		virtual FRigidBodyContainerAsyncGT* GetRigidBodyContainer(const UE::Physics::FRigidBodyContainerId& InBodyContainerId) override final;
-		virtual void RegisterModifier(UE::Physics::IRigidSceneModifier* Modifier) override final;
-		virtual void UnregisterModifier(UE::Physics::IRigidSceneModifier* Modifier) override final;
-		virtual bool Query(const UE::Physics::FRigidQueryDesc& QueryDesc, TArray<UE::Physics::FRigidQueryHit>& OutHits) override final;
-		virtual void Tick(float Dt) override final;
-
 	protected:
-		virtual const char* GetTypeName() const override final;
+		UE_INTERNAL virtual void Lock(UE::Physics::ERigidLockType InLockType) const override final;
+		UE_INTERNAL virtual void Unlock(UE::Physics::ERigidLockType InLockType) const override final;
+		//
+		// End IRigidLockable
+		//
 
+		//
+		// IRIGIDSCENE BEGIN
+		//
+	public:
+		UE_INTERNAL virtual UE::Physics::FRigidSceneHandle GetHandle() const override final;
+		UE_INTERNAL virtual void RegisterModifier(UE::Physics::IRigidSceneModifier* Modifier) override final;
+		UE_INTERNAL virtual void UnregisterModifier(UE::Physics::IRigidSceneModifier* Modifier) override final;
+		UE_INTERNAL virtual bool Query(const UE::Physics::FRigidQueryDesc& QueryDesc, TArray<UE::Physics::FRigidQueryHit>& OutHits) override final;
+		UE_INTERNAL virtual void Tick(float Dt) override final;
+		UE_INTERNAL virtual FRigidBodyAsyncGT* CreateBody(const UE::Physics::FRigidDebugName& InName, UE::Physics::ERigidMovementType InMovementType) override final;
+		UE_INTERNAL virtual void DestroyBody(UE::Physics::IRigidBody* InBody) override final;
 		//
 		// IRIGIDSCENE END
 		//
 
 	public:
-		UE_INTERNAL CHAOSRIGIDPHYSICSASYNC_API UE::Physics::FRigidSceneHandle GetSceneHandle() const;
 		UE_INTERNAL CHAOSRIGIDPHYSICSASYNC_API void Startup();
 		UE_INTERNAL CHAOSRIGIDPHYSICSASYNC_API void Shutdown();
+		UE_INTERNAL CHAOSRIGIDPHYSICSASYNC_API void ActivateBody(FRigidBodyAsyncGT* InBody);
+		UE_INTERNAL CHAOSRIGIDPHYSICSASYNC_API void DeactivateBody(FRigidBodyAsyncGT* InBody);
 		UE_INTERNAL CHAOSRIGIDPHYSICSASYNC_API void OnRigidBodyTransformUpdated(FRigidBodyAsyncGT* Body);
 
 	private:
-		UE::Physics::FRigidSceneHandle GetHandle();
 		void CheckGameReadLocked() const;
 		void CheckGameWriteLocked() const;
+		void ConstructScene();
+		void DestructScene();
 		void OnPreSimulate(FReal Dt);
 		void OnPostSimulate(FReal Dt);
 		void UpdateRigidBodyInAccelerationStructure(FRigidBodyAsyncGT* Body);
 
+		UE::Physics::FRigidDebugName Name;
+
 		FRigidSceneSettingsAsync SceneSettings;
 
 		FPBDRigidsSolver* Solver;
 		ISpatialAccelerationCollection<FAccelerationStructureHandle, FReal, 3>* SpatialAcceleration;
 
-		UE::Physics::FRigidSceneHandle SceneHandle;
-		FRigidBodyContainerAsyncGT SingleBodyContainer;
+		TUniquePtr<FRigidBodyPoolAsyncGT> BodyPool;
 
 		TArray<UE::Physics::IRigidSceneModifier*> SceneModifiers;
 
@@ -96,28 +103,28 @@ namespace Chaos::Rigids::Async
 	class UE_INTERNAL FRigidSceneAsyncPT : public UE::Physics::IRigidScene
 	{
 	public:
-		UE_INTERNAL CHAOSRIGIDPHYSICSASYNC_API FRigidSceneAsyncPT(const UE::Physics::FRigidBodyContainerId& InSingleBodyContainerId);
+		UE_INTERNAL CHAOSRIGIDPHYSICSASYNC_API FRigidSceneAsyncPT(const UE::Physics::FRigidDebugName& InName, const UE::Physics::FRigidBodyContainerHandle& InBodyPoolHandle);
 
 		//
-		// IRIGIDSCENE BEGIN
+		// Begin IRigidLockable
 		//
-
-	public:
-		virtual void Lock(UE::Physics::EContextLockType InLockType) override final;
-		virtual void Unlock(UE::Physics::EContextLockType InLockType) override final;
-		virtual FRigidBodyAsyncPT* CreateRigidBody(const UE::Physics::FRigidNamePtr& InName, UE::Physics::ERigidMovementType InMovementType) override final;
-		virtual void DestroyRigidBody(const UE::Physics::FRigidBodyHandle& BodyHandle) override final;
-		virtual void RegisterRigidBody(const UE::Physics::FRigidBodyHandle& BodyHandle) override final;
-		virtual void UnregisterRigidBody(const UE::Physics::FRigidBodyHandle& BodyHandle) override final;
-		virtual FRigidBodyContainerAsyncPT* GetRigidBodyContainer(const UE::Physics::FRigidBodyContainerId& InBodyContainerId) override final;
-		virtual void RegisterModifier(UE::Physics::IRigidSceneModifier* Modifier) override final;
-		virtual void UnregisterModifier(UE::Physics::IRigidSceneModifier* Modifier) override final;
-		virtual bool Query(const UE::Physics::FRigidQueryDesc& QueryDesc, TArray<UE::Physics::FRigidQueryHit>& OutHits) override final;
-		virtual void Tick(float Dt) override final;
-
 	protected:
-		virtual const char* GetTypeName() const override final;
+		UE_INTERNAL virtual void Lock(UE::Physics::ERigidLockType InLockType) const override final;
+		UE_INTERNAL virtual void Unlock(UE::Physics::ERigidLockType InLockType) const override final;
+		//
+		// End IRigidLockable
+		//
 
+		//
+		// IRIGIDSCENE BEGIN
+		//
+	public:
+		UE_INTERNAL virtual UE::Physics::FRigidSceneHandle GetHandle() const override final;
+		UE_INTERNAL virtual UE::Physics::IRigidBodyContainer* GetBodyContainer(const UE::Physics::FRigidObjectId& InBodyContainerId) const override final;
+		UE_INTERNAL virtual void RegisterModifier(UE::Physics::IRigidSceneModifier* Modifier) override final;
+		UE_INTERNAL virtual void UnregisterModifier(UE::Physics::IRigidSceneModifier* Modifier) override final;
+		UE_INTERNAL virtual bool Query(const UE::Physics::FRigidQueryDesc& QueryDesc, TArray<UE::Physics::FRigidQueryHit>& OutHits) override final;
+		UE_INTERNAL virtual void Tick(float Dt) override final;
 		//
 		// IRIGIDSCENE END
 		//
@@ -126,10 +133,10 @@ namespace Chaos::Rigids::Async
 		UE_INTERNAL CHAOSRIGIDPHYSICSASYNC_API UE::Physics::FRigidSceneHandle GetSceneHandle() const;
 		UE_INTERNAL CHAOSRIGIDPHYSICSASYNC_API void Startup(UE::Physics::FRigidSceneHandle InSceneHandle, FPBDRigidsSolver* InSolver);
 		UE_INTERNAL CHAOSRIGIDPHYSICSASYNC_API void Shutdown();
-		UE_INTERNAL CHAOSRIGIDPHYSICSASYNC_API void CreateRigidBody(const FRigidBodyIdAsync& InBodyId, FSingleParticlePhysicsProxy* Proxy);
-		UE_INTERNAL CHAOSRIGIDPHYSICSASYNC_API void DestroyRigidBody(const UE::Physics::FRigidBodyId& InBodyId);
-		UE_INTERNAL CHAOSRIGIDPHYSICSASYNC_API void RegisterRigidBody(const UE::Physics::FRigidBodyId& InBodyId);
-		UE_INTERNAL CHAOSRIGIDPHYSICSASYNC_API void UnregisterRigidBody(const UE::Physics::FRigidBodyId& InBodyId);
+		UE_INTERNAL CHAOSRIGIDPHYSICSASYNC_API void CreateBodyPT(const UE::Physics::FRigidBodyHandle& InBodyHandle, FSingleParticlePhysicsProxy* Proxy);
+		UE_INTERNAL CHAOSRIGIDPHYSICSASYNC_API void DestroyBodyPT(const UE::Physics::FRigidBodyHandle& InBodyHandle);
+		UE_INTERNAL CHAOSRIGIDPHYSICSASYNC_API void ActivateBody(const UE::Physics::FRigidBodyHandle& InBodyHandle);
+		UE_INTERNAL CHAOSRIGIDPHYSICSASYNC_API void DeactivateBody(const UE::Physics::FRigidBodyHandle& InBodyHandle);
 
 	private:
 		void CheckGameContext() const;
@@ -138,11 +145,11 @@ namespace Chaos::Rigids::Async
 		void OnPreTick(FReal Dt);
 		void OnPostTick(FReal Dt);
 
+		UE::Physics::FRigidDebugName Name;
 		UE::Physics::FRigidSceneHandle SceneHandle;
 		FPBDRigidsSolver* Solver;
-		FRigidBodyContainerAsyncPT SingleBodyContainer;
+		TUniquePtr<FRigidBodyPoolAsyncPT> BodyPool;
 		TArray<UE::Physics::IRigidSceneModifier*> SceneModifiers;
-		UE::Physics::FSimContextRW Context;
 	};
 }
 
diff --git a/Engine/Plugins/Experimental/ChaosRigidPhysicsAsync/Tests/ChaosRigidPhysicsAsyncTests.Target.cs b/Engine/Plugins/Experimental/ChaosRigidPhysicsAsync/Tests/ChaosRigidPhysicsAsyncTests.Target.cs
--- a/Engine/Plugins/Experimental/ChaosRigidPhysicsAsync/Tests/ChaosRigidPhysicsAsyncTests.Target.cs
+++ b/Engine/Plugins/Experimental/ChaosRigidPhysicsAsync/Tests/ChaosRigidPhysicsAsyncTests.Target.cs
@@ -7,6 +7,9 @@ using UnrealBuildTool;
  {
      public ChaosRigidPhysicsAsyncTestsTarget(TargetInfo Target) : base(Target)
      {
+		bTestsRequireEngine = true;
+		bTestsRequireCoreUObject = true;
+
 		bCompileAgainstEngine = true;
 		bCompileAgainstCoreUObject = true;
 		bMockEngineDefaults = true;
diff --git a/Engine/Plugins/Experimental/ChaosRigidPhysicsAsync/Tests/Private/ChaosRigidPhysicsAsyncTest.h b/Engine/Plugins/Experimental/ChaosRigidPhysicsAsync/Tests/Private/ChaosRigidPhysicsAsyncTest.h
new file mode 100644
--- /dev/null
+++ b/Engine/Plugins/Experimental/ChaosRigidPhysicsAsync/Tests/Private/ChaosRigidPhysicsAsyncTest.h
@@ -0,0 +1,27 @@
+// Copyright Epic Games, Inc. All Rights Reserved.
+#pragma once
+
+#include "Chaos/LowLevelTest/ChaosTestHarness.h"
+#include "Chaos/LowLevelTest/ChaosTestScene.h"
+#include "RigidPhysics/RigidTyped.h"
+
+#include <string>
+
+namespace Chaos::LowLevelTest
+{
+	inline std::string ToStdString(const FString& InString)
+	{
+		return std::string(TCHAR_TO_UTF8(*InString));
+	}
+}
+
+namespace Catch
+{
+	// Convert FRigidTypeId to string for output in Catch2 logs
+	template<>
+	struct StringMaker<UE::Physics::FRigidTypeId>
+	{
+		static std::string convert(UE::Physics::FRigidTypeId const& Value) { return Chaos::LowLevelTest::ToStdString(Value.GetTypeName()); }
+	};
+}
+
diff --git a/Engine/Plugins/Experimental/ChaosRigidPhysicsAsync/Tests/Private/SceneLockBenchmarks.cpp b/Engine/Plugins/Experimental/ChaosRigidPhysicsAsync/Tests/Private/SceneLockBenchmarks.cpp
--- a/Engine/Plugins/Experimental/ChaosRigidPhysicsAsync/Tests/Private/SceneLockBenchmarks.cpp
+++ b/Engine/Plugins/Experimental/ChaosRigidPhysicsAsync/Tests/Private/SceneLockBenchmarks.cpp
@@ -3,12 +3,14 @@
 #include "Chaos/LowLevelTest/ChaosTestHarness.h"
 #include "Chaos/LowLevelTest/ChaosTestScene.h"
 #include "ChaosRigidPhysicsAsync/RigidFactoryAsync.h"
+#include "ChaosRigidPhysicsAsync/RigidFwdAsync.h"
 #include "PBDRigidsSolver.h"
 #include "PhysicsInterfaceDeclaresCore.h"
 #include "Physics/Experimental/PhysInterface_Chaos.h"
 #include "PhysicsProxy/SingleParticlePhysicsProxy.h"
-#include "RigidPhysics/RigidBody.h"
+#include "RigidPhysics/RigidContext.h"
 
+#if UE_RIGIDPHYSICS_API_ENABLED
 
 namespace Chaos::LowLevelTest
 {
@@ -25,18 +27,18 @@ namespace Chaos::LowLevelTest
 		FTransform Transform = FTransform::Identity;
 
 		FRigidFactoryAsync Factory;
-		FRigidSceneHandle SceneHandle = Factory.CreateRigidScene(FRigidSceneSettingsAsync());
+		FRigidSceneHandle SceneHandle = Factory.CreateScene(FRigidDebugName::Make(TEXT("Main")));
 
 		// Create the bodies
 		TArray<FRigidBodyHandle> BodyHandles;
 
-		if (TRigidScenePtr<FGameContextRW> Scene = SceneHandle.LockRW())
+		if (FRigidContextGameRW Context = SceneHandle.LockRW())
 		{
 			BodyHandles.Reserve(NumBodies);
 			for (int32 BodyIndex = 0; BodyIndex < NumBodies; ++BodyIndex)
 			{
-				TRigidBodyPtr<FGameContextRW> Body = Scene->CreateRigidBody(MakeRigidName(TEXT("Body{0}"), BodyIndex), ERigidMovementType::Dynamic);
-				Scene->RegisterRigidBody(Body);
+				TRigidBodyPtr<FRigidContextGameRW> Body = Context->CreateBody(FRigidDebugName::Make(TEXT("Body{0}"), BodyIndex), ERigidMovementType::Dynamic);
+				Body->Activate();
 
 				BodyHandles.Add(Body);
 			}
@@ -48,26 +50,28 @@ namespace Chaos::LowLevelTest
 			FVector Pos = Transform.GetTranslation() + FVector(100, 0, 0);
 			Transform.SetTranslation(Pos);
 
-			for (FRigidBodyHandle BodyHandle : BodyHandles)
+			if (FRigidContextGameRW Context = SceneHandle.LockRW())
 			{
-				if (TRigidBodyPtr<FGameContextRW> Body = BodyHandle.LockRW())
+				for (FRigidBodyHandle BodyHandle : BodyHandles)
 				{
-					Body->InitTransform(Transform);
+					if (TRigidBodyPtr<FRigidContextGameRW> Body = BodyHandle.Pin(Context))
+					{
+						Body->InitTransform(Transform);
+					}
 				}
 			}
 		}
 
-		if (TRigidScenePtr<FGameContextRW> Scene = SceneHandle.LockRW())
+		if (FRigidContextGameRW Context = SceneHandle.LockRW())
 		{
 			for (FRigidBodyHandle BodyHandle : BodyHandles)
 			{
-				Scene->UnregisterRigidBody(BodyHandle);
-				Scene->DestroyRigidBody(BodyHandle);
+				Context->DestroyBody(BodyHandle.Pin(Context));
 			}
 		}
 		BodyHandles.Reset();
 
-		Factory.DestroyRigidScene(SceneHandle);
+		Factory.DestroyScene(SceneHandle);
 	}
 
 	// This test is only intended to be run with an attached sampling profiler.
@@ -117,22 +121,26 @@ namespace Chaos::LowLevelTest
 		using namespace Chaos::Rigids::Async;
 		using namespace UE::Physics;
 
+#if UE_RIGIDPHYSICS_CHECK_ENABLED
+		UE_LOG(LogRigidPhysics, Error, TEXT("Benchmarks with asserts enabled will not give best results"));
+#endif
+
 		int32 NumBodies = 10000;
 		FTransform Transform = FTransform::Identity;
 
 		FRigidFactoryAsync Factory;
-		FRigidSceneHandle SceneHandle = Factory.CreateRigidScene(FRigidSceneSettingsAsync());
+		FRigidSceneHandle SceneHandle = Factory.CreateScene(FRigidDebugName::Make(TEXT("Main")));
 
 		// Create the bodies
 		TArray<FRigidBodyHandle> BodyHandles;
 
-		if (TRigidScenePtr<FGameContextRW> Scene = SceneHandle.LockRW())
+		if (FRigidContextGameRW Context = SceneHandle.LockRW())
 		{
 			BodyHandles.Reserve(NumBodies);
 			for (int32 BodyIndex = 0; BodyIndex < NumBodies; ++BodyIndex)
 			{
-				TRigidBodyPtr<FGameContextRW> Body = Scene->CreateRigidBody(MakeRigidName(TEXT("Body{0}"), BodyIndex), ERigidMovementType::Dynamic);
-				Scene->RegisterRigidBody(Body);
+				TRigidBodyPtr<FRigidContextGameRW> Body = Context->CreateBody(FRigidDebugName::Make(TEXT("Body{0}"), BodyIndex), ERigidMovementType::Dynamic);
+				Body->Activate();
 
 				BodyHandles.Add(Body);
 			}
@@ -145,9 +153,12 @@ namespace Chaos::LowLevelTest
 
 			for (FRigidBodyHandle BodyHandle : BodyHandles)
 			{
-				if (TRigidBodyPtr<FGameContextRW> Body = BodyHandle.LockRW())
+				if (FRigidContextGameRW Context = SceneHandle.LockRW())
 				{
-					Body->InitTransform(Transform);
+					if (TRigidBodyPtr<FRigidContextGameRW> Body = BodyHandle.Pin(Context))
+					{
+						Body->InitTransform(Transform);
+					}
 				}
 			}
 		};
@@ -157,12 +168,12 @@ namespace Chaos::LowLevelTest
 			// Change transform to prevent early-out
 			Transform.SetTranslation(Transform.GetTranslation() + FVector(100, 0, 0));
 
-			if (TRigidScenePtr<FGameContextRW> Scene = SceneHandle.LockRW())
+			if (FRigidContextGameRW Context = SceneHandle.LockRW())
 			{
 				for (FRigidBodyHandle BodyHandle : BodyHandles)
 				{
 					// Pin does not add a lock - we are already in a scene lock
-					if (TRigidBodyPtr<FGameContextRW> Body = BodyHandle.Pin(Scene))
+					if (TRigidBodyPtr<FRigidContextGameRW> Body = BodyHandle.Pin(Context))
 					{
 						Body->InitTransform(Transform);
 					}
@@ -170,17 +181,16 @@ namespace Chaos::LowLevelTest
 			}
 		};
 
-		if (TRigidScenePtr<FGameContextRW> Scene = SceneHandle.LockRW())
+		if (FRigidContextGameRW Context = SceneHandle.LockRW())
 		{
 			for (FRigidBodyHandle BodyHandle : BodyHandles)
 			{
-				Scene->UnregisterRigidBody(BodyHandle);
-				Scene->DestroyRigidBody(BodyHandle);
+				Context->DestroyBody(BodyHandle.Pin(Context));
 			}
 		}
 		BodyHandles.Reset();
 
-		Factory.DestroyRigidScene(SceneHandle);
+		Factory.DestroyScene(SceneHandle);
 	}
 
 	TEST_CASE("ChaosAPIBenchmark BodyHandle Lock ProxyComparison", "[!benchmark][ChaosAPIBenchmark]")
@@ -235,4 +245,6 @@ namespace Chaos::LowLevelTest
 		Proxies.Reset();
 	}
 
-}
\ No newline at end of file
+}
+
+#endif
\ No newline at end of file
diff --git a/Engine/Plugins/Experimental/ChaosRigidPhysicsAsync/Tests/Private/SceneTests.cpp b/Engine/Plugins/Experimental/ChaosRigidPhysicsAsync/Tests/Private/SceneTests.cpp
--- a/Engine/Plugins/Experimental/ChaosRigidPhysicsAsync/Tests/Private/SceneTests.cpp
+++ b/Engine/Plugins/Experimental/ChaosRigidPhysicsAsync/Tests/Private/SceneTests.cpp
@@ -3,315 +3,358 @@
 #include "Chaos/Collision/ContactPoint.h"
 #include "Chaos/LowLevelTest/ChaosTestHarness.h"
 #include "Chaos/LowLevelTest/ChaosTestScene.h"
+#include "ChaosRigidPhysicsAsyncTest.h"
 #include "ChaosRigidPhysicsAsync/RigidFactoryAsync.h"
+#include "ChaosRigidPhysicsAsync/RigidBodyAsync.h"
+#include "Framework/Threading.h"
+#include "HAL/Thread.h"
 #include "RigidPhysics/RigidBody.h"
+#include "RigidPhysics/RigidContext.h"
 #include "RigidPhysics/RigidModifier.h"
+#include "RigidPhysics/RigidScene.h"
+
+#if UE_RIGIDPHYSICS_API_ENABLED
 
 namespace Chaos::LowLevelTest
 {
-	class FTestRigidSceneModifier : public UE::Physics::IRigidSceneModifier
+	using namespace Chaos;
+	using namespace Chaos::Rigids::Async;
+	using namespace UE::Physics;
+
+	TEST_CASE("SceneLockTest", "[Chaos][API][Scene][unit]")
 	{
-	public:
-		virtual void SimPreTick(UE::Physics::TRigidScenePtr<UE::Physics::FSimContextRW>& Context) override final
-		{
-			if (SimPreTickFunc)
+		FRigidFactoryAsync SceneFactory;
+		FRigidSceneHandle SceneHandle = SceneFactory.CreateScene(FRigidDebugName::Make(TEXT("Main")));
+
+		// Create a thread that will lock the physics scene.
+		// We initially lock ThreadRunningLock on the main thread to prevent the thread from executing.
+		// The main thread unlocks ThreadRunningLock after it has locked the scene, at which point the thread will try to lock the scene and will wait.
+		// After verifying that the thread is waiting, the main thread unlocks the scene and allows the thread to run.
+		// The thread then locks the scene and incrmement the ThreadLockCount.
+		// The main thread will then wait until ThreadRunningLock is unlocked again before exiting, by which time the thread should been able to lock the scene.
+		FCriticalSection ThreadRunningLock = FCriticalSection();
+		std::atomic_bool bThreadRunning = false;
+		std::atomic_int ThreadLockCount = 0;
+		const auto& ThreadFunc =
+			[SceneHandle, &ThreadRunningLock, &bThreadRunning, &ThreadLockCount]()
 			{
-				SimPreTickFunc(Context);
-			}
-		}
+#if PHYSICS_THREAD_CONTEXT
+				FGameThreadContextScope GTScope(true);
+#endif
 
-		TFunction<void(UE::Physics::TRigidScenePtr<UE::Physics::FSimContextRW>&)> SimPreTickFunc;
-	};
+				ThreadRunningLock.Lock();
 
+				bThreadRunning = true;
 
-	TEST_CASE("ChaosAsyncExample", "[Chaos][API][Scene][unit]")
-	{
-		using namespace Chaos;
-		using namespace Chaos::Rigids::Async;
-		using namespace UE::Physics;
-
-		float Dt = 1.0f / 60.0f;
+				if (FRigidContextGameRW Context = SceneHandle.LockRW())
+				{
+					++ThreadLockCount;
+				}
 
-		SECTION("CreateSceneAndBodyExample")
-		{
-			// Create an async physics scene
-			FRigidFactoryAsync Factory;
-			FRigidSceneHandle SceneHandle = Factory.CreateRigidScene(FRigidSceneSettingsAsync());
+				FPlatformProcess::YieldThread();
 
-			// Create a rigid body
-			FRigidBodyHandle BodyHandle = SceneHandle.LockRW()->CreateRigidBody(MakeRigidName(TEXT("Body")), ERigidMovementType::Dynamic);
+				ThreadRunningLock.Unlock();
+			};
 
-			// Not Allowed from ReadOnly scene handle
-			//FRigidBodyHandle BodyHandle2 = SceneHandle.LockRO().CreateRigidBody(MakeRigidName(TEXT("Body")), ERigidMovementType::Dynamic);
+		// Start the thread but make it wait until we lock the scene before it continues
+		ThreadRunningLock.Lock();
+		FThread Thread(TEXT("TestThread"), ThreadFunc);
 
-			// Set body position from the game thread
-			if (TRigidBodyPtr<FGameContextRW> BodyPtr = BodyHandle.LockRW())
-			{
-				BodyPtr->InitTransform(FTransform(FVector(0.0, 0.0, 100.0)));
-			}
+		// Lock the scene on the main thread and then start the background thread.
+		// The thread should be blocked on SceneHandle.LockRW() until we are done here. 
+		if (FRigidContextGameRW Context = SceneHandle.LockRW())
+		{
+			// Allow the background thread to run
+			ThreadRunningLock.Unlock();
 
-			// Tick physics for a while
-			for (int TickIndex = 0; TickIndex< 100; ++TickIndex)
+			// Wait until the thread is running
+			while (!bThreadRunning)
 			{
-				SceneHandle.LockRW()->Tick(Dt);
+				FPlatformProcess::YieldThread();
 			}
 
-			// Destroy the body
-			SceneHandle.LockRW()->DestroyRigidBody(BodyHandle);
+			// Wait a bit longer - the thread should not be able to get the physics lock
+			// TODO_CHAOSAPI: Can we check the thread/lock state to explicitly verify that the 
+			// background thread is waiting on the lock?
+			FPlatformProcess::Sleep(0.1f);
 
-			// Destroy the scene
-			Factory.DestroyRigidScene(SceneHandle);
+			// If this fails, then the above SceneHandle.LockRW() did not work
+			CHECK(ThreadLockCount == 0);
 		}
 
-		SECTION("CallbackExample")
-		{
-			// Create an async physics scene
-			FRigidFactoryAsync Factory;
-			FRigidSceneHandle SceneHandle = Factory.CreateRigidScene(FRigidSceneSettingsAsync());
-
-			// Create a rigid body
-			FRigidBodyHandle BodyHandle = SceneHandle.LockRW()->CreateRigidBody(MakeRigidName(TEXT("Body")), ERigidMovementType::Dynamic);
-			SceneHandle.LockRW()->RegisterRigidBody(BodyHandle);
-
-			// An example callback
-			FTestRigidSceneModifier SceneModifier;
-			SceneModifier.SimPreTickFunc =
-				[BodyHandle](TRigidScenePtr<FSimContextRW>& Scene)
-				{
-					// Pinning the handle with the specified scene restricts the API to only legal functions
-					if (TRigidBodyPtr<FSimContextRW> PTBodyPtr = BodyHandle.Pin(Scene))
-					{
-						// TODO: Still calls the GT functions so it asserts in SingleParticlePhysicsProxy
-						PTBodyPtr->InitTransform(FTransform(FVector(0.0, 0.0, 100.0)));
-					}
-				};
-
-			SceneHandle.LockRW()->RegisterModifier(&SceneModifier);
+		// Now the thread should get the lock and increment the counter
+		ThreadRunningLock.Lock();
 
-			// Tick physics for a while
-			for (int TickIndex = 0; TickIndex < 100; ++TickIndex)
-			{
-				SceneHandle.LockRW()->Tick(Dt);
-			}
+		// If this fails, then the thread never gained the lock
+		CHECK(ThreadLockCount > 0);
 
-			SceneHandle.LockRW()->UnregisterModifier(&SceneModifier);
+		ThreadRunningLock.Unlock();
 
-			// Destroy the body
-			SceneHandle.LockRW()->UnregisterRigidBody(BodyHandle);
-			SceneHandle.LockRW()->DestroyRigidBody(BodyHandle);
+		// Terminate the thread
+		Thread.Join();
 
-			// Destroy the scene
-			Factory.DestroyRigidScene(SceneHandle);
-		}
+		SceneFactory.DestroyScene(SceneHandle);
 	}
 
 
-	TEST_CASE("ChaosAsync", "[Chaos][API][Scene][unit]")
+	TEST_CASE("CreateSceneAndBodies", "[Chaos][API][Scene][unit]")
 	{
-		using namespace Chaos;
-		using namespace Chaos::Rigids::Async;
-		using namespace UE::Physics;
-
 		float Dt = 1.0f / 60.0f;
 
-		SECTION("CreateScene")
+		FRigidFactoryAsync SceneFactory;
+		FRigidSceneHandle SceneHandle = SceneFactory.CreateScene(FRigidDebugName::Make(TEXT("Main")));
+		REQUIRE(SceneHandle.LockRO());
+		if (FRigidContextGameRW Context = SceneHandle.LockRW())
 		{
-			FRigidFactoryAsync Factory;
-
-			FRigidSceneSettingsAsync SceneSettings;
-			FRigidSceneHandle SceneHandle = Factory.CreateRigidScene(SceneSettings);
+			REQUIRE(Context.GetScene().IsValid());
+		}
 
-			SECTION("TickScene")
+		SECTION("EmptyBodyHandlesAreInvalid")
+		{
+			// Empty RigidBodyHandle
+			FRigidBodyHandle BodyHandle;
+			if (FRigidContextGameRW Context = SceneHandle.LockRW())
 			{
-				SceneHandle.LockRW()->Tick(Dt);
-				SceneHandle.LockRW()->Tick(Dt);
+				// Handle should be invalid (explicit and implicit checks)
+				CHECK_FALSE(BodyHandle.Pin(Context).IsValid());
+				CHECK_FALSE(BodyHandle.Pin(Context));
 			}
+		}
 
-			SECTION("CreateAndDestroyBody")
+		SECTION("AddBodies")
+		{
+			FRigidBodyHandle StaticBodyHandle;
+			FRigidBodyHandle DynamicBodyHandle;
+			if (FRigidContextGameRW Context = SceneHandle.LockRW())
 			{
-				// Empty RigidBodyHandle
-				FRigidBodyHandle BodyHandle;
-
-				// Handle should be invalid (explicit and implicit checks)
-				CHECK(!BodyHandle.LockRO().IsValid());
-				CHECK(!BodyHandle.LockRO());
-
-				// Create a body
-				BodyHandle = SceneHandle.LockRW()->CreateRigidBody(MakeRigidName(TEXT("Body")), ERigidMovementType::Dynamic);
+				StaticBodyHandle = Context->CreateBody(FRigidDebugName::Make(TEXT("Body1")), ERigidMovementType::Static);
+				DynamicBodyHandle = Context->CreateBody(FRigidDebugName::Make(TEXT("Body2")), ERigidMovementType::Dynamic);
+				REQUIRE(StaticBodyHandle.Pin(Context));
+				REQUIRE(DynamicBodyHandle.Pin(Context));
+
+				// Check that we created objects of the correct type
+				CHECK(StaticBodyHandle.Pin(Context)->GetTypeId() == FRigidBodyAsyncGT::GetStaticTypeId());
+				CHECK(DynamicBodyHandle.Pin(Context)->GetTypeId() == FRigidBodyAsyncGT::GetStaticTypeId());
+
+				// Activate Bodies
+				StaticBodyHandle.Pin(Context)->Activate();
+				DynamicBodyHandle.Pin(Context)->Activate();
+				CHECK(StaticBodyHandle.Pin(Context)->IsActive());
+				CHECK(DynamicBodyHandle.Pin(Context)->IsActive());
+			}
 
-				// Body Handle should be valid
-				CHECK(BodyHandle.LockRO().IsValid());
-				CHECK(!!BodyHandle.LockRO());
+			SECTION("TickSystemWithBodies")
+			{
+				if (FRigidContextGameRW Context = SceneHandle.LockRW())
+				{
+					Context.GetScene()->Tick(Dt);
 
-				SceneHandle.LockRW()->Tick(Dt);
-				SceneHandle.LockRW()->Tick(Dt);
+					// Dynamic body should have moved under gravity
+					if (TRigidBodyPtr<FRigidContextGameRW> BodyPtr = DynamicBodyHandle.Pin(Context))
+					{
+						FVector BodyLocation = BodyPtr->GetTransform().GetLocation();
+						CHECK(BodyLocation.Z < 0.0f);
+					}
+				}
+			}
 
-				// Destroy the body
-				SceneHandle.LockRW()->DestroyRigidBody(BodyHandle);
+			SECTION("ChangeBodyState")
+			{
+				if (FRigidContextGameRW Context = SceneHandle.LockRW())
+				{
+					if (TRigidBodyPtr<FRigidContextGameRW> BodyPtr = DynamicBodyHandle.Pin(Context))
+					{
+						BodyPtr->InitTransform(FTransform(FVector(100,0,0)));
+						CHECK(BodyPtr->GetTransform().GetTranslation().X == 100.0f);
+					}
+				}
+			}
 
-				// RigidBody handle should is invalid (internally the container will 
-				// still be valid but the BodyId will not)
-				CHECK(!BodyHandle.LockRO().IsValid());
-				CHECK(!BodyHandle.LockRO());
+			if (FRigidContextGameRW Context = SceneHandle.LockRW())
+			{
+				// Deactivate Bodies
+				StaticBodyHandle.Pin(Context)->Deactivate();
+				DynamicBodyHandle.Pin(Context)->Deactivate();
+				CHECK_FALSE(StaticBodyHandle.Pin(Context)->IsActive());
+				CHECK_FALSE(DynamicBodyHandle.Pin(Context)->IsActive());
+
+				// Destroy Bodies
+				Context->DestroyBody(StaticBodyHandle.Pin(Context));
+				Context->DestroyBody(DynamicBodyHandle.Pin(Context));
+				CHECK_FALSE(StaticBodyHandle.Pin(Context));
+				CHECK_FALSE(DynamicBodyHandle.Pin(Context));
 			}
+		}
 
-			SECTION("BodyHandleDuplication")
+		SECTION("BodyHandleCopy")
+		{
+			if (FRigidContextGameRW Context = SceneHandle.LockRW())
 			{
 				// Create a body
-				FRigidBodyHandle BodyHandle = SceneHandle.LockRW()->CreateRigidBody(MakeRigidName(TEXT("Body")), ERigidMovementType::Dynamic);
-				CHECK(!!BodyHandle.LockRO());
+				FRigidBodyHandle BodyHandle = Context->CreateBody(FRigidDebugName::Make(TEXT("Body1")), ERigidMovementType::Dynamic);
+				CHECK(BodyHandle.Pin(Context));
 
-				// Create a duplicate body handle - it should also be valid
+				// Create a duplicate body handle - it should also be valid and reference the same body object
 				FRigidBodyHandle BodyHandleDupe = BodyHandle;
-				CHECK(!!BodyHandleDupe.LockRO());
+				CHECK(BodyHandleDupe.Pin(Context));
+				CHECK(BodyHandleDupe.Pin(Context).Get() == BodyHandle.Pin(Context).Get());
 
-				SceneHandle.LockRW()->Tick(Dt);
-				SceneHandle.LockRW()->Tick(Dt);
+				Context->Tick(Dt);
 
 				// Destroy the body
-				SceneHandle.LockRW()->DestroyRigidBody(BodyHandle);
+				Context->DestroyBody(BodyHandle.Pin(Context));
 
 				// Body handle should is invalid (references destroyed body)
-				CHECK(!BodyHandle.LockRO());
-				CHECK(!BodyHandleDupe.LockRO());
+				CHECK_FALSE(BodyHandle.Pin(Context));
+				CHECK_FALSE(BodyHandleDupe.Pin(Context));
 			}
+		}
 
-			SECTION("SetBodyTransform")
-			{
-				FRigidBodyHandle BodyHandle = SceneHandle.LockRW()->CreateRigidBody(MakeRigidName(TEXT("Body")), ERigidMovementType::Dynamic);
-				FRigidBodyHandle BodyHandleDupe = BodyHandle;
-
-				// Get write access to the body
-				if (TRigidBodyPtr<FGameContextRW> BodyPtr = BodyHandle.LockRW())
-				{
-					REQUIRE(!!BodyPtr);
-
-					// Body starts at the origin
-					FVector3d InitialBodyX = BodyPtr->GetTransform().GetTranslation();
-					CHECK(InitialBodyX == FVector(0,0,0));
-
-					// Move body
-					FVector3d NewBodyX = FVector(100,0,0);
-					BodyPtr->InitTransform(FTransform(NewBodyX));
-
-					// Check body has moved and the duplicated handle agrees
-					FVector3d FinalBodyX = BodyPtr->GetTransform().GetTranslation();
-					CHECK(FinalBodyX == NewBodyX);
-					CHECK(BodyHandleDupe.LockRO()->GetTransform().GetTranslation() == FinalBodyX);
-				}
-
-				SceneHandle.LockRW()->Tick(Dt);
-				SceneHandle.LockRW()->Tick(Dt);
+		SceneFactory.DestroyScene(SceneHandle);
+	} // TEST_CASE
 
-				// Destroy the body
-				SceneHandle.LockRW()->DestroyRigidBody(BodyHandle);
+	class FTestRigidSceneModifier : public IRigidSceneModifier
+	{
+	public:
+		virtual void PreSimulate(const FRigidContextGameRW& Context) override final
+		{
+			if (PreSimFunc)
+			{
+				PreSimFunc(Context);
 			}
+		}
 
-			SECTION("InvalidBodyHandleAssertsOnUse")
+		virtual void PreTick(const FRigidContextSimRW& Context) override final
+		{
+			if (PreTickFunc)
 			{
-				// Create a body and immediately destroy it
-				FRigidBodyHandle BodyHandle = SceneHandle.LockRW()->CreateRigidBody(MakeRigidName(TEXT("Body")), ERigidMovementType::Dynamic);
-				CHECK(!!BodyHandle.LockRO());
-
-				SceneHandle.LockRW()->Tick(Dt);
-
-				SceneHandle.LockRW()->DestroyRigidBody(BodyHandle);
-				CHECK(!BodyHandle.LockRO());
+				PreTickFunc(Context);
+			}
+		}
 
-				// Get write access to the body - will generate invalid handle
-				TRigidBodyPtr<FGameContextRW> BodyPtrRW = BodyHandle.LockRW();
-				CHECK(!BodyPtrRW);
+		TFunction<void(const FRigidContextGameRW&)> PreSimFunc;
+		TFunction<void(const FRigidContextSimRW&)> PreTickFunc;
+	};
 
-				// These calls assert but there's no way to safely verify that happens
-				// because we don't use exceptions and check() uses __debugbreak
-				//CHECK_THROWS(BodyPtrRW->GetTransform());
-				//CHECK_THROWS(BodyPtrRW->SetTransform(FTransform::Identity));
-			}
+	TEST_CASE("CallbacksAreCalled", "[Chaos][API][Scene][unit]")
+	{
+		float Dt = 1.0f / 60.0f;
 
-			SECTION("SceneModifierIsCalled")
+		std::atomic<int32> PreSimCallCount = 0;
+		std::atomic<int32> PreTickCallCount = 0;
+		FTestRigidSceneModifier SceneModifier;
+		SceneModifier.PreSimFunc = 
+			[&PreSimCallCount](const FRigidContextGameRW& Context)
 			{
-				std::atomic<int32> SceneModifierCallCount = 0;
-				FTestRigidSceneModifier SceneModifier;
-				SceneModifier.SimPreTickFunc =
-					[&SceneModifierCallCount](TRigidScenePtr<FSimContextRW>& Context)
-					{
-						++SceneModifierCallCount;
-					};
+				++PreSimCallCount;
+			};
+		SceneModifier.PreTickFunc =
+			[&PreTickCallCount](const FRigidContextSimRW& Context)
+			{
+				++PreTickCallCount;
+			};
+
+		FRigidFactoryAsync SceneFactory;
+		FRigidSceneHandle SceneHandle = SceneFactory.CreateScene(FRigidDebugName::Make(TEXT("Main")));
 
-				SceneHandle.LockRW()->RegisterModifier(&SceneModifier);
+		if (FRigidContextGameRW Context = SceneHandle.LockRW())
+		{
+			Context->RegisterModifier(&SceneModifier);
 
-				SceneHandle.LockRW()->Tick(Dt);
-				CHECK(SceneModifierCallCount == 1);
+			Context->Tick(Dt);
+			CHECK(PreSimCallCount == 1);
+			CHECK(PreTickCallCount == 1);
 
-				SceneHandle.LockRW()->Tick(Dt);
-				CHECK(SceneModifierCallCount == 2);
+			Context->Tick(Dt);
+			CHECK(PreSimCallCount == 2);
+			CHECK(PreTickCallCount == 2);
 
-				SceneHandle.LockRW()->UnregisterModifier(&SceneModifier);
+			Context->UnregisterModifier(&SceneModifier);
 
-				SceneHandle.LockRW()->Tick(Dt);
-				CHECK(SceneModifierCallCount == 2);
-			}
+			Context->Tick(Dt);
+			CHECK(PreSimCallCount == 2);
+			CHECK(PreTickCallCount == 2);
+		}
 
-			SECTION("SceneModifierBodyHandles")
-			{
-				const FVector3d GTBodyX = FVector(100, 0, 0);
-				const FVector3d PTBodyX = FVector(200, 0, 0);
+		SceneFactory.DestroyScene(SceneHandle);
+	};
 
-				FRigidBodyHandle BodyHandle = SceneHandle.LockRW()->CreateRigidBody(MakeRigidName(TEXT("Body")), ERigidMovementType::Dynamic);
-				SceneHandle.LockRW()->RegisterRigidBody(BodyHandle);
+	// Change the position of Dynamic bodies from GT and PT callbacks 
+	// and  verify that the changes are visible on the GT.
+	// NOTE: InitTransform is currently not supported on the PT for 
+	// Static and Kinematic bodies - the changes do not replicate
+	// back to the GT.
+	TEST_CASE("GameAndPhysicsThreadCallbacks", "[Chaos][API][Scene][unit]")
+	{
+		float Dt = 1.0f / 60.0f;
 
-				BodyHandle.LockRW()->InitTransform(FTransform(GTBodyX));
-				CHECK(BodyHandle.LockRO()->GetTransform().GetTranslation() == GTBodyX);
+		FRigidBodyHandle BodyHandle1;
+		FRigidBodyHandle BodyHandle2;
+		FTestRigidSceneModifier SceneModifier;
 
-				// We don't want to particles to move on their own...
-				BodyHandle.LockRW()->SetGravityScale(0.0f);
+		// GT Callback - Change the position of Body1
+		SceneModifier.PreSimFunc =
+			[&BodyHandle1](const FRigidContextGameRW& Context)
+			{
+				CHECK(BodyHandle1.Pin(Context));
+				if (TRigidBodyPtr<FRigidContextGameRW> BodyPtr = BodyHandle1.Pin(Context))
+				{
+					// Make sure this is a GT body
+					REQUIRE(BodyPtr->GetTypeId() == FRigidBodyAsyncGT::GetStaticTypeId());
 
-				// A BeginTick SceneModifier that changes the body position to PTBodyX
-				FTestRigidSceneModifier SceneModifier;
-				SceneModifier.SimPreTickFunc =
-					[BodyHandle, GTBodyX, PTBodyX](TRigidScenePtr<FSimContextRW>& Scene)
-					{
-						TRigidBodyPtr<FSimContextRW> PTBodyPtr = BodyHandle.Pin(Scene);
-						CHECK(PTBodyPtr.IsValid());
+					BodyPtr->InitTransform(FTransform(FVector(200.0, 0.0, 0.0)));
+					CHECK(BodyPtr->GetTransform().GetLocation().X == 200.0);
+				}
+			};
 
-						if (PTBodyPtr.IsValid())
-						{
-							// Make sure GT position change is reflected on PT
-							CHECK(PTBodyPtr->GetTransform().GetTranslation() == GTBodyX);
+		// PT Callback - Change the position of Body2
+		SceneModifier.PreTickFunc =
+			[&BodyHandle2](const FRigidContextSimRW& Context)
+			{
+				CHECK(BodyHandle2.Pin(Context));
+				if (TRigidBodyPtr<FRigidContextSimRW> BodyPtr = BodyHandle2.Pin(Context))
+				{
+					// Make sure this is a PT body
+					REQUIRE(BodyPtr->GetTypeId() == FRigidBodyAsyncPT::GetStaticTypeId());
 
-							// Change the position of the body
-							PTBodyPtr->InitTransform(FTransform(PTBodyX));
-							CHECK(PTBodyPtr->GetTransform().GetTranslation() == PTBodyX);
-						}
-					};
+					BodyPtr->InitTransform(FTransform(FVector(100.0, 0.0, 0.0)));
+					CHECK(BodyPtr->GetTransform().GetLocation().X == 100.0);
+				}
+			};
 
-				SceneHandle.LockRW()->RegisterModifier(&SceneModifier);
+		// Create an async physics scene
+		FRigidFactoryAsync SceneFactory;
+		FRigidSceneHandle SceneHandle = SceneFactory.CreateScene(FRigidDebugName::Make(TEXT("Main")));
 
-				// Tick the scene - modifier will be called
-				SceneHandle.LockRW()->Tick(Dt);
+		if (FRigidContextGameRW Context = SceneHandle.LockRW())
+		{
+			Context->RegisterModifier(&SceneModifier);
 
-				// Modifier should have moved the body, and the results
-				// should have synched back to the GT
-				CHECK(BodyHandle.LockRO()->GetTransform().GetTranslation() == PTBodyX);
+			TRigidBodyPtr<FRigidContextGameRW> BodyPtr1 = Context->CreateBody(FRigidDebugName::Make(TEXT("Body1")), ERigidMovementType::Dynamic);
+			BodyPtr1->Activate();
+			BodyHandle1 = BodyPtr1;
 
-				// Put the body back
-				BodyHandle.LockRW()->InitTransform(FTransform(GTBodyX));
-				CHECK(BodyHandle.LockRO()->GetTransform().GetTranslation() == GTBodyX);
+			TRigidBodyPtr<FRigidContextGameRW> BodyPtr2 = Context->CreateBody(FRigidDebugName::Make(TEXT("Body2")), ERigidMovementType::Dynamic);
+			BodyPtr2->Activate();
+			BodyHandle2 = BodyPtr2;
 
-				// Tick the scene - modifier will be called
-				SceneHandle.LockRW()->Tick(Dt);
+			CHECK(BodyPtr1->GetTransform().GetLocation().X == 0.0);
+			CHECK(BodyPtr2->GetTransform().GetLocation().X == 0.0);
 
-				// Modifier should have moved the body again
-				CHECK(BodyHandle.LockRO()->GetTransform().GetTranslation() == PTBodyX);
+			Context->Tick(Dt);
 
-				// Clean up
-				SceneHandle.LockRW()->UnregisterModifier(&SceneModifier);
-				SceneHandle.LockRW()->UnregisterRigidBody(BodyHandle);
-				SceneHandle.LockRW()->DestroyRigidBody(BodyHandle);
-			}
+			// Both bodies were moved in callbacks
+			CHECK(BodyPtr1->GetTransform().GetLocation().X == 200.0);
+			CHECK(BodyPtr2->GetTransform().GetLocation().X == 100.0);
 
-			Factory.DestroyRigidScene(SceneHandle);
+			Context->DestroyBody(MoveTemp(BodyPtr1));
+			Context->DestroyBody(MoveTemp(BodyPtr2));
 		}
+
+		// Destroy the scene
+		SceneFactory.DestroyScene(SceneHandle);
 	}
 
 } // Chaos::LowLevelTest
+
+#endif
\ No newline at end of file
diff --git a/Engine/Plugins/Experimental/ChaosRigidPhysicsAsync/Tests/Private/TypeIdTests.cpp b/Engine/Plugins/Experimental/ChaosRigidPhysicsAsync/Tests/Private/TypeIdTests.cpp
new file mode 100644
--- /dev/null
+++ b/Engine/Plugins/Experimental/ChaosRigidPhysicsAsync/Tests/Private/TypeIdTests.cpp
@@ -0,0 +1,167 @@
+// Copyright Epic Games, Inc. All Rights Reserved.
+
+#include "Chaos/Collision/ContactPoint.h"
+#include "Chaos/LowLevelTest/ChaosTestHarness.h"
+#include "Chaos/LowLevelTest/ChaosTestScene.h"
+#include "ChaosRigidPhysicsAsyncTest.h"
+
+#include <catch2/matchers/catch_matchers_string.hpp>
+
+#if UE_RIGIDPHYSICS_API_ENABLED
+
+namespace Chaos::LowLevelTest
+{
+	using namespace Chaos;
+	using namespace UE::Physics;
+
+	class FTestTypeIdBase : public IRigidTyped
+	{
+	public:
+		UE_RIGIDPHYSICS_RIGIDTYPED_DECL(CHAOSRIGIDPHYSICSASYNC_API, FTestTypeIdBase);
+	};
+
+	class FTestTypeIdDerived1 : public FTestTypeIdBase
+	{
+	public:
+		UE_RIGIDPHYSICS_RIGIDTYPED_DECL(CHAOSRIGIDPHYSICSASYNC_API, FTestTypeIdDerived1);
+	};
+
+	class FTestTypeIdDerived2 : public FTestTypeIdBase
+	{
+	public:
+		UE_RIGIDPHYSICS_RIGIDTYPED_DECL(CHAOSRIGIDPHYSICSASYNC_API, FTestTypeIdDerived2);
+	};
+
+	class FTestTypeIdLeaf1A : public FTestTypeIdDerived1
+	{
+	public:
+		UE_RIGIDPHYSICS_RIGIDTYPED_DECL(CHAOSRIGIDPHYSICSASYNC_API, FTestTypeIdLeaf1A);
+	};
+
+	class FTestTypeIdLeaf1B : public FTestTypeIdDerived1
+	{
+	public:
+		UE_RIGIDPHYSICS_RIGIDTYPED_DECL(CHAOSRIGIDPHYSICSASYNC_API, FTestTypeIdLeaf1B);
+	};
+
+	UE_RIGIDPHYSICS_RIGIDTYPED_IMPL(FTestTypeIdBase, IRigidTyped);
+	UE_RIGIDPHYSICS_RIGIDTYPED_IMPL(FTestTypeIdDerived1, FTestTypeIdBase);
+	UE_RIGIDPHYSICS_RIGIDTYPED_IMPL(FTestTypeIdDerived2, FTestTypeIdBase);
+	UE_RIGIDPHYSICS_RIGIDTYPED_IMPL(FTestTypeIdLeaf1A, FTestTypeIdDerived1);
+	UE_RIGIDPHYSICS_RIGIDTYPED_IMPL(FTestTypeIdLeaf1B, FTestTypeIdDerived1);
+
+	// Verify that up/downcasting succeeds and fails as expected
+	TEST_CASE("TypeIdTests", "[Chaos][API][TypeId][unit]")
+	{
+		// Class Hierarchy:
+		// IRigidTyped
+		//	+- FTestTypeIdBase
+		//		+- FTestTypeIdDerived1
+		//		| +- FTestTypeIdLeaf1A
+		//		| +- FTestTypeIdLeaf1B
+		//		+- FTestTypeIdDerived2
+		FTestTypeIdBase Base;
+		FTestTypeIdDerived1 Derived1;
+		FTestTypeIdDerived2 Derived2;
+		FTestTypeIdLeaf1A Leaf1A;
+		FTestTypeIdLeaf1B Leaf1B;
+
+		const FTestTypeIdBase* BasePtr = &Base;
+		const FTestTypeIdBase* Derived1BasePtr = &Derived1;
+		const FTestTypeIdBase* Derived2BasePtr = &Derived2;
+		const FTestTypeIdBase* Leaf1ABasePtr = &Leaf1A;
+		const FTestTypeIdBase* Leaf1BBasePtr = &Leaf1B;
+
+		const FTestTypeIdDerived1* Leaf1ADerived1Ptr = &Leaf1A;
+
+		// Check type names (generated in the macro)
+		CHECK_THAT(ToStdString(BasePtr->GetTypeId().GetTypeName()), Catch::Matchers::Equals("FTestTypeIdBase"));
+		CHECK_THAT(ToStdString(Derived1BasePtr->GetTypeId().GetTypeName()), Catch::Matchers::Equals("FTestTypeIdDerived1"));
+		CHECK_THAT(ToStdString(Derived2BasePtr->GetTypeId().GetTypeName()), Catch::Matchers::Equals("FTestTypeIdDerived2"));
+		CHECK_THAT(ToStdString(Leaf1ABasePtr->GetTypeId().GetTypeName()), Catch::Matchers::Equals("FTestTypeIdLeaf1A"));
+		CHECK_THAT(ToStdString(Leaf1BBasePtr->GetTypeId().GetTypeName()), Catch::Matchers::Equals("FTestTypeIdLeaf1B"));
+
+		// Type checks from known leaf type
+		CHECK(Base.IsA<IRigidTyped>());
+		CHECK(Base.IsA<FTestTypeIdBase>());
+		CHECK_FALSE(Base.IsA<FTestTypeIdDerived1>());
+		CHECK_FALSE(Base.IsA<FTestTypeIdDerived2>());
+		CHECK_FALSE(Base.IsA<FTestTypeIdLeaf1A>());
+		CHECK_FALSE(Base.IsA<FTestTypeIdLeaf1B>());
+
+		CHECK(Derived1.IsA<IRigidTyped>());
+		CHECK(Derived1.IsA<FTestTypeIdBase>());
+		CHECK(Derived1.IsA<FTestTypeIdDerived1>());
+		CHECK_FALSE(Derived1.IsA<FTestTypeIdDerived2>());
+		CHECK_FALSE(Derived1.IsA<FTestTypeIdLeaf1A>());
+		CHECK_FALSE(Derived1.IsA<FTestTypeIdLeaf1B>());
+
+		CHECK(Derived2.IsA<IRigidTyped>());
+		CHECK(Derived2.IsA<FTestTypeIdBase>());
+		CHECK_FALSE(Derived2.IsA<FTestTypeIdDerived1>());
+		CHECK(Derived2.IsA<FTestTypeIdDerived2>());
+		CHECK_FALSE(Derived2.IsA<FTestTypeIdLeaf1A>());
+		CHECK_FALSE(Derived2.IsA<FTestTypeIdLeaf1B>());
+
+		CHECK(Leaf1A.IsA<IRigidTyped>());
+		CHECK(Leaf1A.IsA<FTestTypeIdBase>());
+		CHECK(Leaf1A.IsA<FTestTypeIdDerived1>());
+		CHECK_FALSE(Leaf1A.IsA<FTestTypeIdDerived2>());
+		CHECK(Leaf1A.IsA<FTestTypeIdLeaf1A>());
+		CHECK_FALSE(Leaf1A.IsA<FTestTypeIdLeaf1B>());
+
+		CHECK(Leaf1B.IsA<IRigidTyped>());
+		CHECK(Leaf1B.IsA<FTestTypeIdBase>());
+		CHECK(Leaf1B.IsA<FTestTypeIdDerived1>());
+		CHECK_FALSE(Leaf1B.IsA<FTestTypeIdDerived2>());
+		CHECK_FALSE(Leaf1B.IsA<FTestTypeIdLeaf1A>());
+		CHECK(Leaf1B.IsA<FTestTypeIdLeaf1B>());
+
+		// Type checks from base pointer type
+
+		CHECK(BasePtr->IsA<IRigidTyped>());
+		CHECK(BasePtr->IsA<FTestTypeIdBase>());
+		CHECK_FALSE(BasePtr->IsA<FTestTypeIdDerived1>());
+		CHECK_FALSE(BasePtr->IsA<FTestTypeIdDerived2>());
+		CHECK_FALSE(BasePtr->IsA<FTestTypeIdLeaf1A>());
+		CHECK_FALSE(BasePtr->IsA<FTestTypeIdLeaf1B>());
+
+		CHECK(Derived1BasePtr->IsA<IRigidTyped>());
+		CHECK(Derived1BasePtr->IsA<FTestTypeIdBase>());
+		CHECK(Derived1BasePtr->IsA<FTestTypeIdDerived1>());
+		CHECK_FALSE(Derived1BasePtr->IsA<FTestTypeIdDerived2>());
+		CHECK_FALSE(Derived1BasePtr->IsA<FTestTypeIdLeaf1A>());
+		CHECK_FALSE(Derived1BasePtr->IsA<FTestTypeIdLeaf1B>());
+
+		CHECK(Derived2BasePtr->IsA<IRigidTyped>());
+		CHECK(Derived2BasePtr->IsA<FTestTypeIdBase>());
+		CHECK_FALSE(Derived2BasePtr->IsA<FTestTypeIdDerived1>());
+		CHECK(Derived2BasePtr->IsA<FTestTypeIdDerived2>());
+		CHECK_FALSE(Derived2BasePtr->IsA<FTestTypeIdLeaf1A>());
+		CHECK_FALSE(Derived2BasePtr->IsA<FTestTypeIdLeaf1B>());
+
+		CHECK(Leaf1ABasePtr->IsA<IRigidTyped>());
+		CHECK(Leaf1ABasePtr->IsA<FTestTypeIdBase>());
+		CHECK(Leaf1ABasePtr->IsA<FTestTypeIdDerived1>());
+		CHECK_FALSE(Leaf1ABasePtr->IsA<FTestTypeIdDerived2>());
+		CHECK(Leaf1ABasePtr->IsA<FTestTypeIdLeaf1A>());
+		CHECK_FALSE(Leaf1ABasePtr->IsA<FTestTypeIdLeaf1B>());
+
+		CHECK(Leaf1BBasePtr->IsA<IRigidTyped>());
+		CHECK(Leaf1BBasePtr->IsA<FTestTypeIdBase>());
+		CHECK(Leaf1BBasePtr->IsA<FTestTypeIdDerived1>());
+		CHECK_FALSE(Leaf1BBasePtr->IsA<FTestTypeIdDerived2>());
+		CHECK_FALSE(Leaf1BBasePtr->IsA<FTestTypeIdLeaf1A>());
+		CHECK(Leaf1BBasePtr->IsA<FTestTypeIdLeaf1B>());
+
+		// Type checks from intermediate pointer
+		CHECK(Leaf1ADerived1Ptr->IsA<IRigidTyped>());
+		CHECK(Leaf1ADerived1Ptr->IsA<FTestTypeIdBase>());
+		CHECK(Leaf1ADerived1Ptr->IsA<FTestTypeIdDerived1>());
+		CHECK_FALSE(Leaf1ADerived1Ptr->IsA<FTestTypeIdDerived2>());
+		CHECK(Leaf1ADerived1Ptr->IsA<FTestTypeIdLeaf1A>());
+		CHECK_FALSE(Leaf1ADerived1Ptr->IsA<FTestTypeIdLeaf1B>());
+	}
+}
+
+#endif
\ No newline at end of file
diff --git a/Engine/Source/Runtime/Experimental/Chaos/Private/Chaos/ChaosDebugName.cpp b/Engine/Source/Runtime/Experimental/Chaos/Private/Chaos/ChaosDebugName.cpp
--- a/Engine/Source/Runtime/Experimental/Chaos/Private/Chaos/ChaosDebugName.cpp
+++ b/Engine/Source/Runtime/Experimental/Chaos/Private/Chaos/ChaosDebugName.cpp
@@ -4,32 +4,5 @@
 
 namespace Chaos
 {
-	FString FSharedDebugName::DefaultName = FString("NoName");
-
-#if CHAOS_DEBUG_NAME
-	FSharedDebugName::FSharedDebugName(const FString& S)
-		: Name(TSharedPtr<FString, ESPMode::ThreadSafe>(new FString(S)))
-	{
-	}
-
-	FSharedDebugName::FSharedDebugName(FString&& S)
-		: Name(TSharedPtr<FString, ESPMode::ThreadSafe>(new FString(MoveTemp(S))))
-	{
-	}
-
-	inline bool FSharedDebugName::IsValid() const
-	{ 
-		return Name.IsValid();
-	}
-
-	const FString& FSharedDebugName::Value() const
-	{
-		if (Name.IsValid())
-		{
-			return *Name.Get();
-		}
-		return DefaultName;
-	}
-
-#endif
+	TSharedPtr<FString, ESPMode::ThreadSafe> FSharedDebugName::DefaultName = TSharedPtr<FString, ESPMode::ThreadSafe>(new FString(TEXT("NoName")));
 }
\ No newline at end of file
diff --git a/Engine/Source/Runtime/Experimental/Chaos/Public/Chaos/ChaosDebugName.h b/Engine/Source/Runtime/Experimental/Chaos/Public/Chaos/ChaosDebugName.h
--- a/Engine/Source/Runtime/Experimental/Chaos/Public/Chaos/ChaosDebugName.h
+++ b/Engine/Source/Runtime/Experimental/Chaos/Public/Chaos/ChaosDebugName.h
@@ -16,6 +16,19 @@ namespace Chaos
 	class FSharedDebugName
 	{
 	public:
+		static FSharedDebugName Make(const TCHAR* InName)
+		{
+			return FSharedDebugName(InName);
+		}
+
+		template<typename... ArgTypes>
+		static FSharedDebugName Make(const TCHAR* InFmt, ArgTypes... Args)
+		{
+			FStringFormatOrderedArguments ArgsArray = { std::forward<ArgTypes>(Args)... };
+			return FSharedDebugName(FString::Format(InFmt, ArgsArray));
+		}
+
+
 		FSharedDebugName() = default;
 		FSharedDebugName(const FSharedDebugName& Other) = default;
 		FSharedDebugName(FSharedDebugName&& Other) = default;
@@ -27,19 +40,50 @@ namespace Chaos
 
 		bool IsValid() const;
 		const FString& Value() const;
+		const TSharedPtr<FString, ESPMode::ThreadSafe>& SharedValue() const;
 
 	private:
 #if CHAOS_DEBUG_NAME
 		TSharedPtr<FString, ESPMode::ThreadSafe> Name;
 #endif
-		static UE_API FString DefaultName;
+		static UE_API TSharedPtr<FString, ESPMode::ThreadSafe> DefaultName;
 	};
 
-#if !CHAOS_DEBUG_NAME
+#if CHAOS_DEBUG_NAME
+	inline FSharedDebugName::FSharedDebugName(const FString& S)
+		: Name(TSharedPtr<FString, ESPMode::ThreadSafe>(new FString(S)))
+	{
+	}
+
+	inline FSharedDebugName::FSharedDebugName(FString&& S)
+		: Name(TSharedPtr<FString, ESPMode::ThreadSafe>(new FString(MoveTemp(S))))
+	{
+	}
+
+	inline bool FSharedDebugName::IsValid() const
+	{
+		return Name.IsValid();
+	}
+
+	inline const FString& FSharedDebugName::Value() const
+	{
+		if (Name.IsValid())
+		{
+			return *Name.Get();
+		}
+		return *DefaultName.Get();
+	}
+
+	inline const TSharedPtr<FString, ESPMode::ThreadSafe>& FSharedDebugName::SharedValue() const
+	{
+		return Name;
+	}
+#else
 	inline FSharedDebugName::FSharedDebugName(const FString& S) {}
 	inline FSharedDebugName::FSharedDebugName(FString&& S) {}
 	inline bool FSharedDebugName::IsValid() const { return false; }
-	inline const FString& FSharedDebugName::Value() const { return DefaultName; }
+	inline const FString& FSharedDebugName::Value() const { return *DefaultName.Get(); }
+	inline const TSharedPtr<FString, ESPMode::ThreadSafe>& FSharedDebugName::SharedValue() const { return DefaultName; }
 #endif
 }
 
diff --git a/Engine/Source/Runtime/Experimental/RigidPhysics/Private/RigidPhysics/RigidBody.cpp b/Engine/Source/Runtime/Experimental/RigidPhysics/Private/RigidPhysics/RigidBody.cpp
new file mode 100644
--- /dev/null
+++ b/Engine/Source/Runtime/Experimental/RigidPhysics/Private/RigidPhysics/RigidBody.cpp
@@ -0,0 +1,12 @@
+// Copyright Epic Games, Inc. All Rights Reserved.
+#include "RigidPhysics/RigidBody.h"
+#include "RigidPhysics/RigidLog.h"
+
+#if UE_RIGIDPHYSICS_API_ENABLED
+
+namespace UE::Physics
+{
+	UE_RIGIDPHYSICS_RIGIDTYPED_IMPL(IRigidBody, IRigidTyped);
+}
+
+#endif // UE_RIGIDPHYSICS_API_ENABLED
\ No newline at end of file
diff --git a/Engine/Source/Runtime/Experimental/RigidPhysics/Private/RigidPhysics/RigidBodyContainer.cpp b/Engine/Source/Runtime/Experimental/RigidPhysics/Private/RigidPhysics/RigidBodyContainer.cpp
new file mode 100644
--- /dev/null
+++ b/Engine/Source/Runtime/Experimental/RigidPhysics/Private/RigidPhysics/RigidBodyContainer.cpp
@@ -0,0 +1,27 @@
+// Copyright Epic Games, Inc. All Rights Reserved.
+#include "RigidPhysics/RigidBodyContainer.h"
+
+#if UE_RIGIDPHYSICS_API_ENABLED
+
+namespace UE::Physics
+{
+	UE_RIGIDPHYSICS_RIGIDTYPED_IMPL(IRigidBodyContainer, IRigidTyped);
+
+
+	FRigidBodyContainer::FRigidBodyContainer()
+	{
+		// TODO_CHAOSAPI: Make slack a parameter
+		BodyRegistry = MakeUnique<FRigidBodyRegistry>(10);
+	}
+
+	FRigidBodyContainer::~FRigidBodyContainer()
+	{
+	}
+
+	IRigidBody* FRigidBodyContainer::GetBody(const FRigidObjectId& InBodyId) const
+	{
+		return BodyRegistry->FindObject(InBodyId);
+	}
+}
+
+#endif // UE_RIGIDPHYSICS_API_ENABLED
\ No newline at end of file
diff --git a/Engine/Source/Runtime/Experimental/RigidPhysics/Private/RigidPhysics/RigidLockable.cpp b/Engine/Source/Runtime/Experimental/RigidPhysics/Private/RigidPhysics/RigidLockable.cpp
new file mode 100644
--- /dev/null
+++ b/Engine/Source/Runtime/Experimental/RigidPhysics/Private/RigidPhysics/RigidLockable.cpp
@@ -0,0 +1,11 @@
+// Copyright Epic Games, Inc. All Rights Reserved.
+#include "RigidPhysics/RigidLockable.h"
+
+#if UE_RIGIDPHYSICS_API_ENABLED
+
+namespace UE::Physics
+{
+	UE_RIGIDPHYSICS_RIGIDTYPED_IMPL(IRigidLockable, IRigidTyped);
+}
+
+#endif // UE_RIGIDPHYSICS_API_ENABLED
\ No newline at end of file
diff --git a/Engine/Source/Runtime/Experimental/RigidPhysics/Private/RigidPhysics/RigidLog.cpp b/Engine/Source/Runtime/Experimental/RigidPhysics/Private/RigidPhysics/RigidLog.cpp
new file mode 100644
--- /dev/null
+++ b/Engine/Source/Runtime/Experimental/RigidPhysics/Private/RigidPhysics/RigidLog.cpp
@@ -0,0 +1,22 @@
+// Copyright Epic Games, Inc. All Rights Reserved.
+#include "RigidPhysics/RigidLog.h"
+
+#if UE_RIGIDPHYSICS_API_ENABLED
+
+DEFINE_LOG_CATEGORY(LogRigidPhysics);
+
+namespace UE::Physics::Debugger
+{
+	// Used in natvis for FRigidObjectId.
+	// If true, will show <Index.Epoch> for an object Id. E.g.: Body<1.0>
+	// If false, will show <Index> for an object Id. E.g.: Body<1>
+	extern bool GShowIdEpoch = false;
+
+	// Used in natvis for FRigidBodyContainerHandle.
+	// Is true when a BodyContainerHandle contains a SceneHandle, which is
+	// required for natvis to be able to reach the object referenced by the handle
+	// (in the runtime, the scene is provided by the Context). 
+	extern bool GHaveSceneHandle = UE_RIGIDPHYSICS_CONTAINERSCENEHANDLE_ENABLED;
+}
+
+#endif // UE_RIGIDPHYSICS_API_ENABLED
\ No newline at end of file
diff --git a/Engine/Source/Runtime/Experimental/RigidPhysics/Private/RigidPhysics/RigidScene.cpp b/Engine/Source/Runtime/Experimental/RigidPhysics/Private/RigidPhysics/RigidScene.cpp
--- a/Engine/Source/Runtime/Experimental/RigidPhysics/Private/RigidPhysics/RigidScene.cpp
+++ b/Engine/Source/Runtime/Experimental/RigidPhysics/Private/RigidPhysics/RigidScene.cpp
@@ -2,53 +2,64 @@
 #include "RigidPhysics/RigidScene.h"
 #include "RigidPhysics/RigidBody.h"
 #include "RigidPhysics/RigidBodyContainer.h"
+#include "RigidPhysics/RigidContext.h"
 #include "RigidPhysics/RigidLog.h"
 
 #if UE_RIGIDPHYSICS_API_ENABLED
 
-DEFINE_LOG_CATEGORY(LogChaosRB);
-
 namespace UE::Physics
 {
-	template class TGameContext<EContextLockType::None>;
-	template class TGameContext<EContextLockType::ReadOnly>;
-	template class TGameContext<EContextLockType::ReadWrite>;
+	UE_RIGIDPHYSICS_RIGIDTYPED_IMPL(IRigidScene, IRigidLockable);
 
-	IRigidScene::IRigidScene()
+	FRigidSceneHandle::FRigidSceneHandle()
+		: SceneRaw(nullptr)
 	{
 	}
 
-	IRigidScene::~IRigidScene()
+	FRigidSceneHandle::FRigidSceneHandle(FRigidScene* InScene)
+		: SceneRaw(InScene)
+#if UE_RIGIDPHYSICS_SCENEHANDLE_SAFTEY_ENABLED
+		, SceneWeak(InScene)
+#endif
 	{
 	}
 
-
-	IRigidBodyContainer::IRigidBodyContainer(const FRigidBodyContainerTypeId& InTypeId)
-		: TypeId(InTypeId)
+	void FRigidSceneHandle::Reset()
 	{
+		SceneRaw = nullptr;
+#if UE_RIGIDPHYSICS_SCENEHANDLE_SAFTEY_ENABLED
+		SceneWeak.Reset();
+#endif
 	}
 
-	IRigidBodyContainer::~IRigidBodyContainer()
+	IRigidScene* FRigidSceneHandle::Get() const
 	{
+		return SceneRaw;
 	}
 
-	const FRigidBodyContainerTypeId& IRigidBodyContainer::GetTypeId() const
+	FRigidContextGameRO FRigidSceneHandle::LockRO() const
 	{
-		return TypeId;
+		return Lock<FRigidContextGameRO>();
 	}
 
-	bool IRigidBodyContainer::IsCastableTo(const FRigidBodyContainerTypeId& InTypeId) const
+	FRigidContextGameRW FRigidSceneHandle::LockRW() const
 	{
-		return GetTypeId() == InTypeId;
+		return Lock<FRigidContextGameRW>();
 	}
 
+	FRigidScene::FRigidScene()
+	{
+		// TODO_CHAOSAPI: registry slack should be a parameter
+		BodyContainerRegistry = MakeUnique<FRigidBodyContainerRegistry>(10);
+	}
 
-	IRigidBody::IRigidBody()
+	FRigidScene::~FRigidScene()
 	{
 	}
 
-	IRigidBody::~IRigidBody()
+	IRigidBodyContainer* FRigidScene::GetBodyContainer(const FRigidObjectId& InBodyContainerId) const
 	{
+		return BodyContainerRegistry->FindObject(InBodyContainerId);
 	}
 }
 
diff --git a/Engine/Source/Runtime/Experimental/RigidPhysics/Private/RigidPhysics/RigidTyped.cpp b/Engine/Source/Runtime/Experimental/RigidPhysics/Private/RigidPhysics/RigidTyped.cpp
new file mode 100644
--- /dev/null
+++ b/Engine/Source/Runtime/Experimental/RigidPhysics/Private/RigidPhysics/RigidTyped.cpp
@@ -0,0 +1,29 @@
+// Copyright Epic Games, Inc. All Rights Reserved.
+#include "RigidPhysics/RigidTyped.h"
+
+#if UE_RIGIDPHYSICS_API_ENABLED
+
+namespace UE::Physics
+{
+	IRigidTyped::IRigidTyped()
+	{
+	}
+
+	IRigidTyped::~IRigidTyped()
+	{
+	}
+
+	const FRigidTypeId& IRigidTyped::GetStaticTypeId()
+	{
+		static const TCHAR* StaticTypeName = TEXT("IRigidTyped");
+		static FRigidTypeId StaticTypeId = FRigidTypeId(StaticTypeName, nullptr);
+		return StaticTypeId;
+	}
+
+	const FRigidTypeId& IRigidTyped::GetTypeId() const
+	{
+		return IRigidTyped::GetStaticTypeId();
+	}
+}
+
+#endif // UE_RIGIDPHYSICS_API_ENABLED
\ No newline at end of file
diff --git a/Engine/Source/Runtime/Experimental/RigidPhysics/Public/RigidPhysics/RigidAccessor.h b/Engine/Source/Runtime/Experimental/RigidPhysics/Public/RigidPhysics/RigidAccessor.h
deleted file mode 100644
diff --git a/Engine/Source/Runtime/Experimental/RigidPhysics/Public/RigidPhysics/RigidBody.h b/Engine/Source/Runtime/Experimental/RigidPhysics/Public/RigidPhysics/RigidBody.h
--- a/Engine/Source/Runtime/Experimental/RigidPhysics/Public/RigidPhysics/RigidBody.h
+++ b/Engine/Source/Runtime/Experimental/RigidPhysics/Public/RigidPhysics/RigidBody.h
@@ -6,124 +6,90 @@
 #if UE_RIGIDPHYSICS_API_ENABLED
 
 #include "Math/Bounds.h"
-#include "Misc/Crc.h"
+#include "Math/Transform.h"
 #include "RigidPhysics/RigidBodyContainer.h"
 #include "RigidPhysics/RigidLog.h"
+#include "RigidPhysics/RigidObjectId.h"
+#include "RigidPhysics/RigidObjectPtr.h"
 #include "RigidPhysics/RigidScene.h"
+#include "RigidPhysics/RigidTyped.h"
 #include "Templates/SharedPointer.h"
 
-// NOTES
-// 
-// Using SharedPtr and WeakPtr in a few places where they are technically required, 
-// but in practice may not be (e.g., FBodyHandle). However this is hidden as
-// an implementation detail so we have the option to switch to raw pointers if we
-// can guarantee lifetime, or retain pointer tracking in non-ship builds, etc.
-//
-
-extern uint32 GetTypeHash(UE::Physics::FRigidBodyId& RigidBodyId);
-
 namespace UE::Physics
 {
-	// An ID for a rigid body that is guaranteed to be unique and never reused
-	// within its container. There is no guarantee of uniqueness between
-	// containers, which means you need a Conainer + BodyId to access a body.
-	// (see FRigidBodyHandle, and FRigidBodyContainerHandle).
-	class FRigidBodyId
+	// The full rigid body API. There will be several implementations of this API. E.g., synchronous
+	// physics, asynchronous physics, immediate physics, remote physics, etc.
+	// NOTE: for some implementations like async physics, there are both Game- and Physics-Context
+	// implementations within the async implementation. Pinning a handle will return the object 
+	// appropriate for the context.
+	class UE_INTERNAL IRigidBody: public IRigidTyped
 	{
 	public:
-		FRigidBodyId()
-			: IntId(INDEX_NONE)
+		UE_RIGIDPHYSICS_RIGIDTYPED_DECL(RIGIDPHYSICS_API, IRigidBody);
+
+		IRigidBody()
 		{
 		}
 
-		explicit FRigidBodyId(uint64 InId)
-			: IntId(InId)
+		virtual ~IRigidBody()
 		{
 		}
 
-		void Reset()
+		// Is the body active in its scene?
+		// Inactive bodies do not participate in physics in any way.
+		UE_INTERNAL RIGIDPHYSICS_API virtual bool IsActive() const
 		{
-			IntId = INDEX_NONE;
+			RIGIDPHYSICS_API_UNSUPPORTED();
+			return false;
 		}
 
-		friend bool operator==(const FRigidBodyId& L, const FRigidBodyId& R) = default;
-		
-		friend bool operator<(const FRigidBodyId& L, const FRigidBodyId& R)
+		UE_INTERNAL RIGIDPHYSICS_API virtual void Activate()
 		{
-			return L.IntId < R.IntId;
-		}
-		
-		friend uint32 GetTypeHash(const FRigidBodyId& BodyId)
-		{ 
-			return GetTypeHash(BodyId.IntId);
+			RIGIDPHYSICS_API_UNSUPPORTED();
 		}
 
-	private:
-		friend class Private::FRigidAccessor;
-
-		uint64 IntId;
-	};
-
-	// The full rigid body API. There will be several implementations of this API. E.g., synchronous
-	// physics, asynchronous physics, immediate physics, remote physics, etc.
-	// NOTE: for some implementations like async physics, there are both Game- and Physics-Context
-	// implementations within the async implementation. Locking a handle will return the object 
-	// appropriate for the context.
-	class RIGIDPHYSICS_API IRigidBody
-	{
-	public:
-		IRigidBody();
-		virtual ~IRigidBody();
-
-		virtual bool IsValid() const = 0;
-		virtual bool IsInScene() const = 0;
+		UE_INTERNAL RIGIDPHYSICS_API virtual void Deactivate()
+		{
+			RIGIDPHYSICS_API_UNSUPPORTED();
+		}
 
-		// The ID for this body. This is only guaranteed to be unique in its container.
-		// i.e., you need a Conainer+ID to access a body (or a FRigidBodyHandle).
-		virtual FRigidBodyId GetBodyId() const = 0;
 
 		// The handle to this body
-		virtual FRigidBodyHandle GetBodyHandle() const = 0;
+		UE_INTERNAL RIGIDPHYSICS_API virtual FRigidBodyHandle GetHandle() const = 0;
 
 		// Is this body Static? Static bodies can be moved but have no velocity.
-		virtual bool IsStatic() const
-		{
-			RIGIDPHYSICS_API_UNSUPPORTED();
-			return true;
-		}
+		UE_INTERNAL RIGIDPHYSICS_API virtual bool IsStatic() const = 0;
 
 		// Is this body kinematic? Kinematic bodies can be explicitly moved and have velocity, but do
 		// not react to forces, collisions, etc.
-		virtual bool IsKinematic() const
-		{
-			RIGIDPHYSICS_API_UNSUPPORTED();
-			return false;
-		}
+		UE_INTERNAL RIGIDPHYSICS_API virtual bool IsKinematic() const = 0;
 
 		// Is this a dynamic body? Dynamic bodies have momentum and react to forces and collisions.
-		virtual bool IsDynamic() const
+		UE_INTERNAL RIGIDPHYSICS_API virtual bool IsDynamic() const = 0;
+
+		// Make the body dynamic (will only affect Kinematic bodies. Static bodies cannot be made dynamic)
+		UE_INTERNAL RIGIDPHYSICS_API virtual void SetIsDynamic(bool bIsDynamic) const
 		{
 			RIGIDPHYSICS_API_UNSUPPORTED();
-			return false;
 		}
 
 		// Is this a sleeping dynamic body? Sleeping bodies are dynamic bodies that are ignored by
 		// the simulation (as an optimization) until the scene changes in a way that might affect them, 
 		// or they are explicitly woken.
-		virtual bool IsSleeping() const
+		UE_INTERNAL RIGIDPHYSICS_API virtual bool IsSleeping() const
 		{
 			RIGIDPHYSICS_API_UNSUPPORTED();
 			return false;
 		}
 
 		// Set the transform and initialize body state (velocity etc.) at the new location.
-		virtual void InitTransform(const FTransform& InTransform)
+		UE_INTERNAL RIGIDPHYSICS_API virtual void InitTransform(const FTransform3d& InTransform)
 		{
 			RIGIDPHYSICS_API_UNSUPPORTED();
 		}
 
 		// Set the transform without changing any other rigid body state.
-		virtual void UpdateTransform(const FTransform& InTransform)
+		UE_INTERNAL RIGIDPHYSICS_API virtual void UpdateTransform(const FTransform3d& InTransform)
 		{
 			RIGIDPHYSICS_API_UNSUPPORTED();
 		}
@@ -132,7 +98,7 @@ namespace UE::Physics
 		// tick. Kinematic bodies will always reach the target bscause they do not respond
 		// to forces or collisions. Dynamic bodies will only reach their target if their
 		// velocity in the absence of interactions (forces, collisions, or other constraints).
-		virtual void SetKinematicTarget(const FTransform& InTransform)
+		UE_INTERNAL RIGIDPHYSICS_API virtual void SetKinematicTarget(const FTransform3d& InTransform)
 		{
 			RIGIDPHYSICS_API_UNSUPPORTED();
 		}
@@ -140,135 +106,174 @@ namespace UE::Physics
 		// Get the current transform. NOTE: This will immediately reflect any recent calls to
 		// to InitTransform or UpdateTransform, but will not reflect the latest KinematicTarget
 		// until after the next physics tick.
-		virtual FTransform GetTransform() const
+		UE_INTERNAL RIGIDPHYSICS_API virtual FTransform3d GetTransform() const
 		{
 			RIGIDPHYSICS_API_UNSUPPORTED();
-			return FTransform::Identity;
+			return FTransform3d::Identity;
 		}
 
 		// Get the world-space bounds at the current transform
-		virtual FBounds3d GetBounds() const
+		UE_INTERNAL RIGIDPHYSICS_API virtual FBounds3d GetBounds() const
 		{
 			RIGIDPHYSICS_API_UNSUPPORTED();
 			return FBounds3d(FVector3d::Zero(), FVector3d::Zero());
 		}
 
 		// TODO: Maybe just Set/GetMassProperties
-		virtual double GetMass() const
+		UE_INTERNAL RIGIDPHYSICS_API virtual double GetMass() const
 		{
 			RIGIDPHYSICS_API_UNSUPPORTED();
 			return 0;
 		}
 
-		virtual void SetMass(double InM)
+		UE_INTERNAL RIGIDPHYSICS_API virtual void SetMass(double InM)
 		{
 			RIGIDPHYSICS_API_UNSUPPORTED();
 		}
 
-		virtual void SetCCDEnabled(bool bInEnabled)
+		UE_INTERNAL RIGIDPHYSICS_API virtual void SetCCDEnabled(bool bInEnabled)
 		{
 			RIGIDPHYSICS_API_UNSUPPORTED();
 		}
 
-		virtual void SetMACDEnabled(bool bInEnabled)
+		UE_INTERNAL RIGIDPHYSICS_API virtual void SetMACDEnabled(bool bInEnabled)
 		{
 			RIGIDPHYSICS_API_UNSUPPORTED();
 		}
 		
-		virtual void SetSmoothEdgeCollisionsEnabled(bool bInEnabled)
+		UE_INTERNAL RIGIDPHYSICS_API virtual void SetSmoothEdgeCollisionsEnabled(bool bInEnabled)
 		{
 			RIGIDPHYSICS_API_UNSUPPORTED();
 		}
 
-		virtual void SetMaxLinearVelocity(float InMax)
+		UE_INTERNAL RIGIDPHYSICS_API virtual void SetMaxLinearVelocity(float InMax)
 		{
 			RIGIDPHYSICS_API_UNSUPPORTED();
 		}
 		
-		virtual void SetMaxAngularVelocity(float InMax)
+		UE_INTERNAL RIGIDPHYSICS_API virtual void SetMaxAngularVelocity(float InMax)
 		{
 			RIGIDPHYSICS_API_UNSUPPORTED();
 		}
 		
-		virtual void SetInertiaConditioningEnabled(bool bInEnabled)
+		UE_INTERNAL RIGIDPHYSICS_API virtual void SetInertiaConditioningEnabled(bool bInEnabled)
 		{
 			RIGIDPHYSICS_API_UNSUPPORTED();
 		}
 		
-		virtual void SetGyroscopicTorqueEnabled(bool bInEnabled)
+		UE_INTERNAL RIGIDPHYSICS_API virtual void SetGyroscopicTorqueEnabled(bool bInEnabled)
 		{
 			RIGIDPHYSICS_API_UNSUPPORTED();
 		}
 
-		virtual void SetGravityGroup(int32 InId)
+		UE_INTERNAL RIGIDPHYSICS_API virtual void SetGravityGroup(int32 InId)
 		{
 			RIGIDPHYSICS_API_UNSUPPORTED();
 		}
 		
-		virtual void SetGravityScale(float InScale)
+		UE_INTERNAL RIGIDPHYSICS_API virtual void SetGravityScale(float InScale)
 		{
 			RIGIDPHYSICS_API_UNSUPPORTED();
 		}
 
-		virtual void SetWakeEventsEnabled(bool bInEnabled)
+		UE_INTERNAL RIGIDPHYSICS_API virtual void SetWakeEventsEnabled(bool bInEnabled)
 		{
 			RIGIDPHYSICS_API_UNSUPPORTED();
 		}
 		
-		virtual void SetNumPositionIterations(int32 InNum)
+		UE_INTERNAL RIGIDPHYSICS_API virtual void SetNumPositionIterations(int32 InNum)
 		{
 			RIGIDPHYSICS_API_UNSUPPORTED();
 		}
 		
-		virtual void SetNumVelocityIterations(int32 InNum)
+		UE_INTERNAL RIGIDPHYSICS_API virtual void SetNumVelocityIterations(int32 InNum)
 		{
 			RIGIDPHYSICS_API_UNSUPPORTED();
 		}
 		
-		virtual void SetNumProjectionIterations(int32 InNum)
+		UE_INTERNAL RIGIDPHYSICS_API virtual void SetNumProjectionIterations(int32 InNum)
 		{
 			RIGIDPHYSICS_API_UNSUPPORTED();
 		}
 		
-		virtual void SetIsSleeping(bool bInSleeping)
+		UE_INTERNAL RIGIDPHYSICS_API virtual void SetIsSleeping(bool bInSleeping)
 		{
 			RIGIDPHYSICS_API_UNSUPPORTED();
 		}
-
-	protected:
-
-		virtual const char* GetTypeName() const
-		{
-			return "Unknown";
-		}
-	};
-
-	// Internal API for rigid bodies. Not intended for public use.
-	class RIGIDPHYSICS_API IRigidBodyInternal : public IRigidBody
-	{
-	public:
-		IRigidBodyInternal() {}
-		virtual ~IRigidBodyInternal() {}
-
-		virtual void SetSpatialIdx() {}
 	};
 
 
-	// This class only exists so that TRigidBodyPtr can have pointer semantics rather than object semantics.
+	// This class only exists so that TRigidBodyPtr can have pointer semantics rather than object semantics,
+	// and so that we can disable functions based on ContextType
 	// See TRigidBodyPtr for comments.
 	template<typename ContextType>
-	class TRigidBodyPtrImpl
+	class UE_INTERNAL TRigidBodyPtrImpl
 	{
 	public:
 		using FContext = ContextType;
-		using FRigidScenePtr = TRigidScenePtr<ContextType>;
+		using FInterface = IRigidBody;
+		using FHandle = FRigidBodyHandle;
+
+		UE_INTERNAL TRigidBodyPtrImpl()
+			: BodyRaw(nullptr)
+		{
+		}
+
+		UE_INTERNAL TRigidBodyPtrImpl(IRigidBody* InBody)
+			: BodyRaw(InBody)
+		{
+		}
+
+		TRigidBodyPtrImpl(const TRigidBodyPtrImpl& InBodyPtr) = delete;
+
+		TRigidBodyPtrImpl(TRigidBodyPtrImpl&& InBodyPtr)
+			: BodyRaw(InBodyPtr.BodyRaw)
+		{
+			InBodyPtr.BodyRaw = nullptr;
+		}
 
 		~TRigidBodyPtrImpl()
 		{
 			Reset();
 		}
 
-		// Forward to IRigidBody. See IRigidBody for comments
+		UE_INTERNAL void Reset()
+		{
+			BodyRaw = nullptr;
+		}
+
+		UE_INTERNAL bool IsValid() const
+		{
+			return (BodyRaw != nullptr);
+		}
+
+		const FRigidTypeId& GetTypeId() const
+		{
+			return BodyRaw->GetTypeId();
+		}
+
+		FRigidBodyHandle GetHandle() const;
+
+		FRigidBodyContainerHandle GetContainerHandle() const
+		{
+			return BodyRaw->GetContainerHandle();
+		}
+
+		bool IsActive() const
+		{
+			return BodyRaw->IsActive();
+		}
+
+		void Activate()
+		{
+			BodyRaw->Activate();
+		}
+
+		void Deactivate()
+		{
+			BodyRaw->Deactivate();
+		}
+
 		bool IsStatic() const
 		{
 			return BodyRaw->IsStatic();
@@ -284,22 +289,22 @@ namespace UE::Physics
 			return BodyRaw->IsDynamic();
 		}
 
-		void InitTransform(const FTransform& InTransform)
+		void InitTransform(const FTransform3d& InTransform)
 		{
 			BodyRaw->InitTransform(InTransform);
 		}
 
-		void UpdateTransform(const FTransform& InTransform)
+		void UpdateTransform(const FTransform3d& InTransform)
 		{
 			BodyRaw->UpdateTransform(InTransform);
 		}
 
-		void SetKinematicTarget(const FTransform& InTransform)
+		void SetKinematicTarget(const FTransform3d& InTransform)
 		{
 			BodyRaw->SetKinematicTarget(InTransform);
 		}
 
-		FTransform GetTransform() const
+		FTransform3d GetTransform() const
 		{
 			return BodyRaw->GetTransform();
 		}
@@ -368,10 +373,12 @@ namespace UE::Physics
 		{
 			BodyRaw->SetWakeEventsEnabled(bInEnabled);
 		}
+
 		void SetNumPositionIterations(int32 InNum)
 		{
 			BodyRaw->SetNumPositionIterations(InNum);
 		}
+
 		void SetNumVelocityIterations(int32 InNum)
 		{
 			BodyRaw->SetNumVelocityIterations(InNum);
@@ -390,167 +397,41 @@ namespace UE::Physics
 		friend bool operator==(const TRigidBodyPtrImpl& L, const TRigidBodyPtrImpl& R) = default;
 		friend bool operator!=(const TRigidBodyPtrImpl& L, const TRigidBodyPtrImpl& R) = default;
 
-	private:
-		friend class Private::FRigidAccessor;
-		friend class TRigidBodyPtr<FContext>;
-
-		TRigidBodyPtrImpl(const TRigidBodyPtrImpl& R) = default;
-		TRigidBodyPtrImpl(TRigidBodyPtrImpl&& R) = default;
-
-		TRigidBodyPtrImpl(IRigidBody* InBody)
-			: BodyRaw(InBody)
-			, SceneLock()
-		{
-		}
-
-		// NOTE: Receiving a transient lock, so we must store it
-		void Lock(FRigidScenePtr&& InScene)
-		{
-			SceneLock = MoveTemp(InScene);
-		}
-
-		void Reset()
-		{
-			BodyRaw = nullptr;
-		}
-
-		bool IsValid() const
+		UE_INTERNAL IRigidBody* Get() const
 		{
-			return (BodyRaw != nullptr);
+			return BodyRaw;
 		}
 
+	private:
 		// IRigidBody implementation depends on Context type
 		IRigidBody* BodyRaw;
-
-		// Holds a lock on the scene (if set)
-		FRigidScenePtr SceneLock;
-	};
-
-
-	// A context-aware IRigidBody pointer that simply forwards all context-legal operations to IRigidBody. 
-	// Provides access to a restricted RigidBody API where the restrictions depend on the context 
-	// (template-parameter), and attempts to call disallowed functions (in the context) produce compiler errors.
-	// 
-	// Uses pointer semantics and may be "null". Check IsValid() before dereferencing. Also implicitly casts 
-	// to boolean for validity checks.
-	// 
-	// Cannot be copied or stored - use FRigidBodyHandle for that.
-	template<typename ContextType>
-	class TRigidBodyPtr
-	{
-	public:
-		using FContext = ContextType;
-		using FRigidBodyPtrImpl = TRigidBodyPtrImpl<FContext>;
-		using FRigidBodyContainerPtr = TRigidBodyContainerPtr<FContext>;
-		using FRigidScenePtr = TRigidScenePtr<ContextType>;
-
-		TRigidBodyPtr()
-			: BodyImpl(nullptr)
-		{
-		}
-
-		~TRigidBodyPtr()
-		{
-		}
-
-		// Has this pointer been initialized with a body pointer?
-		bool IsValid() const
-		{
-			return BodyImpl.IsValid();
-		}
-
-		// Auto cast to bool to support using in if statements:
-		//	if (TRigidBodyPtr<FGameContextRO> Body = BodyHandle.LockRO())...
-		operator bool() const
-		{
-			return IsValid();
-		}
-
-		// Get the handle for this body
-		FRigidBodyHandle GetHandle() const;
-
-		operator FRigidBodyHandle() const;
-
-		// Pointer sematics
-		const FRigidBodyPtrImpl* operator->() const
-		{
-			checkf(IsValid(), TEXT("Check IsValid() before dereferencing"));
-			return &BodyImpl;
-		}
-
-		// Pointer sematics
-		FRigidBodyPtrImpl* operator->() const requires (FContext::bWriteEnabled)
-		{
-			checkf(IsValid(), TEXT("Check IsValid() before dereferencing"));
-			return &BodyImpl;
-		}
-
-		friend bool operator==(const TRigidBodyPtr& L, const TRigidBodyPtr& R) = default;
-		friend bool operator!=(const TRigidBodyPtr& L, const TRigidBodyPtr& R) = default;
-
-	private:
-		friend class Private::FRigidAccessor;
-		friend class TRigidBodyContainerPtrImpl<FContext>;
-		friend class TRigidScenePtrImpl<FContext>;
-		friend class FRigidBodyHandle;
-
-		TRigidBodyPtr(const TRigidBodyPtr& R) = default;
-		TRigidBodyPtr(TRigidBodyPtr&& R) = default;
-
-		// Create a pointer with no scene lock. Used in situations were the scene is known to  be 
-		// locked before the rigid body pointer is retrieved from the ID. This will be when 
-		// the user wants to use multiple physics objects under one lock.
-		TRigidBodyPtr(IRigidBody* InBody)
-			: BodyImpl(InBody)
-		{
-		}
-
-		// NOTE: Receiving a transient lock, so we must store it. Should be used
-		// when only a small number of physics objects are being accessed and extra locking
-		// overhead is not an issue. Otherise an external scene lock is preferred.
-		TRigidBodyPtr& Lock(FRigidScenePtr&& InScene)
-		{
-			BodyImpl.Lock(MoveTemp(InScene));
-			return *this;
-		}
-
-		mutable FRigidBodyPtrImpl BodyImpl;
 	};
 
 
 	// A handle to a rigid body in a container.
-	// This handle may be copied and stored, and can be used to retrieve a useful
-	// body reference by "locking" it and providing the necessary context.
+	// This handle may be copied and stored, and can be used to retrieve a
+	// body reference by "pinning" it and providing the context.
 	// 
 	// Example contexts are: 
 	// - Read-Only Game-Thread
 	// - Read-Write Game Thread
-	// - Physics-Thread BeginTick Callback
-	// - etc.
-	//
-	// NOTE: If you have many bodies from the same container, it will be more efficient
-	// to store a set of BodyIds separately from the ContainerHandle. See TRigidBodyPtr.
+	// - Read-Write Physics-Thread Callback
 	//
 	class FRigidBodyHandle
 	{
 	public:
 		FRigidBodyHandle() = default;
+		FRigidBodyHandle(const FRigidBodyHandle& R) = default;
 
-		FRigidBodyHandle(const FRigidBodyContainerHandle& InBodyContainerHandle, const FRigidBodyId& InBodyId)
-			: BodyContainerHandle(InBodyContainerHandle)
-			, BodyId(InBodyId)
+		UE_INTERNAL FRigidBodyHandle(const FRigidBodyContainerHandle& InBodyContainerHandle, const FRigidObjectId& InBodyId)
+			: BodyId(InBodyId)
+			, BodyContainerHandle(InBodyContainerHandle)
 		{
 		}
 
-		void Reset()
-		{
-			BodyContainerHandle.Reset();
-			BodyId.Reset();
-		}
-
-		const FRigidSceneHandle& GetSceneHandle() const
+		const FRigidObjectId& GetId() const
 		{
-			return BodyContainerHandle.GetSceneHandle();
+			return BodyId;
 		}
 
 		const FRigidBodyContainerHandle& GetContainerHandle() const
@@ -558,127 +439,64 @@ namespace UE::Physics
 			return BodyContainerHandle;
 		}
 
-		const FRigidBodyId& GetBodyId() const
-		{
-			return BodyId;
-		}
-
-		// Get a pointer to the GameContext rigid body and hold a read-lock on the scene for the duration of the pointer.
-		TRigidBodyPtr<FGameContextRO> LockRO() const
-		{
-			// NOTE: Passing a transient lock.
-			return Lock(GetSceneHandle().LockRO());
-		}
-
-		// Get a pointer to the GameContext rigid body and hold a write-lock on the scene for the duration of the pointer.
-		TRigidBodyPtr<FGameContextRW> LockRW() const
+		void Reset()
 		{
-			// NOTE: Passing a transient lock.
-			return Lock(GetSceneHandle().LockRW());
+			BodyContainerHandle.Reset();
+			BodyId.Reset();
 		}
 
-		// Get a pointer to the Context-specific rigid body. This does not hold a lock on the scene,
-		// which is assumed to be locked externally with a TRigidScenePtr whose lifetime
-		// exceeds this pointer. Using this pointer after the external lock is released will 
-		// cause race conditions and random crashes.
+		// Convert the handle to a pointer in the specified context
+		// Impl in RigidContext.h
 		template<typename ContextType>
-		TRigidBodyPtr<ContextType> Pin(const TRigidScenePtr<ContextType>& InScene) const
-		{
-			if (TRigidBodyContainerPtr<ContextType> Container = BodyContainerHandle.Pin(InScene))
-			{
-				return Container->GetRigidBody(BodyId);
-			}
-			return TRigidBodyPtr<ContextType>();
-		}
+		TRigidBodyPtr<ContextType> Pin(const ContextType& InContext) const;
 
 		friend bool operator==(const FRigidBodyHandle& L, const FRigidBodyHandle& R) = default;
 		friend bool operator!=(const FRigidBodyHandle& L, const FRigidBodyHandle& R) = default;
 
 	private:
-		friend class Private::FRigidAccessor;
-
-		// Rvalue version is private and deleted to prevent users from being able to Pin with a
-		// lock whose lifetime is limited to the scope of the call. Do not implement this 
-		// as it could easily be used incorrectly. See Lock() instead.
-		template<typename ContextType>
-		TRigidBodyPtr<ContextType> Pin(TRigidScenePtr<ContextType>&& InScene) const = delete;
-
-		// NOTE: Receiving a transient lock, so we must store it in RigidBodyPtr.
-		template<typename ContextType>
-		TRigidBodyPtr<ContextType> Lock(TRigidScenePtr<ContextType>&& InScene) const
-		{
-			return Pin(InScene).Lock(MoveTemp(InScene));
-		}
-
+		FRigidObjectId BodyId;
 		FRigidBodyContainerHandle BodyContainerHandle;
-		FRigidBodyId BodyId;
 	};
 
-
-	template<typename ContextType>
-	FRigidBodyHandle TRigidBodyPtr<ContextType>::GetHandle() const
-	{
-		if (IsValid())
-		{
-			return BodyImpl.BodyRaw->GetBodyHandle();
-		}
-		return FRigidBodyHandle();
-	}
-
-	template<typename ContextType>
-	TRigidBodyPtr<ContextType>::operator FRigidBodyHandle() const
-	{
-		return GetHandle();
-	}
-
-	//
-	// Container implementations to avoid circular dependency in headers
-	//
-
 	template<typename ContextType>
-	TRigidBodyPtr<ContextType> TRigidBodyContainerPtrImpl<ContextType>::GetRigidBody(const FRigidBodyId& BodyId) const
+	FRigidBodyHandle TRigidBodyPtrImpl<ContextType>::GetHandle() const
 	{
-		UE_RIGIDPHYSICS_CHECK(IsValid());
-		return TRigidBodyPtr<ContextType>(BodyContainerRaw->GetRigidBody(BodyId));
+		return BodyRaw->GetHandle();
 	}
 
 	//
-	// Scene implementations to avoid circular dependency in headers
+	// Implementations to avoid circular dependency in headers
 	//
 
 	template<typename ContextType>
-	TRigidBodyPtr<ContextType> TRigidScenePtrImpl<ContextType>::CreateRigidBody(const FRigidNamePtr& InName, ERigidMovementType InMovementType) const
+	TRigidBodyPtr<ContextType> TRigidBodyContainerPtrImpl<ContextType>::GetBody(const FRigidObjectId& BodyId) const
 	{
-		UE_RIGIDPHYSICS_CHECK(IsValid());
-		return TRigidBodyPtr<FContext>(SceneRaw->CreateRigidBody(InName, InMovementType));
+		return BodyContainerRaw->GetBody(BodyId);
 	}
 
 	template<typename ContextType>
-	void TRigidScenePtrImpl<ContextType>::DestroyRigidBody(const FRigidBodyHandle& BodyHandle) const
+	TRigidBodyPtr<ContextType> TRigidBodyContainerPtrImpl<ContextType>::CreateBody(const FRigidDebugName& InName, ERigidMovementType InMovementType) const
 	{
-		UE_RIGIDPHYSICS_CHECK(IsValid());
-		SceneRaw->DestroyRigidBody(BodyHandle);
+		return BodyContainerRaw >CreateBody(InName, InMovementType);
 	}
 
 	template<typename ContextType>
-	void TRigidScenePtrImpl<ContextType>::RegisterRigidBody(const FRigidBodyHandle& BodyHandle)
+	void TRigidBodyContainerPtrImpl<ContextType>::DestroyBody(TRigidBodyPtr<FContext>&& InBodyPtr)
 	{
-		UE_RIGIDPHYSICS_CHECK(IsValid());
-		SceneRaw->RegisterRigidBody(BodyHandle);
+		BodyContainerRaw->DestroyBody(InBodyPtr.Ge());
 	}
 
 	template<typename ContextType>
-	void TRigidScenePtrImpl<ContextType>::UnregisterRigidBody(const FRigidBodyHandle& BodyHandle)
+	TRigidBodyPtr<ContextType> TRigidScenePtrImpl<ContextType>::CreateBody(const FRigidDebugName& InName, ERigidMovementType InMovementType)
 	{
-		UE_RIGIDPHYSICS_CHECK(IsValid());
-		SceneRaw->UnregisterRigidBody(BodyHandle);
+		return SceneRaw->CreateBody(InName, InMovementType);
 	}
 
 	template<typename ContextType>
-	TRigidBodyContainerPtr<ContextType> TRigidScenePtrImpl<ContextType>::GetRigidBodyContainer(const FRigidBodyContainerId& InBodyContainerId) const
+	void TRigidScenePtrImpl<ContextType>::DestroyBody(TRigidBodyPtr<ContextType>&& InBody)
 	{
-		UE_RIGIDPHYSICS_CHECK(IsValid());
-		return TRigidBodyContainerPtr<ContextType>(SceneRaw->GetRigidBodyContainer(InBodyContainerId));
+		SceneRaw->DestroyBody(InBody.Get());
+		InBody.Reset();
 	}
 
 }
diff --git a/Engine/Source/Runtime/Experimental/RigidPhysics/Public/RigidPhysics/RigidBodyContainer.h b/Engine/Source/Runtime/Experimental/RigidPhysics/Public/RigidPhysics/RigidBodyContainer.h
--- a/Engine/Source/Runtime/Experimental/RigidPhysics/Public/RigidPhysics/RigidBodyContainer.h
+++ b/Engine/Source/Runtime/Experimental/RigidPhysics/Public/RigidPhysics/RigidBodyContainer.h
@@ -5,245 +5,127 @@
 
 #if UE_RIGIDPHYSICS_API_ENABLED
 
-#include "Misc/Crc.h"
+#include "RigidPhysics/RigidLog.h"
+#include "RigidPhysics/RigidObjectId.h"
+#include "RigidPhysics/RigidObjectPtr.h"
+#include "RigidPhysics/RigidObjectRegistry.h"
 #include "RigidPhysics/RigidScene.h"
-#include "Templates/SharedPointer.h"
+#include "RigidPhysics/RigidTyped.h"
 
 namespace UE::Physics
 {
-	// A type ID for a rigid body containers. There will be multiple types of
-	// container, and this can grow over time. Initially we will have:
-	// SingleBodyContainer, GeometryCollection, ClusterUnion.
-	class FRigidBodyContainerTypeId
+	// The interface for rigid body containers. There will be several implementations of this API. E.g., synchronous
+	// physics, asynchronous physics, immediate physics, remote physics, etc. 
+	class UE_INTERNAL IRigidBodyContainer : public IRigidTyped
 	{
 	public:
-		FRigidBodyContainerTypeId(const TCHAR* InTypeName)
-			: TypeId(FCrc::StrCrc32(InTypeName))
-			, TypeName(InTypeName)
+		UE_RIGIDPHYSICS_RIGIDTYPED_DECL(RIGIDPHYSICS_API, IRigidBodyContainer);
+
+		UE_INTERNAL IRigidBodyContainer()
 		{
 		}
 
-		friend bool operator==(const FRigidBodyContainerTypeId& L, const FRigidBodyContainerTypeId& R) = default;
-		friend bool operator<(const FRigidBodyContainerTypeId& L, const FRigidBodyContainerTypeId& R)
+		UE_INTERNAL virtual ~IRigidBodyContainer()
 		{
-			return L.TypeId < R.TypeId;
 		}
 
-	private:
-		uint64 TypeId;
-		const TCHAR* TypeName;
-	};
-
+		UE_INTERNAL RIGIDPHYSICS_API virtual FRigidBodyContainerHandle GetHandle() const = 0;
 
-	// A container ID. All containers are registered with with RigidScene and can be found
-	// from their ID (usually from a FRigidBodyContainerHandle).
-	class FRigidBodyContainerId
-	{
-	public:
-		FRigidBodyContainerId()
-		{
-			IntId = -1;
-		}
+		UE_INTERNAL RIGIDPHYSICS_API virtual int32 GetNumBodies() const = 0;
 
-		explicit FRigidBodyContainerId(uint64 InId)
-		{
-			IntId = InId;
-		}
+		UE_INTERNAL RIGIDPHYSICS_API virtual IRigidBody* GetBody(const FRigidObjectId& InBodyId) const = 0;
 
-		void Reset()
+		UE_INTERNAL RIGIDPHYSICS_API virtual IRigidBody* CreateBody(const FRigidDebugName& InName, ERigidMovementType InMovementType)
 		{
-			IntId = -1;
+			RIGIDPHYSICS_API_UNSUPPORTED();
+			return nullptr;
 		}
 
-		friend bool operator==(const FRigidBodyContainerId& L, const FRigidBodyContainerId& R) = default;
-		friend bool operator<(const FRigidBodyContainerId& L, const FRigidBodyContainerId& R)
+		UE_INTERNAL RIGIDPHYSICS_API virtual void DestroyBody(IRigidBody* InBody)
 		{
-			return L.IntId < R.IntId;
-		}
-		friend uint32 GetTypeHash(const FRigidBodyContainerId& BodyContainerId)
-		{ 
-			return GetTypeHash(BodyContainerId.IntId);
+			RIGIDPHYSICS_API_UNSUPPORTED();
 		}
 
-	private:
-		friend class Private::FRigidAccessor;
-
-		uint64 IntId;
 	};
 
-
-	// The interface for rigid body containers. There will be several implementations of this API. E.g., synchronous
-	// physics, asynchronous physics, immediate physics, remote physics, etc. 
-	class RIGIDPHYSICS_API IRigidBodyContainer
+	class UE_INTERNAL FRigidBodyContainer : public IRigidBodyContainer
 	{
 	public:
-		IRigidBodyContainer(const FRigidBodyContainerTypeId& InTypeId);
-		virtual ~IRigidBodyContainer();
-
-		const FRigidBodyContainerTypeId& GetTypeId() const;
+		UE_INTERNAL RIGIDPHYSICS_API FRigidBodyContainer();
+		UE_INTERNAL RIGIDPHYSICS_API virtual ~FRigidBodyContainer();
 
-		bool IsCastableTo(const FRigidBodyContainerTypeId& InTypeId) const;
+		UE_INTERNAL RIGIDPHYSICS_API virtual IRigidBody* GetBody(const FRigidObjectId& InBodyId) const override final;
 
-		// Downcast to the specified container type
-		template<typename T>
-		T* TryAsA()
+	protected:
+		FRigidBodyRegistry* GetBodyRegistry() const
 		{
-			if (IsCastableTo(T::StaticRigidBodyContainerTypeId()))
-			{
-				return static_cast<T&>(*this);
-			}
-			return nullptr;
+			return BodyRegistry.Get();
 		}
 
-		virtual FRigidBodyContainerId GetContainerId() const = 0;
-
-		virtual IRigidScene* GetScene() const = 0;
-
-		virtual IRigidBody* GetRigidBody(const FRigidBodyId& BodyId) const = 0;
-
 	private:
-		FRigidBodyContainerTypeId TypeId;
+		TUniquePtr<FRigidBodyRegistry> BodyRegistry;
 	};
 
-
 	// Utility wrapper around IRigidBodyContainer that adds Context and allows
-	// TRigidBodyContainerPtr to have pointer semantics
+	// TRigidBodyContainerPtr to have pointer semantics without exposing IRigidBodyContainer
 	template<typename ContextType>
-	class TRigidBodyContainerPtrImpl
+	class UE_INTERNAL TRigidBodyContainerPtrImpl
 	{
 	public:
 		using FContext = ContextType;
-		using FRigidScenePtr = TRigidScenePtr<ContextType>;
+		using FInterface = IRigidBodyContainer;
+		using FHandle = FRigidBodyContainerHandle;
+		using FRigidScenePtr = TRigidScenePtr<FContext>;
 
-		TRigidBodyContainerPtrImpl(IRigidBodyContainer* InBodyContainer)
+		UE_INTERNAL TRigidBodyContainerPtrImpl(IRigidBodyContainer* InBodyContainer)
 			: BodyContainerRaw(InBodyContainer)
-			, SceneLock()
 		{
 		}
 
 		virtual ~TRigidBodyContainerPtrImpl()
 		{
-			Reset();
 		}
 
-		void Reset()
+		UE_INTERNAL void Reset()
 		{
 			BodyContainerRaw = nullptr;
 		}
 
-		bool IsValid() const
+		UE_INTERNAL bool IsValid() const
 		{
 			return (BodyContainerRaw != nullptr);
 		}
 
-		// Get the type ID of this container
-		const FRigidBodyContainerTypeId& GetContainerTypeId() const
+		const FRigidTypeId& GetTypeId() const
 		{
-			UE_RIGIDPHYSICS_CHECK(IsValid());
 			return BodyContainerRaw->GetTypeId();
 		}
 
-		// Get the ID of this container
-		FRigidBodyContainerId GetContainerId() const
+		FRigidObjectId GetId() const
 		{
-			UE_RIGIDPHYSICS_CHECK(IsValid());
-			return BodyContainerRaw->GetContainerId();
+			return BodyContainerRaw->GetId();
 		}
 
-		// Impl in RigidBody.h
-		TRigidBodyPtr<FContext> GetRigidBody(const FRigidBodyId& BodyId) const;
-
-	private:
-		friend class TRigidBodyContainerPtr<ContextType>;
+		FRigidBodyContainerHandle GetHandle() const;
 
-		TRigidBodyContainerPtrImpl(const TRigidBodyContainerPtrImpl& R)
-			: BodyContainerRaw(R.BodyContainerRaw)
-			, SceneLock(R.SceneLock)
+		bool IsInScene() const
 		{
+			return BodyContainerRaw->IsInScene();
 		}
 
-		// We are passed a transient lock, so we must store it
-		void Lock(FRigidScenePtr&& InScene)
+		int32 GetNumBodies() const
 		{
-			SceneLock = MoveTemp(InScene);
+			return BodyContainerRaw->GetNumBodies();
 		}
 
+		// Impl in RigidBody.h
+		TRigidBodyPtr<FContext> GetBody(const FRigidObjectId& InBodyId) const;
+		TRigidBodyPtr<FContext> CreateBody(const FRigidDebugName& InName, ERigidMovementType InMovementType) const;
+		void DestroyBody(TRigidBodyPtr<FContext>&& InBodyPtr);
+
+	private:
 		// Pointer to implementation, object depends on Context type
 		IRigidBodyContainer* BodyContainerRaw;
-
-		// Holds a lock on the scene if set
-		FRigidScenePtr SceneLock;
-	};
-
-
-	// A context-aware IRigidBodyContainer pointer that simply forwards all context-legal operations to IRigidBodyContainer. 
-	// Provides access to a restricted RigidBodyContainer API where the restrictions depend on the context 
-	// (template-parameter), and attempts to call disallowed functions (in the context) produce compiler errors.
-	// 
-	// Uses pointer semantics and may be "null". Check IsValid() before dereferencing. Also implicitly casts 
-	// to boolean for validity checks.
-	// 
-	// Cannot be copied or stored - use FRigidBodyContainerHandle for that.
-	template<typename ContextType>
-	class TRigidBodyContainerPtr
-	{
-	public:
-		using FContext = ContextType;
-		using FRigidBodyContainerPtrImpl = TRigidBodyContainerPtrImpl<FContext>;
-		using FRigidScenePtr = TRigidScenePtr<ContextType>;
-
-		TRigidBodyContainerPtr()
-			: BodyContainerImpl(nullptr)
-		{
-		}
-
-		~TRigidBodyContainerPtr()
-		{
-		}
-
-		void Reset()
-		{
-			BodyContainerImpl.Reset();
-		}
-
-		bool IsValid() const
-		{
-			return BodyContainerImpl.IsValid();
-		}
-
-		operator bool() const
-		{
-			return IsValid();
-		}
-
-		const FRigidBodyContainerPtrImpl* operator->() const
-		{
-			return &BodyContainerImpl;
-		}
-
-	protected:
-		friend class Private::FRigidAccessor;
-		friend class FRigidBodyContainerHandle;
-		friend class TRigidScenePtrImpl<FContext>;
-
-		TRigidBodyContainerPtr(IRigidBodyContainer* InBodyContainer)
-			: BodyContainerImpl(InBodyContainer)
-		{
-		}
-
-		TRigidBodyContainerPtr(const TRigidBodyContainerPtr& R)
-			: BodyContainerImpl(R.BodyContainerImpl)
-		{
-		}
-
-		// We are passed a transient lock, so we must store it
-		TRigidBodyContainerPtr& Lock(FRigidScenePtr&& InScene)
-		{
-			BodyContainerImpl.Lock(MoveTemp(InScene));
-			return *this;
-		}
-
-		FRigidBodyContainerPtrImpl BodyContainerImpl;
 	};
 
 
@@ -254,79 +136,46 @@ namespace UE::Physics
 		{
 		}
 
-		FRigidBodyContainerHandle(FRigidSceneHandle InSceneHandle, const FRigidBodyContainerId& InBodyContainerId)
-			: SceneHandle(InSceneHandle)
-			, BodyContainerId(InBodyContainerId)
+		UE_INTERNAL FRigidBodyContainerHandle(const FRigidSceneHandle& InSceneHandle, const FRigidObjectId& InBodyContainerId)
+			: BodyContainerId(InBodyContainerId)
+			, SceneHandle(InSceneHandle)
 		{
 		}
 
 		void Reset()
 		{
-			SceneHandle.Reset();
 			BodyContainerId.Reset();
+			SceneHandle.Reset();
 		}
 
-		const FRigidSceneHandle& GetSceneHandle() const
-		{
-			return SceneHandle;
-		}
-
-		const FRigidBodyContainerId& GetContainerId() const
+		const FRigidObjectId& GetId() const
 		{
 			return BodyContainerId;
 		}
 
-		// Get a pointer to the container and hold a read-lock on the scene for the duration of the pointer.
-		TRigidBodyContainerPtr<FGameContextRO> LockRO() const
-		{
-			// NOTE: Passing a transient lock to Lock()
-			return Lock(SceneHandle.LockRO());
-		}
-
-		// Get a pointer to the container and hold a write-lock on the scene for the duration of the pointer.
-		TRigidBodyContainerPtr<FGameContextRW> LockRW() const
-		{
-			// NOTE: Passing a transient lock to Lock()
-			return Lock(SceneHandle.LockRW());
-		}
-
-		// Get a pointer to the GameContext rigid body container. This does not hold a lock on the scene,
-		// which is assumed to be locked externally with a TRigidScenePtr whose lifetime
-		// exceeds this pointer's. Using this pointer after the external lock is released will 
-		// cause race conditions and random crashes.
+		// Convert the handle to a pointer in the specified context
+		// Impl in RigidContext.h
 		template<typename ContextType>
-		TRigidBodyContainerPtr<ContextType> Pin(const TRigidScenePtr<ContextType>& InScene) const
-		{
-			if (InScene.IsValid())
-			{
-				return InScene->GetRigidBodyContainer(BodyContainerId);
-			}
-			return TRigidBodyContainerPtr<ContextType>();
-		}
+		TRigidBodyContainerPtr<ContextType> Pin(const ContextType& InContext) const;
 
 		friend bool operator==(const FRigidBodyContainerHandle& L, const FRigidBodyContainerHandle& R) = default;
 		friend bool operator!=(const FRigidBodyContainerHandle& L, const FRigidBodyContainerHandle& R) = default;
 
 	private:
-		friend class Private::FRigidAccessor;
-
-		// Rvalue version is private and deleted to prevent users from being able to Pin with a
-		// lock whose lifetime is limited to the scope of the call. Do not implement this 
-		// as it could easily be used incorrectly. See Lock() instead.
-		template<typename ContextType>
-		TRigidBodyContainerPtr<ContextType> Pin(TRigidScenePtr<ContextType>&& InScene) const = delete;
-
-		// NOTE: Receiving a transient lock, so we must store it in TRigidBodyContainerPtr.
-		template<typename ContextType>
-		TRigidBodyContainerPtr<ContextType> Lock(TRigidScenePtr<ContextType>&& InScene) const
-		{
-			return Pin(InScene).Lock(MoveTemp(InScene));
-		}
-
+		FRigidObjectId BodyContainerId;
 
+		// NOTE: This is deliberately not exposed through a GetSceneHandle() method
+		// We don't really need the scene handle because it is provided by the Context
+		// in Pin(). It is only needed for the debugger visualizers.
+		// TODO_CHAOSAPI: Make this shipping only?
 		FRigidSceneHandle SceneHandle;
-		FRigidBodyContainerId BodyContainerId;
 	};
+
+	template<typename ContextType>
+	FRigidBodyContainerHandle TRigidBodyContainerPtrImpl<ContextType>::GetHandle() const
+	{
+		return BodyContainerRaw->GetHandle();
+	}
 }
 
 #endif // UE_RIGIDPHYSICS_API_ENABLED
\ No newline at end of file
diff --git a/Engine/Source/Runtime/Experimental/RigidPhysics/Public/RigidPhysics/RigidContext.h b/Engine/Source/Runtime/Experimental/RigidPhysics/Public/RigidPhysics/RigidContext.h
--- a/Engine/Source/Runtime/Experimental/RigidPhysics/Public/RigidPhysics/RigidContext.h
+++ b/Engine/Source/Runtime/Experimental/RigidPhysics/Public/RigidPhysics/RigidContext.h
@@ -1,100 +1,119 @@
 // Copyright Epic Games, Inc. All Rights Reserved.
 #pragma once
 
-#include "CoreMinimal.h"
 #include "RigidPhysics/RigidFwd.h"
-#include "RigidPhysics/RigidLog.h"
 
 #if UE_RIGIDPHYSICS_API_ENABLED
 
+#include "RigidPhysics/RigidLockable.h"
+#include "RigidPhysics/RigidLog.h"
+#include "RigidPhysics/RigidBody.h"
+#include "RigidPhysics/RigidBodyContainer.h"
+#include "RigidPhysics/RigidScene.h"
+
 namespace UE::Physics
 {
-	// The game callback phase.
-	enum class EContextGamePhase
+	// A Context holds a lock (RW or RO) and converts object handles
+	// to context-specific object pointers.
+	template<ERigidContextType ContextTypeParam, ERigidLockType LockTypeParam>
+	class TRigidContext
 	{
-		None,
-		PreSimulate,
-		PostSimulate,
-		//...
-	};
+	public:
+		static const ERigidContextType ContextType = ContextTypeParam;
+		static const ERigidLockType LockType = LockTypeParam;
+		static const bool bWriteEnabled = (LockType == ERigidLockType::ReadWrite);
 
-	// The simulation callback phase.
-	enum class EContextSimPhase
-	{
-		None,
-		PreTick,
-		PostTick,
-		//...
-	};
+		using FContext = TRigidContext<ContextType, LockType>;
 
-	enum class ERigidMovementType
-	{
-		Static,
-		Kinematic,
-		Dynamic
-	};
+		UE_INTERNAL TRigidContext(IRigidScene* InScene)
+		{
+			if (InScene != nullptr)
+			{
+				InScene->Lock(LockType);
+				ScenePtr.Init(InScene);
+			}
+		}
 
-	class IRigidContext
-	{
-	public:
-		IRigidContext() {}
-		virtual ~IRigidContext() {}
-	};
+		~TRigidContext()
+		{
+			if (ScenePtr.IsValid())
+			{
+				ScenePtr.Get()->Unlock(LockType);
+			}
+		}
 
+		bool IsValid() const
+		{
+			return ScenePtr.IsValid();
+		}
 
+		operator bool() const
+		{
+			return IsValid();
+		}
 
-	template<EContextLockType LockTypeParam>
-	class RIGIDPHYSICS_API TGameContext : public IRigidContext
-	{
-	public:
-		static const EContextEnv ContextEnv = EContextEnv::Game;
-		static const EContextLockType LockType = LockTypeParam;
-		static const EContextSimPhase SimPhase = EContextSimPhase::None;
-		static const bool bWriteEnabled = (LockType == EContextLockType::ReadWrite);
+		// Expose Scene API on Context
+		const TRigidScenePtrImpl<FContext>* operator->() const requires (!FContext::bWriteEnabled)
+		{
+			return ScenePtr.operator->();
+		}
 
-		TGameContext(EContextGamePhase InGamePhase)
-			: GamePhase(InGamePhase)
+		TRigidScenePtrImpl<FContext>* operator->() const requires (FContext::bWriteEnabled)
 		{
+			return ScenePtr.operator->();
 		}
 
-		EContextGamePhase GetGamePhase() const
+		UE_INTERNAL const TRigidScenePtr<FContext>& GetScene() const
 		{
-			return GamePhase;
+			return ScenePtr;
 		}
 
 	private:
-		EContextGamePhase GamePhase;
+		TRigidScenePtr<FContext> ScenePtr;
 	};
 
-	using FGameContextRO = TGameContext<EContextLockType::ReadOnly>;
-	using FGameContextRW = TGameContext<EContextLockType::ReadWrite>;
 
+	//
+	// Functions below are here to prevent circular dependencies
+	//
 
-	class RIGIDPHYSICS_API FSimContextRW : public IRigidContext
+	template<typename ContextType>
+	inline ContextType FRigidSceneHandle::Lock() const
 	{
-	public:
-		static const EContextEnv ContextEnv = EContextEnv::Sim;
-		static const EContextLockType LockType = EContextLockType::None;
-		static const bool bWriteEnabled = true;
-
-		FSimContextRW()
-			: SimPhase(EContextSimPhase::None)
-		{
-		}
-
-		EContextSimPhase GetSimPhase() const
+#if UE_RIGIDPHYSICS_SCENEHANDLE_SAFTEY_ENABLED
+		// NOTE: In order to support the UE_RIGIDPHYSICS_SCENEHANDLE_SAFTEY_ENABLED=0 optimization
+		// SceneHandles cannot be used after the scene has expired. The onus is on the 
+		// user to ensure that this is the case, but when we have handle safety
+		// enabled we can assert on this. Without the weak pointer, there's no way 
+		// to know that the raw pointer held in SceneWrapper points to a deleted object.
+		check(SceneRaw == SceneWeak.Pin().Get());
+#endif
+		return ContextType(SceneRaw);
+	}
+
+	template<typename ContextType>
+	inline TRigidBodyContainerPtr<ContextType> FRigidBodyContainerHandle::Pin(const ContextType& InContext) const
+	{
+		// NOTE: we don't use our copy of SceneHandle here (or anywhere!)
+		if (InContext.IsValid())
 		{
-			return SimPhase;
+			return InContext.GetScene()->GetBodyContainer(BodyContainerId);
 		}
+		return nullptr;
+	}
 
-		void SetSimPhase(EContextSimPhase InSimPhase)
+	template<typename ContextType>
+	inline TRigidBodyPtr<ContextType> FRigidBodyHandle::Pin(const ContextType& InContext) const
+	{
+		if (InContext.IsValid())
 		{
-			SimPhase = InSimPhase;
+			if (TRigidBodyContainerPtr<ContextType> BodyContainer = InContext.GetScene()->GetBodyContainer(BodyContainerHandle.GetId()))
+			{
+				return BodyContainer->GetBody(BodyId);
+			}
 		}
-
-	private:
-		EContextSimPhase SimPhase = EContextSimPhase::None;
-	};
+		return TRigidBodyPtr<ContextType>();
+	}
 }
 
 #endif // UE_RIGIDPHYSICS_API_ENABLED
\ No newline at end of file
diff --git a/Engine/Source/Runtime/Experimental/RigidPhysics/Public/RigidPhysics/RigidFwd.h b/Engine/Source/Runtime/Experimental/RigidPhysics/Public/RigidPhysics/RigidFwd.h
--- a/Engine/Source/Runtime/Experimental/RigidPhysics/Public/RigidPhysics/RigidFwd.h
+++ b/Engine/Source/Runtime/Experimental/RigidPhysics/Public/RigidPhysics/RigidFwd.h
@@ -3,91 +3,105 @@
 
 // A define to allow us to toggle between old and new APIs in BodyInstance etc
 // Also can be searched to identify all engine changes required to change to the new API
-//
-// NOTE: Do not check this in set to 1 until we are done!!
 // For now it is set only in Test Target.cs files like this:
 //   GlobalDefinitions.Add("UE_RIGIDPHYSICS_API_ENABLED=1");
-//
 #ifndef UE_RIGIDPHYSICS_API_ENABLED
+// NOTE: Must be 0 (for now)
 #define UE_RIGIDPHYSICS_API_ENABLED 0
 #endif
 
 #if UE_RIGIDPHYSICS_API_ENABLED
 
+#include "Chaos/ChaosDebugName.h"
 #include "Templates/SharedPointer.h"
 
 #if DO_CHECK
 
 // Enable checks in the Rigid accessors. This can add significant overhead.
-#define UE_RIGIDPHYSICS_CHECK(expr) check(expr)
-#define UE_RIGIDPHYSICS_CHECKF(expr, format, ...) checkf(expr, format, ##__VA_ARGS__)
+#define UE_RIGIDPHYSICS_CHECK_ENABLED 1
 
-// Enable scene handle safety where we assert that the handle is still valid when accessed.
-// NOTE: When disabled, the scene handle wraps a raw pointer and we have no way of knowing
-// whether the referenced object still exists or not. It is therefre an error to ever use a 
-// scene handle after the scene is destroyed, even if UE_RIGIDPHYSICS_HANDLE_SAFTEY_ENABLED=1
-// Safety comes at the cost of pinning a weak pointer whenever the scene is accessed, which
-// can add up if the scene is pinned many times. Shipping builds should probably disable
-// handle safety unless you are trying to track down a handle lifetime issue.
-#define UE_RIGIDPHYSICS_HANDLE_SAFTEY_ENABLED 1
+// Required to show the Objects associated with a Handle in the debugger
+#define UE_RIGIDPHYSICS_CONTAINERSCENEHANDLE_ENABLED 1
+
+// Enable to assert when Handles are Pinned after their scene is destroyed.
+// This adds a Weak Pointer to SceneHandle which must be checked every time we pin so is
+// not recommended if performance is a concern
+#define UE_RIGIDPHYSICS_SCENEHANDLE_SAFTEY_ENABLED 0
 
+#else
+#define UE_RIGIDPHYSICS_CHECK_ENABLED 0
+#define UE_RIGIDPHYSICS_CONTAINERSCENEHANDLE_ENABLED 0
+#define UE_RIGIDPHYSICS_SCENEHANDLE_SAFTEY_ENABLED 0
+#endif
+
+#if UE_RIGIDPHYSICS_CHECK_ENABLED
+#define UE_RIGIDPHYSICS_CHECK(expr) check(expr)
+#define UE_RIGIDPHYSICS_CHECKF(expr, format, ...) checkf(expr, format, ##__VA_ARGS__)
 #else
 #define UE_RIGIDPHYSICS_CHECK(X)
 #define UE_RIGIDPHYSICS_CHECKF(expr, format, ...)
-#define UE_RIGIDPHYSICS_HANDLE_SAFTEY_ENABLED 0
 #endif
 
 
-// NOTES
-// 
-// Using SharedPtr and WeakPtr in a few places where they are technically required, 
-// but in practice may not be (e.g., FBodyHandle). However this is hidden as
-// an implementation detail so we have the option to switch to raw pointers if we
-// can guarantee lifetime, or retain pointer tracking in non-ship builds, etc.
-//
-
 // Math/Bounds.h
 template<typename T> struct TBounds;
 using FBounds3f = TBounds< float >;
 using FBounds3d = TBounds< double >;
 
+// User-facing classes
 namespace UE::Physics
 {
-	class FRigidBodyHandle;
-	class FRigidBodyId;
 	class FRigidBodyContainerHandle;
-	class FRigidBodyContainerId;
+	class FRigidBodyHandle;
+	class FRigidObjectId;
 	class FRigidQueryDesc;
 	class FRigidQueryHit;
 	class FRigidSceneHandle;
 
 	class IRigidBody;
 	class IRigidBodyContainer;
+	class IRigidLockable;
 	class IRigidQuery;
 	class IRigidScene;
+	class IRigidSceneFactory;
 	class IRigidSceneModifier;
+	class IRigidTyped;
 
-	template<typename ContextType> class TConstLockedRigidSceneHandle;
-	template<typename ContextType> class TRigidScenePtr;
-	template<typename ContextType> class TRigidScenePtrImpl;
-	template<typename ContextType> class TRigidBodyPtr;
-	template<typename ContextType> class TRigidBodyPtrImpl;
-	template<typename ContextType> class TRigidBodyContainerPtr;
 	template<typename ContextType> class TRigidBodyContainerPtrImpl;
+	template<typename ContextType> class TRigidBodyPtrImpl;
+	template<typename PtrImplType> class TRigidObjectPtr;
+	template<typename ContextType> class TRigidScenePtrImpl;
+
+	template<typename ContextType> using TRigidBodyContainerPtr = TRigidObjectPtr<TRigidBodyContainerPtrImpl<ContextType>>;
+	template<typename ContextType> using TRigidBodyPtr = TRigidObjectPtr<TRigidBodyPtrImpl<ContextType>>;
+	template<typename ContextType> using TRigidScenePtr = TRigidObjectPtr<TRigidScenePtrImpl<ContextType>>;
 
 	using FRigidSceneSharedPtr = TSharedPtr<IRigidScene, ESPMode::ThreadSafe>;
 	using FRigidSceneWeakPtr = TWeakPtr<IRigidScene, ESPMode::ThreadSafe>;
 	using FRigidBodyContainerSharedPtr = TSharedPtr<IRigidBodyContainer, ESPMode::ThreadSafe>;
 	using FRigidBodyContainerWeakPtr = TWeakPtr<IRigidBodyContainer, ESPMode::ThreadSafe>;
 
-	using FRigidNamePtr = TSharedPtr<FString, ESPMode::ThreadSafe>;
+	using FRigidDebugName = Chaos::FSharedDebugName;
 }
 
 namespace UE::Physics
 {
+	enum class ERigidVisitorResponse
+	{
+		Continue,
+		Break,
+	};
+
+	enum class ERigidMovementType
+	{
+		Static,
+		Kinematic,
+		Dynamic
+	};
+
 	// Physics execution is split into Game and Simulation contexts. In async
 	// modes, these are on different threads.
-	enum class EContextEnv
+	enum class ERigidContextType
 	{
 		Game,
 		Sim,
@@ -95,22 +109,30 @@ namespace UE::Physics
 
 	// The context must be locked before reading or writing to any objects in physics.
 	// The lock may be read-only or read-write.
-	enum class EContextLockType
+	enum class ERigidLockType
 	{
 		None,
 		ReadOnly,
 		ReadWrite,
 	};
 
-	class FSimContextRW;
-	template<EContextLockType LockTypeParam> class TGameContext;
-	using FGameContextRO = TGameContext<EContextLockType::ReadOnly>;
-	using FGameContextRW = TGameContext<EContextLockType::ReadWrite>;
+	template<ERigidContextType ContextTypeParam, ERigidLockType LockTypeParam> class TRigidContext;
+	using FRigidContextGameRO = TRigidContext<ERigidContextType::Game, ERigidLockType::ReadOnly>;
+	using FRigidContextGameRW = TRigidContext<ERigidContextType::Game, ERigidLockType::ReadWrite>;
+	using FRigidContextSimRW = TRigidContext<ERigidContextType::Sim, ERigidLockType::ReadWrite>;
 }
 
-namespace UE::Physics::Private
+//
+// Internal Forward Decls
+//
+
+namespace UE::Physics
 {
-	class UE_INTERNAL FRigidAccessor;
+	template<typename ObjectType> class UE_INTERNAL TRigidObjectRegistry;
+	using UE_INTERNAL FRigidBodyContainerRegistry = TRigidObjectRegistry<IRigidBodyContainer>;
+	using UE_INTERNAL FRigidBodyRegistry = TRigidObjectRegistry<IRigidBody>;
 }
 
+extern uint32 GetTypeHash(UE::Physics::FRigidObjectId& RigidBodyId);
+
 #endif // UE_RIGIDPHYSICS_API_ENABLED
\ No newline at end of file
diff --git a/Engine/Source/Runtime/Experimental/RigidPhysics/Public/RigidPhysics/RigidLockable.h b/Engine/Source/Runtime/Experimental/RigidPhysics/Public/RigidPhysics/RigidLockable.h
new file mode 100644
--- /dev/null
+++ b/Engine/Source/Runtime/Experimental/RigidPhysics/Public/RigidPhysics/RigidLockable.h
@@ -0,0 +1,29 @@
+// Copyright Epic Games, Inc. All Rights Reserved.
+#pragma once
+
+#include "RigidPhysics/RigidFwd.h"
+#include "RigidPhysics/RigidTyped.h"
+
+#if UE_RIGIDPHYSICS_API_ENABLED
+
+namespace UE::Physics
+{
+	class UE_INTERNAL IRigidLockable : public IRigidTyped
+	{
+	public:
+		UE_RIGIDPHYSICS_RIGIDTYPED_DECL(RIGIDPHYSICS_API, IRigidLockable);
+
+		UE_INTERNAL IRigidLockable()
+		{
+		}
+
+		UE_INTERNAL virtual ~IRigidLockable()
+		{
+		}
+
+		UE_INTERNAL RIGIDPHYSICS_API virtual void Lock(ERigidLockType InLockType) const = 0;
+		UE_INTERNAL RIGIDPHYSICS_API virtual void Unlock(ERigidLockType InLockType) const = 0;
+	};
+}
+
+#endif // UE_RIGIDPHYSICS_API_ENABLED
\ No newline at end of file
diff --git a/Engine/Source/Runtime/Experimental/RigidPhysics/Public/RigidPhysics/RigidLog.h b/Engine/Source/Runtime/Experimental/RigidPhysics/Public/RigidPhysics/RigidLog.h
--- a/Engine/Source/Runtime/Experimental/RigidPhysics/Public/RigidPhysics/RigidLog.h
+++ b/Engine/Source/Runtime/Experimental/RigidPhysics/Public/RigidPhysics/RigidLog.h
@@ -7,7 +7,7 @@
 
 #include "Logging/StructuredLog.h"
 
-RIGIDPHYSICS_API DECLARE_LOG_CATEGORY_EXTERN(LogChaosRB, Log, All);
+RIGIDPHYSICS_API DECLARE_LOG_CATEGORY_EXTERN(LogRigidPhysics, Log, All);
 
 namespace UE::Physics::Private
 {
@@ -16,7 +16,7 @@ namespace UE::Physics::Private
 	{
 		// This warning means a function is being called in a context where it is not supported.
 		// E.g., Calling SetMass from the physics thread.
-		UE_LOGFMT(LogChaosRB, Warning, "{0} not supported on {1}", FuncName, TypeName);
+		UE_LOGFMT(LogRigidPhysics, Warning, "{0} not supported on {1}", FuncName, TypeName);
 	}
 }
 
diff --git a/Engine/Source/Runtime/Experimental/RigidPhysics/Public/RigidPhysics/RigidModifier.h b/Engine/Source/Runtime/Experimental/RigidPhysics/Public/RigidPhysics/RigidModifier.h
--- a/Engine/Source/Runtime/Experimental/RigidPhysics/Public/RigidPhysics/RigidModifier.h
+++ b/Engine/Source/Runtime/Experimental/RigidPhysics/Public/RigidPhysics/RigidModifier.h
@@ -1,8 +1,8 @@
 // Copyright Epic Games, Inc. All Rights Reserved.
 #pragma once
 
-#include "RigidPhysics/RigidContext.h"
 #include "Misc/Crc.h"
+#include "RigidPhysics/RigidFwd.h"
 #include "Templates/SharedPointer.h"
 
 #if UE_RIGIDPHYSICS_API_ENABLED
@@ -21,26 +21,26 @@ namespace UE::Physics
 		}
 
 		// Called from the game context before physics is ticked for this frame. Can be used to
-		// update inputs to the objects used in the physics context.
-		virtual void PreSimulate(TRigidScenePtr<FGameContextRW>& ScenePtr)
+		// update inputs to the objects used in the sim context.
+		virtual void PreSimulate(const FRigidContextGameRW& Context)
 		{
 		}
 
 		// Call from the game context after physics is ticked for this frame, but before the
-		// rigid body results have been sent out to the game framework objects (Actors, Components, etc.)
-		virtual void PostSimulate(TRigidScenePtr<FGameContextRW>& ScenePtr)
+		// sim results have been sent out to the game framework objects (Actors, Components, etc.)
+		virtual void PostSimulate(const FRigidContextGameRW& Context)
 		{
 		}
 
 		// Called from the simulation context before each simulation tick. At this point you can add forces, move bodies,
 		// or do pretty much anything to the scene or objects in it, and the results will impact the next tick.
-		virtual void SimPreTick(TRigidScenePtr<FSimContextRW>& ScenePtr)
+		virtual void PreTick(const FRigidContextSimRW& Context)
 		{
 		}
 
 		// Called from the simulation context at the end of each simulation tick. At this point you can modify physics
 		// object state and the altered results will be reported back to the game, and be used for the next tick.
-		virtual void SimPostTick(TRigidScenePtr<FSimContextRW>& ScenePtr)
+		virtual void PostTick(const FRigidContextSimRW& Context)
 		{
 		}
 	};
diff --git a/Engine/Source/Runtime/Experimental/RigidPhysics/Public/RigidPhysics/RigidObjectId.h b/Engine/Source/Runtime/Experimental/RigidPhysics/Public/RigidPhysics/RigidObjectId.h
new file mode 100644
--- /dev/null
+++ b/Engine/Source/Runtime/Experimental/RigidPhysics/Public/RigidPhysics/RigidObjectId.h
@@ -0,0 +1,79 @@
+// Copyright Epic Games, Inc. All Rights Reserved.
+#pragma once
+
+#include "RigidPhysics/RigidFwd.h"
+
+#if UE_RIGIDPHYSICS_API_ENABLED
+
+namespace UE::Physics
+{
+	// An ID for an object held in an array that is guaranteed to be unique and 
+	// never reused within its container (as long as less than 1<<31 objects are allocated). 
+	// There is no guarantee of uniqueness between containers, which means you need a 
+	// Container + Id to access an object. (E.g., see FRigidBodyHandle, and FRigidBodyContainerHandle).
+	class FRigidObjectId
+	{
+	public:
+		FRigidObjectId()
+			: ArrayIndex(INDEX_NONE)
+			, Epoch(INDEX_NONE)
+		{
+		}
+
+		FRigidObjectId(int32 InIndex, int32 InEpoch)
+			: ArrayIndex(InIndex)
+			, Epoch(InEpoch)
+		{
+		}
+
+		bool IsValid() const
+		{
+			return (ArrayIndex >= 0);
+		}
+
+		UE_INTERNAL int32 GetArrayIndex() const
+		{
+			return ArrayIndex;
+		}
+
+		UE_INTERNAL int32 GetEpoch() const
+		{
+			return Epoch;
+		}
+
+		UE_INTERNAL void Reset()
+		{
+			ArrayIndex = INDEX_NONE;
+			Epoch = INDEX_NONE;
+		}
+
+		friend bool operator==(const FRigidObjectId& L, const FRigidObjectId& R) = default;
+
+		friend bool operator<(const FRigidObjectId& L, const FRigidObjectId& R)
+		{
+			return L.ToUInt64() < R.ToUInt64();
+		}
+
+		friend uint32 GetTypeHash(const FRigidObjectId& InId)
+		{
+			return GetTypeHash(InId.ToUInt64());
+		}
+
+	protected:
+		uint64 ToUInt64() const
+		{
+			return (((uint64)Epoch << 32) & 0xFFFFFFFF00000000LL) | (((uint64)ArrayIndex) & 0x00000000FFFFFFFFLL);
+		}
+
+		void FromInt64(uint64 InInt)
+		{
+			Epoch = (int32)((InInt >> 32) & 0xFFFFFFFF);
+			ArrayIndex = (int32)(InInt & 0xFFFFFFFF);
+		}
+
+		int32 ArrayIndex;
+		int32 Epoch;
+	};
+}
+
+#endif // UE_RIGIDPHYSICS_API_ENABLED
diff --git a/Engine/Source/Runtime/Experimental/RigidPhysics/Public/RigidPhysics/RigidObjectPtr.h b/Engine/Source/Runtime/Experimental/RigidPhysics/Public/RigidPhysics/RigidObjectPtr.h
new file mode 100644
--- /dev/null
+++ b/Engine/Source/Runtime/Experimental/RigidPhysics/Public/RigidPhysics/RigidObjectPtr.h
@@ -0,0 +1,123 @@
+// Copyright Epic Games, Inc. All Rights Reserved.
+#pragma once
+
+#include "RigidPhysics/RigidFwd.h"
+
+#if UE_RIGIDPHYSICS_API_ENABLED
+
+namespace UE::Physics
+{
+
+	// A Context-Specific Smart Pointer to a physics object. 
+	//
+	// This is a wrapper around a smart pointer implementation object like TRigidBodyPtrImpl, 
+	// TRigidBodyContainerPtrImpl, etc. It provides operator-> to give it const and non-const pointer 
+	// semantics and behaviour.
+	//
+	// May be "null". Check IsValid() (or cast to bool) before dereferencing.
+	//
+	// 	Cannot be copied or stored - use Handles for that.
+	//
+	// NOTE: Much like a regular pointer, a "valid" pointer may still reference deleted memory if the 
+	// object is destroyed after the Handle is Pinned. Do not use pointers, even to check for validity, 
+	// after the underlying object has been destroyed. If the underlying object may have been destroyed, 
+	// you should re-Pin the Handle to obtain a new pointer.
+	//
+	template<typename PointerImplType>
+	class TRigidObjectPtr
+	{
+	public:
+		using FPointerImpl = PointerImplType;
+		using FContext = FPointerImpl::FContext;
+		using FInterface = FPointerImpl::FInterface;
+		using FHandle = FPointerImpl::FHandle;
+
+		UE_NONCOPYABLE(TRigidObjectPtr);
+
+		TRigidObjectPtr()
+			: Impl(nullptr)
+		{
+		}
+
+		UE_INTERNAL TRigidObjectPtr(FInterface* InObject)
+			: Impl(InObject)
+		{
+		}
+
+		~TRigidObjectPtr()
+		{
+		}
+
+		// Reset the pointer (it will now be Invalid)
+		void Reset()
+		{
+			Impl.Reset();
+		}
+
+		// Has this pointer been initialized with a valid body pointer?
+		// NOTE: This will still return true if the underlying object is destroyed after the Handle is Pinned.
+		// If this is a possibility, the Handle should be re-Pinned.
+		bool IsValid() const
+		{
+			return Impl.IsValid();
+		}
+
+		// Auto cast to bool to support using in if statements
+		//	e.g., if (TRigidBodyPtr<FRigidContextGameRW> Body = BodyHandle.Pin(Context)) { ... }
+		operator bool() const
+		{
+			return IsValid();
+		}
+
+		// Extract the handle
+		FHandle ToHandle() const
+		{
+			return Impl.GetHandle();
+		}
+
+		// Cast to a handle
+		operator FHandle() const
+		{
+			return ToHandle();
+		}
+
+		// Const pointer sematics
+		const FPointerImpl* operator->() const
+		{
+			CheckIsValid();
+			return &Impl;
+		}
+
+		// Non-const pointer sematics
+		FPointerImpl* operator->() const requires (FContext::bWriteEnabled)
+		{
+			CheckIsValid();
+			return &Impl;
+		}
+
+		// Extract the object interface (internal use only)
+		UE_INTERNAL FInterface* Get() const
+		{
+			return Impl.Get();
+		}
+
+		UE_INTERNAL void Init(FInterface* InObject)
+		{
+			Impl = FPointerImpl(InObject);
+		}
+
+	private:
+		void CheckIsValid() const
+		{
+			UE_RIGIDPHYSICS_CHECKF(IsValid(), TEXT("Ptr is null. Check IsValid() before dereferencing"));
+		}
+
+		// Mutable to support const TRigidObjectPtr to an non-const object
+		// i.e., this wrapper cannot change, but the pointee is mutabale
+		// if the wrapper was created from a mutable context. This avoids
+		// an undefined behaviour compiler warning.
+		mutable FPointerImpl Impl;
+	};
+}
+
+#endif // UE_RIGIDPHYSICS_API_ENABLED
diff --git a/Engine/Source/Runtime/Experimental/RigidPhysics/Public/RigidPhysics/RigidObjectRegistry.h b/Engine/Source/Runtime/Experimental/RigidPhysics/Public/RigidPhysics/RigidObjectRegistry.h
new file mode 100644
--- /dev/null
+++ b/Engine/Source/Runtime/Experimental/RigidPhysics/Public/RigidPhysics/RigidObjectRegistry.h
@@ -0,0 +1,244 @@
+// Copyright Epic Games, Inc. All Rights Reserved.
+#pragma once
+
+#include "RigidPhysics/RigidFwd.h"
+
+#if UE_RIGIDPHYSICS_API_ENABLED
+
+#include "Chaos/ObjectPool.h"
+#include "RigidPhysics/RigidLog.h"
+#include "RigidPhysics/RigidObjectId.h"
+
+namespace UE::Physics
+{
+	// Info held in the TRigidObjectRegistry indexed by ObjectId
+	template<typename ObjectPtrType>
+	class UE_INTERNAL TRigidObjectRegistryItem
+	{
+	public:
+		using FObjectPtrType = ObjectPtrType;
+
+		UE_INTERNAL TRigidObjectRegistryItem()
+			: Object(nullptr)
+			, Epoch(0)
+			, AddedObjectsIndex(INDEX_NONE)
+		{
+		}
+
+		// The object that was added
+		FObjectPtrType Object;
+
+		// How many times the Object's slot in TRigidObjectRegistry has been used.
+		// This will be unique for every reuse of the slot and is used by Object
+		// Handles to determine if they have expired or not.
+		int32 Epoch;
+
+		// Index into the AddedObjects array in TRigidObjectRegistry
+		int32 AddedObjectsIndex;
+	};
+
+
+	// A Registry for Rigid Objects that allows us to convert from a Handle to an Object Pointer.
+	// Each registered object is assigned an Id, which is a combination of an array index and an epoch counter.
+	// The array index and epoch counter combined are unique for every allocated object, even when array indices
+	// get reused after an object is unregistered and a new one takes its place.
+	// 
+	// ObjectType must have a GetHandle() function.
+	//
+	// TODO_CHAOSAPI: ideally would not require a GetHandle() function on the object - it is only used the repack
+	// the AddedObjects array. Should be a template policy instead to get to the Id? See RemoveObject().
+	// 
+	// TODO_CHAOSAPI: this container will fail to work if we reuse any ObjectsById array index enough times to
+	// wrap its Epoch. We can mitigate this by removing those indices from the free list. This will leave unusable
+	// gaps in the array but would allow for infinite growth at the cost of some wasted memory.
+	//
+	template<typename ObjectType>
+	class UE_INTERNAL TRigidObjectRegistry
+	{
+	public:
+		using FObjectType = ObjectType;
+		using FObjectPtrType = ObjectType*;
+		using FObjectRegistryItem = TRigidObjectRegistryItem<FObjectPtrType>;
+
+	public:
+		UE_INTERNAL TRigidObjectRegistry(int32 InSlack)
+			: ArraySlack(InSlack)
+			, TotalNumAdded(0)
+			, bIsVisiting(false)
+		{
+		}
+
+		UE_INTERNAL ~TRigidObjectRegistry()
+		{
+		}
+
+		// The number of objects that have been added (and not yet removed)
+		UE_INTERNAL int32 GetNumObjects() const
+		{
+			return ObjectsById.Num() - ObjectsByIdFreeIndices.Num();
+		}
+
+		// True if there are no objects in the registry
+		UE_INTERNAL bool IsEmpty() const
+		{
+			return GetNumObjects() == 0;
+		}
+
+		// Add an object to the registry and assign it an Id. The Id is persistent
+		// and will never be reused as long as less that 1<<31 objects are added.
+		// After that, Ids may start to repeat.
+		UE_INTERNAL FRigidObjectId AddObject(FObjectPtrType InObject)
+		{
+			UE_RIGIDPHYSICS_CHECK(InObject != nullptr);
+
+			// Cannot be called while visiting elements
+			UE_RIGIDPHYSICS_CHECK(!bIsVisiting);
+
+			// Assert if we add enough objects such that we might wrap the Epoch counter
+			// (this is a lower bound, but is consistent and doesn't depend on add/remove order)
+			// See class comments for future solution to this issue.
+			UE_RIGIDPHYSICS_CHECK(TotalNumAdded < std::numeric_limits<int32>::max());
+			++TotalNumAdded;
+
+			// Make sure we have free array entries
+			if (ObjectsByIdFreeIndices.IsEmpty())
+			{
+				ReserveArray(ObjectsById.Num() + ArraySlack);
+			}
+
+			// Assign an array index
+			int32 ObjectsByIdIndex = ObjectsByIdFreeIndices.Pop(EAllowShrinking::No);
+
+			// Put the item in the array
+			FObjectRegistryItem& ObjectsByIdItem = ObjectsById[ObjectsByIdIndex];
+			UE_RIGIDPHYSICS_CHECK(ObjectsByIdItem.Object == nullptr);
+			ObjectsByIdItem.Object = InObject;
+			ObjectsByIdItem.AddedObjectsIndex = AddedObjects.Add(InObject);
+
+			// Build the Id from the Index and Epoch
+			return FRigidObjectId(ObjectsByIdIndex, ObjectsByIdItem.Epoch);
+		}
+
+		UE_INTERNAL void RemoveObject(const FRigidObjectId& InObjectId)
+		{
+			// Cannot be called while visiting elements
+			UE_RIGIDPHYSICS_CHECK(!bIsVisiting);
+
+			// Call FindObject() first if presence is unknown
+			UE_RIGIDPHYSICS_CHECK(ObjectsById.IsValidIndex(InObjectId.GetArrayIndex()));
+
+			// Get the entry
+			FObjectRegistryItem& ObjectsByIdItem = ObjectsById[InObjectId.GetArrayIndex()];
+
+			// NOTE: We already know that this object was in the map before calling
+			// so the epoch should always match
+			UE_RIGIDPHYSICS_CHECK(ObjectsByIdItem.Epoch == InObjectId.GetEpoch());
+
+			// Remove from the Allocated array (and fix up the AddedObjectsIndex of the moved item)
+			int32 AddedObjectsIndex = ObjectsByIdItem.AddedObjectsIndex;
+			if (AddedObjects.IsValidIndex(AddedObjectsIndex))
+			{
+				// If this fires, we have a ObjectsByIdItem.AddedObjectsIndex mismatch
+				UE_RIGIDPHYSICS_CHECK(AddedObjects[AddedObjectsIndex] == ObjectsByIdItem.Object);
+
+				AddedObjects.RemoveAtSwap(AddedObjectsIndex, EAllowShrinking::No);
+
+				if (AddedObjects.IsValidIndex(AddedObjectsIndex))
+				{
+					// TODO_CHAOSAPI: GetId() is only required to maintain the packed allocated object array. Try to remove it.
+					FRigidObjectId RelocatedObjectId = AddedObjects[AddedObjectsIndex]->GetHandle().GetId();
+					FObjectRegistryItem& ReloctedArrayItem = ObjectsById[RelocatedObjectId.GetArrayIndex()];
+					ReloctedArrayItem.AddedObjectsIndex = AddedObjectsIndex;
+				}
+			}
+
+			// Reset the entry and increment the epoch to invalidate outstanding Ids using the same index
+			ObjectsByIdItem.Object = nullptr;
+			ObjectsByIdItem.Epoch = ObjectsByIdItem.Epoch + 1;
+			ObjectsByIdItem.AddedObjectsIndex = INDEX_NONE;
+
+			// Add the array item back into the free pool
+			ObjectsByIdFreeIndices.Push(InObjectId.GetArrayIndex());
+
+			// TODO_CHAOSAPI: we could shrink the ObjectArray here if we have too much slack, but
+			// we'd need to know what the max allocated index was, and would have to prune
+			// the free indices.
+		}
+
+		UE_INTERNAL FObjectType* FindObject(const FRigidObjectId& InId) const
+		{
+			if (ObjectsById.IsValidIndex(InId.GetArrayIndex()))
+			{
+				const FObjectRegistryItem& ArrayItem = ObjectsById[InId.GetArrayIndex()];
+				if (ArrayItem.Epoch == InId.GetEpoch())
+				{
+					return ArrayItem.Object;
+				}
+			}
+			return nullptr;
+		}
+
+		// Visit all of the added objects in somewhat arbitrary order that may change whenever
+		// objects are added or removed.
+		UE_INTERNAL void VisitObjects(const TFunctionRef<ERigidVisitorResponse(FObjectType*)>& Visitor)
+		{
+			bIsVisiting = true;
+			for (FObjectType* Object : AddedObjects)
+			{
+				ERigidVisitorResponse VisitorResponse = Visitor(Object);
+				if (VisitorResponse == ERigidVisitorResponse::Break)
+				{
+					break;
+				}
+			}
+			bIsVisiting = false;
+		}
+
+	private:
+		void ReserveArray(int32 MaxItems)
+		{
+			if (MaxItems > ObjectsById.Num())
+			{
+				// Resize the arrays
+				int32 PrevNumItems = ObjectsById.Num();
+				ObjectsById.SetNum(MaxItems);
+				ObjectsByIdFreeIndices.Reserve(MaxItems);
+				AddedObjects.Reserve(MaxItems);
+
+				// Add to the free indices
+				// Loop backawards so we allocate lowest indices first in AddObject
+				for (int32 ItemIndex = MaxItems - 1; ItemIndex >= PrevNumItems; --ItemIndex)
+				{
+					ObjectsByIdFreeIndices.Push(ItemIndex);
+				}
+			}
+		}
+
+		// An array with no holes containing all currently allocated objects indexed by ObjectId.AddedObjectsIndex.
+		// This exists to support iteration over items and easier debugging - we may be able remove it later.
+		// NOTE: Object position in this array is not persistent (RemoveObject uses RemoveAtSwap)
+		// TODO_CHAOSAPI: see if we really need the AddedObjects array.
+		TArray<FObjectType*> AddedObjects;
+
+		// An Id->Object map, implemented as an array indexed by ObjectId.ArrayIndex.
+		// Indices into this array are persistent as long as the object exists, but the array may contain holes.
+		// Array Entries will be reused after objects have been deleted, and the Epoch counter is used to 
+		// validate handles against the current object at the array index.
+		// This array is presized and will often contain many empty (unused) elements.
+		TArray<FObjectRegistryItem> ObjectsById;
+
+		// A pool of unused entries in ObjectsById.
+		TArray<int32> ObjectsByIdFreeIndices;
+
+		// How much extra to allocate in ObjectsById when we run out of space.
+		int32 ArraySlack;
+
+		// How many objects have been added. This is only used to assert if we may reuse Ids.
+		int32 TotalNumAdded;
+
+		// Use to trap modifications while in the Visitor loop
+		bool bIsVisiting;
+	};
+}
+
+#endif // UE_RIGIDPHYSICS_API_ENABLED
\ No newline at end of file
diff --git a/Engine/Source/Runtime/Experimental/RigidPhysics/Public/RigidPhysics/RigidScene.h b/Engine/Source/Runtime/Experimental/RigidPhysics/Public/RigidPhysics/RigidScene.h
--- a/Engine/Source/Runtime/Experimental/RigidPhysics/Public/RigidPhysics/RigidScene.h
+++ b/Engine/Source/Runtime/Experimental/RigidPhysics/Public/RigidPhysics/RigidScene.h
@@ -5,418 +5,220 @@
 
 #if UE_RIGIDPHYSICS_API_ENABLED
 
-#include "RigidPhysics/RigidContext.h"
+#include "RigidPhysics/RigidLockable.h"
+#include "RigidPhysics/RigidLog.h"
+#include "RigidPhysics/RigidObjectId.h"
+#include "RigidPhysics/RigidObjectPtr.h"
+#include "RigidPhysics/RigidTyped.h"
 #include "Templates/SharedPointer.h"
 
 namespace UE::Physics
 {
-	inline FRigidNamePtr MakeRigidName(const TCHAR* InName)
+	class IRigidSceneFactory
 	{
-		return MakeShared<FString, ESPMode::ThreadSafe>(FString(InName));
-	}
+	public:
+		IRigidSceneFactory()
+		{
+		}
 
-	template<typename... ArgTypes>
-	inline FRigidNamePtr MakeRigidName(const TCHAR* InFmt, ArgTypes... Args)
-	{
-		// TODO: Optimize
-		FStringFormatOrderedArguments ArgsArray = { std::forward<ArgTypes>(Args)... };
-		return MakeShared<FString, ESPMode::ThreadSafe>(FString::Format(InFmt, ArgsArray));
-	}
+		virtual ~IRigidSceneFactory()
+		{
+		}
+
+		virtual FRigidSceneHandle CreateScene(const FRigidDebugName& InName) = 0;
+		virtual void DestroyScene(const FRigidSceneHandle& InSceneHandle) = 0;
+	};
 
 	// The interface to a rigid body scene. The rigid body scene owns the physics simulations and
 	// manages a set of rigid body containers which in turn manage sets of rigid bodies.
-	class RIGIDPHYSICS_API IRigidScene
+	class UE_INTERNAL IRigidScene : public IRigidLockable
 	{
 	public:
-		IRigidScene();
-		virtual ~IRigidScene();
+		UE_RIGIDPHYSICS_RIGIDTYPED_DECL(RIGIDPHYSICS_API, IRigidScene);
 
-		// Lock the scene for read or write if supported (GameContext only), otherwise do nothing
-		virtual void Lock(EContextLockType InLockType)
+		IRigidScene()
 		{
-			RIGIDPHYSICS_API_UNSUPPORTED();
 		}
 
-		// Unlock the scene if supported (GameContext only), otherwise do nothing
-		virtual void Unlock(EContextLockType InLockType)
+		virtual ~IRigidScene()
 		{
-			RIGIDPHYSICS_API_UNSUPPORTED();
 		}
 
-		// Create a rigid body. It will take part in the simulation until RegisterRigidBody is called.
-		virtual IRigidBody* CreateRigidBody(const FRigidNamePtr& InName, ERigidMovementType InMovementType)
-		{
-			RIGIDPHYSICS_API_UNSUPPORTED();
-			return nullptr;
-		}
+		UE_INTERNAL RIGIDPHYSICS_API virtual FRigidSceneHandle GetHandle() const = 0;
 
-		// Destroy a rigid body. It will no longer take part in the simulation.
-		virtual void DestroyRigidBody(const FRigidBodyHandle& RigidBodyHandle)
-		{
-			RIGIDPHYSICS_API_UNSUPPORTED();
-		}
+		UE_INTERNAL RIGIDPHYSICS_API virtual IRigidBodyContainer* GetBodyContainer(const FRigidObjectId& InBodyContainerId) const = 0;
 
-		// Add a rigid body to the scene. It will now be colideable and simulated as appropriate.
-		virtual void RegisterRigidBody(const FRigidBodyHandle& BodyHandle)
+		UE_INTERNAL RIGIDPHYSICS_API virtual bool Query(const FRigidQueryDesc& InQueryDesc, TArray<FRigidQueryHit>& OutHits)
 		{
 			RIGIDPHYSICS_API_UNSUPPORTED();
+			return false;
 		}
 
-		// Remove a rigid body from the scene.
-		virtual void UnregisterRigidBody(const FRigidBodyHandle& BodyHandle)
+		// Advance the simulation by Dt. If the implementation is asynchronous, this will return immediately
+		UE_INTERNAL RIGIDPHYSICS_API virtual void Tick(float InDt)
 		{
 			RIGIDPHYSICS_API_UNSUPPORTED();
 		}
 
-		// Get the specified rigid body container.
-		virtual IRigidBodyContainer* GetRigidBodyContainer(const FRigidBodyContainerId& InRigidBodyContainerId)
+		// GT Only API - Move this somehow
+		UE_INTERNAL RIGIDPHYSICS_API virtual void RegisterModifier(IRigidSceneModifier* InModifier)
 		{
 			RIGIDPHYSICS_API_UNSUPPORTED();
-			return nullptr;
 		}
 
-		virtual void RegisterModifier(IRigidSceneModifier* Modifier)
-		{
-			RIGIDPHYSICS_API_UNSUPPORTED();
-		}
-		
-		virtual void UnregisterModifier(IRigidSceneModifier* Modifier)
+		UE_INTERNAL RIGIDPHYSICS_API virtual void UnregisterModifier(IRigidSceneModifier* InModifier)
 		{
 			RIGIDPHYSICS_API_UNSUPPORTED();
 		}
 
-		virtual bool Query(const FRigidQueryDesc& QueryDesc, TArray<FRigidQueryHit>& OutHits)
+		UE_INTERNAL RIGIDPHYSICS_API virtual IRigidBody* CreateBody(const FRigidDebugName& InName, ERigidMovementType InMovementType)
 		{
 			RIGIDPHYSICS_API_UNSUPPORTED();
-			return false;
+			return nullptr;
 		}
 
-		// Advance the simulation by Dt. If the implementation is asynchronous, this will return immediately
-		virtual void Tick(float Dt)
+		UE_INTERNAL RIGIDPHYSICS_API virtual void DestroyBody(IRigidBody* InBody)
 		{
 			RIGIDPHYSICS_API_UNSUPPORTED();
 		}
 
-	protected:
-
-		virtual const char* GetTypeName() const
-		{
-			return "Unknown";
-		}
 	};
 
-
-	// Exists only to give pointer semantics to TRigidScenePtr.
-	// See TRigidScenePtr for comments.
-	template<typename ContextType>
-	class TRigidScenePtrImpl
+	//
+	// Base class for the Game Context Scene. This provides the BodyContainerRegistry, partly for convenience
+	// and partly so that we can convert Handles to Objects in the debugger with natvis
+	//
+	class UE_INTERNAL FRigidScene : public IRigidScene
 	{
 	public:
-		using FContext = ContextType;
-
-		TRigidScenePtrImpl()
-			: SceneRaw(nullptr)
-		{
-		}
-
-		~TRigidScenePtrImpl()
-		{
-			Reset();
-		}
+		RIGIDPHYSICS_API FRigidScene();
+		RIGIDPHYSICS_API virtual ~FRigidScene();
 
-		// Impl in RigidBody.h
-		TRigidBodyPtr<FContext> CreateRigidBody(const FRigidNamePtr& InName, ERigidMovementType InMovementType) const;
+		UE_INTERNAL RIGIDPHYSICS_API virtual IRigidBodyContainer* GetBodyContainer(const FRigidObjectId& InBodyContainerId) const override final;
 
-		// Impl in RigidBody.h
-		void DestroyRigidBody(const FRigidBodyHandle& BodyHandle) const;
-
-		// Impl in RigidBody.h
-		TRigidBodyContainerPtr<FContext> GetRigidBodyContainer(const FRigidBodyContainerId& InBodyContainerId) const;
-
-		// Impl in RigidBody.h
-		void RegisterRigidBody(const FRigidBodyHandle& BodyHandle);
-
-		// Impl in RigidBody.h
-		void UnregisterRigidBody(const FRigidBodyHandle& BodyHandle);
-
-		void RegisterModifier(IRigidSceneModifier* Modifier)
-		{
-			SceneRaw->RegisterModifier(Modifier);
-		}
-
-		void UnregisterModifier(IRigidSceneModifier* Modifier)
-		{
-			SceneRaw->UnregisterModifier(Modifier);
-		}
-
-		bool Query(const FRigidQueryDesc& QueryDesc, TArray<FRigidQueryHit>& OutHits)  const
-		{
-			return SceneRaw->Query(QueryDesc, OutHits);
-		}
-
-		void Tick(float Dt)
+	protected:
+		UE_INTERNAL RIGIDPHYSICS_API FRigidBodyContainerRegistry* GetBodyContainerRegistry() const
 		{
-			SceneRaw->Tick(Dt);
+			return BodyContainerRegistry.Get();
 		}
 
 	private:
-		friend class TRigidScenePtr<FContext>;
-		friend class Private::FRigidAccessor;
-
-		TRigidScenePtrImpl(IRigidScene* InScene)
-			: SceneRaw(nullptr)
-		{
-			Init(InScene);
-		}
+		TUniquePtr<FRigidBodyContainerRegistry> BodyContainerRegistry;
 
-		TRigidScenePtrImpl(const TRigidScenePtrImpl& InImpl)
-			: SceneRaw(nullptr)
-		{
-			*this = InImpl;
-		}
-
-		TRigidScenePtrImpl(TRigidScenePtrImpl&& InImpl)
-			: SceneRaw(nullptr)
-		{
-			// Take over incoming scene lock - do not add a ref
-			SceneRaw = InImpl.SceneRaw;
-			InImpl.SceneRaw = nullptr;
-		}
-
-		TRigidScenePtrImpl& operator=(const TRigidScenePtrImpl& InImpl)
-		{
-			if (InImpl.SceneRaw != SceneRaw)
-			{
-				// Undo current lock if there is one
-				Reset();
-
-				// Scene will already be lock - add a ref
-				Init(InImpl.SceneRaw);
-			}
-			return *this;
-		}
-
-		TRigidScenePtrImpl& operator=(TRigidScenePtrImpl&& InImpl)
-		{
-			// Undo current lock if there is one
-			Reset();
-
-			// Take over incoming scene lock - do not add a ref
-			SceneRaw = InImpl.SceneRaw;
-			InImpl.SceneRaw = nullptr;
-
-			return *this;
-		}
-
-		bool IsValid() const
-		{
-			return (SceneRaw != nullptr);
-		}
-
-		void Init(IRigidScene* InScene)
-		{
-			SceneRaw = InScene;
-			if (SceneRaw != nullptr)
-			{
-				SceneRaw->Lock(FContext::LockType);
-			}
-		}
-
-		void Reset()
-		{
-			if (SceneRaw != nullptr)
-			{
-				SceneRaw->Unlock(FContext::LockType);
-				SceneRaw = nullptr;
-			}
-		}
-
-		IRigidScene* SceneRaw;
 	};
 
 
-	// A context-aware IRigidScene pointer that simply forwards all context-legal operations to IRigidScene. 
-	// Provides access to a restricted Scene API where the restrictions depend on the context 
-	// (template-parameter), and attempts to call disallowed functions (in the context) produce compiler errors.
-	// 
-	// Uses pointer semantics and may be "null". Check IsValid() before dereferencing. Also implicitly casts 
-	// to boolean for validity checks.
-	// 
-	// Cannot be copied or stored - use FRigidSceneHandle for that.
-	//
-	// Where appropriate, the Scene Ptr holds a lock on the scene. The type of lock depends on the context.
-	// E.g., in the async implementation, for GameContexts it may be read-only or write-enabled, but the
-	// PhysicsContext does not support multi-threaded access, and therefore there is no lock.
+	// A wrapper around IRigidScene.
+	// Exists only to give pointer semantics to TRigidScenePtr.
 	template<typename ContextType>
-	class TRigidScenePtr
+	class UE_INTERNAL TRigidScenePtrImpl
 	{
 	public:
 		using FContext = ContextType;
-		using FRigidScenePtrImpl = TRigidScenePtrImpl<FContext>;
+		using FInterface = IRigidScene;
+		using FHandle = FRigidSceneHandle;
 
-		TRigidScenePtr() = default;
-		~TRigidScenePtr() = default;
-
-		void Reset()
+		UE_INTERNAL TRigidScenePtrImpl()
+			: SceneRaw(nullptr)
 		{
-			SceneImpl.Reset();
 		}
 
-		bool IsValid() const
+		UE_INTERNAL TRigidScenePtrImpl(IRigidScene* InScene)
+			: SceneRaw(InScene)
 		{
-			return SceneImpl.IsValid();
 		}
 
-		operator bool() const
+		~TRigidScenePtrImpl()
 		{
-			return IsValid();
 		}
 
-		// Pointer Semantics
-		const FRigidScenePtrImpl* operator->() const requires (!FContext::bWriteEnabled)
+		void Reset()
 		{
-			UE_RIGIDPHYSICS_CHECKF(IsValid(), TEXT("Attempt to dereference null scene handle. Missing IsValid() check"));
-			return &SceneImpl;
+			SceneRaw = nullptr;
 		}
 
-		// Pointer Semantics
-		FRigidScenePtrImpl* operator->() const requires (FContext::bWriteEnabled)
+		bool IsValid() const
 		{
-			UE_RIGIDPHYSICS_CHECKF(IsValid(), TEXT("Attempt to dereference null scene handle. Missing IsValid() check"));
-			return const_cast<FRigidScenePtrImpl*>(&SceneImpl);
+			return (SceneRaw != nullptr);
 		}
 
-	private:
-		friend class FRigidSceneHandle;
-		friend class Private::FRigidAccessor;
-		friend class TRigidScenePtrImpl<FContext>;
-		friend class TRigidBodyContainerPtrImpl<FContext>;
-		friend class TRigidBodyPtrImpl<FContext>;
-
-		TRigidScenePtr(IRigidScene* InScene)
-			: SceneImpl(InScene)
-		{
-		}
+		FRigidSceneHandle GetHandle() const;
 
-		TRigidScenePtr(const TRigidScenePtr<ContextType>& InScene)
-			: SceneImpl(InScene.SceneImpl)
+		IRigidBodyContainer* GetBodyContainer(const FRigidObjectId& InBodyContainerId) const
 		{
+			return SceneRaw->GetBodyContainer(InBodyContainerId);
 		}
 
-		TRigidScenePtr(TRigidScenePtr<ContextType>&& InScene)
-			: SceneImpl(MoveTemp(InScene.SceneImpl))
+		void RegisterModifier(IRigidSceneModifier* Modifier)
 		{
+			SceneRaw->RegisterModifier(Modifier);
 		}
 
-		TRigidScenePtr(const TRigidScenePtrImpl<ContextType>& InSceneImpl)
-			: SceneImpl(InSceneImpl)
+		void UnregisterModifier(IRigidSceneModifier* Modifier)
 		{
+			SceneRaw->UnregisterModifier(Modifier);
 		}
 
-		TRigidScenePtr& operator=(const TRigidScenePtr& R) = default;
-		TRigidScenePtr& operator=(TRigidScenePtr&& R) = default;
-
-		FRigidScenePtrImpl SceneImpl;
-	};
-
+		// Impl in RigidBody.h
+		TRigidBodyPtr<FContext> CreateBody(const FRigidDebugName& InName, ERigidMovementType InMovementType);
+		void DestroyBody(TRigidBodyPtr<FContext>&& InBody);
 
-	class FRigidObjectId
-	{
-	public:
-		FRigidObjectId()
+		bool Query(const FRigidQueryDesc& QueryDesc, TArray<FRigidQueryHit>& OutHits)  const
 		{
-			IntId = -1;
+			return SceneRaw->Query(QueryDesc, OutHits);
 		}
 
-		explicit FRigidObjectId(uint64 InId)
+		void Tick(float Dt)
 		{
-			IntId = InId;
+			SceneRaw->Tick(Dt);
 		}
 
-		friend bool operator==(const FRigidObjectId& L, const FRigidObjectId& R) = default;
-		friend bool operator<(const FRigidObjectId& L, const FRigidObjectId& R)
+		UE_INTERNAL IRigidScene* Get() const
 		{
-			return L.IntId < R.IntId;
+			return SceneRaw;
 		}
 
 	private:
-		friend class Private::FRigidAccessor;
-
-		uint64 IntId;
+		IRigidScene* SceneRaw;
 	};
 
 
-	// A handle to a rigid body scene. FRigidSceneHandle is generally only used in the GameContext
+	// A handle to a rigid body scene. FRigidSceneHandle is only used in the GameContext
 	// because all PhysicsContect use cases (callbacks, modifiers, etc.) will be passed their
 	// scene as a TRigidScenePtr.
 	class FRigidSceneHandle
 	{
 	public:
-		FRigidSceneHandle()
-			: SceneRaw(nullptr)
-		{
-		}
-
-		FRigidSceneHandle(const FRigidSceneSharedPtr& InScene)
-			: SceneRaw(InScene.Get())
-#if UE_RIGIDPHYSICS_HANDLE_SAFTEY_ENABLED
-			, SceneWeak(InScene)
-#endif
-		{
-		}
+		RIGIDPHYSICS_API FRigidSceneHandle();
+		UE_INTERNAL RIGIDPHYSICS_API FRigidSceneHandle(FRigidScene* InScene);
 
-		void Reset()
-		{
-			SceneRaw = nullptr;
-
-#if UE_RIGIDPHYSICS_HANDLE_SAFTEY_ENABLED
-			SceneWeak.Reset();
-#endif
-		}
+		// Lock the scene on the game thread and return a locked context for use in pinning object handles
+		RIGIDPHYSICS_API FRigidContextGameRO LockRO() const;
+		RIGIDPHYSICS_API FRigidContextGameRW LockRW() const;
 
-		// Game-thread read-only scene lock
-		TRigidScenePtr<FGameContextRO> LockRO() const
-		{
-			return Lock<FGameContextRO>();
-		}
-
-		// Game-thread read-write scene lock
-		TRigidScenePtr<FGameContextRW> LockRW() const
-		{
-			return Lock<FGameContextRW>();
-		}
+		UE_INTERNAL RIGIDPHYSICS_API void Reset();
+		UE_INTERNAL RIGIDPHYSICS_API IRigidScene* Get() const;
 
 		friend bool operator==(const FRigidSceneHandle& L, const FRigidSceneHandle& R) = default;
 		friend bool operator!=(const FRigidSceneHandle& L, const FRigidSceneHandle& R) = default;
 
 	private:
-		friend class Private::FRigidAccessor;
-
+		// Impl in RigidContext.h
 		template<typename ContextType>
-		TRigidScenePtr<ContextType> Lock() const
-		{
-			// NOTE: SceneHandles are not used in PhysicsContext callbacks which pass
-			// TRigidScenePtr objects directly to the APIs. 
-			static_assert(ContextType::ContextEnv == EContextEnv::Game, "FRigidSceneHandle::Lock() is for the GameContext");
-
-#if UE_RIGIDPHYSICS_HANDLE_SAFTEY_ENABLED
-			// NOTE: In order to support the CHAOSRB_HANDLE_SAFTEY_ENABLED=0 optimization
-			// SceneHandles cannot be used after the scene has expired. The onus is on the 
-			// user to ensure that this is the case, but when we have handle safety
-			// enabled we can assert on this. Without the weak pointer, there's no way 
-			// to know that the raw pointer held in SceneWrapper points to a deleted object.
-			check(SceneRaw == SceneWeak.Pin().Get());
-#endif
-
-			return TRigidScenePtr<ContextType>(SceneRaw);
-		}
+		ContextType Lock() const;
 
-		IRigidScene* SceneRaw;
-#if UE_RIGIDPHYSICS_HANDLE_SAFTEY_ENABLED
+		// NOTE: Always the GT Scene, but never accessed on PT
+		FRigidScene* SceneRaw;
+#if UE_RIGIDPHYSICS_SCENEHANDLE_SAFTEY_ENABLED
 		FRigidSceneWeakPtr SceneWeak;
 #endif
 	};
 
+	template<typename ContextType>
+	FRigidSceneHandle TRigidScenePtrImpl<ContextType>::GetHandle() const
+	{
+		return SceneRaw->GetHandle();
+	}
 }
 
 #endif // UE_RIGIDPHYSICS_API_ENABLED
\ No newline at end of file
diff --git a/Engine/Source/Runtime/Experimental/RigidPhysics/Public/RigidPhysics/RigidTyped.h b/Engine/Source/Runtime/Experimental/RigidPhysics/Public/RigidPhysics/RigidTyped.h
new file mode 100644
--- /dev/null
+++ b/Engine/Source/Runtime/Experimental/RigidPhysics/Public/RigidPhysics/RigidTyped.h
@@ -0,0 +1,155 @@
+// Copyright Epic Games, Inc. All Rights Reserved.
+#pragma once
+
+#include "RigidPhysics/RigidFwd.h"
+
+#if UE_RIGIDPHYSICS_API_ENABLED
+
+#include "Misc/Crc.h"
+
+namespace UE::Physics
+{
+	// A type ID for a rigid object to support downcasting and type checking.
+	// Each type-checkable class will contain a static FRigidTypeId.
+	// NOTE: Only support single inheritance because that's all we need for now
+	// and hopefully we can keep it that way.
+	// TODO_CHAOSAPI: Consider using a more robust Id generator (to avoid CRC collisions).
+	// TODO_CHAOSAPI: Consider using USTRUCT rather than this custom RTTI system.
+	class FRigidTypeId
+	{
+	public:
+		FRigidTypeId(const TCHAR* InTypeName, const FRigidTypeId* InBaseTypeId)
+			: TypeIdValue(FCrc::StrCrc32(InTypeName))
+			, BaseTypeId(InBaseTypeId)
+			, TypeName(InTypeName)
+		{
+		}
+
+		const uint64 GetValue() const
+		{
+			return TypeIdValue;
+		}
+
+		const FString& GetTypeName() const
+		{
+			return TypeName;
+		}
+
+		const FRigidTypeId* GetBase() const
+		{
+			return BaseTypeId;
+		}
+
+		inline bool IsCastableTo(const FRigidTypeId& Other) const
+		{
+			const FRigidTypeId* Base = this;
+			while (Base != nullptr)
+			{
+				if (*Base == Other)
+				{
+					return true;
+				}
+				Base = Base->BaseTypeId;
+			}
+			return false;
+		}
+
+		friend inline bool operator==(const FRigidTypeId& L, const FRigidTypeId& R)
+		{
+			return L.TypeIdValue == R.TypeIdValue;
+		}
+
+		friend inline bool operator<(const FRigidTypeId& L, const FRigidTypeId& R)
+		{
+			return L.TypeIdValue < R.TypeIdValue;
+		}
+
+	private:
+		uint64 TypeIdValue;
+		const FRigidTypeId* BaseTypeId;
+		const FString TypeName;
+	};
+
+	class IRigidTyped
+	{
+	public:
+		RIGIDPHYSICS_API IRigidTyped();
+		RIGIDPHYSICS_API virtual ~IRigidTyped();
+
+		// Get the TypeId of this class
+		static const FRigidTypeId& GetStaticTypeId();
+
+		// Get the leaf TypeId of this object
+		virtual const FRigidTypeId& GetTypeId() const;
+
+		const FString& GetTypeName() const
+		{
+			return GetTypeId().GetTypeName();
+		}
+
+		// Is T the same class or a base class of this object?
+		template<typename T>
+		bool IsA() const
+		{
+			return GetTypeId().IsCastableTo(T::GetStaticTypeId());
+		}
+
+		// Try to cast this object to type T. Returns null if T is not a base class of this object.
+		template<typename T>
+		T* AsA()
+		{
+			if (IsA<T>())
+			{
+				return static_cast<T*>(this);
+			}
+			return nullptr;
+		}
+		
+		// Try to cast this object to type T. Returns null if T is not a base class of this object.
+		template<typename T>
+		const T* AsA() const
+		{
+			if (IsA<T>())
+			{
+				return static_cast<const T*>(this);
+			}
+			return nullptr;
+		}
+
+		// Try to cast this object to type T. Assert T is a base class of this object.
+		template<typename T>
+		T* AsAChecked()
+		{
+			check(IsA<T>());
+			return static_cast<T*>(this);
+		}
+
+		// Try to cast this object to type T. Assert T is a base class of this object.
+		template<typename T>
+		const T* AsAChecked() const
+		{
+			check(IsA<T>());
+			return static_cast<const T*>(this);
+		}
+	};
+
+}
+
+#define UE_RIGIDPHYSICS_RIGIDTYPED_DECL(API, T) \
+	API static const UE::Physics::FRigidTypeId& GetStaticTypeId(); \
+	API virtual const UE::Physics::FRigidTypeId& GetTypeId() const override
+
+#define UE_RIGIDPHYSICS_RIGIDTYPED_IMPL(T, TBASE) \
+	const UE::Physics::FRigidTypeId& T::GetTypeId() const \
+	{ \
+		return T::GetStaticTypeId(); \
+	} \
+	const UE::Physics::FRigidTypeId& T::GetStaticTypeId() \
+	{ \
+		using namespace UE::Physics; \
+		static const TCHAR* StaticTypeName = TEXT(#T); \
+		static FRigidTypeId StaticTypeId = FRigidTypeId(StaticTypeName, &TBASE::GetStaticTypeId()); \
+		return StaticTypeId; \
+	}
+
+#endif
diff --git a/Engine/Source/Runtime/Experimental/RigidPhysics/RigidPhysics.Build.cs b/Engine/Source/Runtime/Experimental/RigidPhysics/RigidPhysics.Build.cs
--- a/Engine/Source/Runtime/Experimental/RigidPhysics/RigidPhysics.Build.cs
+++ b/Engine/Source/Runtime/Experimental/RigidPhysics/RigidPhysics.Build.cs
@@ -13,9 +13,16 @@ namespace UnrealBuildTool.Rules
                 new string[] {
                 "Core",
                 "CoreUObject",
+				"Engine",
 				}
 			);
 
+			// Ideally RigidPhysics would not depend on Chaos Physics because this makes it harder
+			// to implement completely separate physics engines, but for now we need some Chaos types.
+			// TODO_CHAOSAPI: See if we can move everything we need into ChaosCore, or split Chaos
+			// into more modules to reduce this dependency.
+			SetupModulePhysicsSupport(Target);
+
 			CppCompileWarningSettings.UnsafeTypeCastWarningLevel = WarningLevel.Error;
 		}
     }
