Get StructUtils.h out of the PCH when using include order 5.7

#rnx
#Synced-CL 48776365

diff --git a/Engine/Plugins/ChaosVD/Source/ChaosVD/Private/Trace/ChaosVDTraceManager.cpp b/Engine/Plugins/ChaosVD/Source/ChaosVD/Private/Trace/ChaosVDTraceManager.cpp
--- a/Engine/Plugins/ChaosVD/Source/ChaosVD/Private/Trace/ChaosVDTraceManager.cpp
+++ b/Engine/Plugins/ChaosVD/Source/ChaosVD/Private/Trace/ChaosVDTraceManager.cpp
@@ -13,6 +13,7 @@
 #include "Misc/Paths.h"
 #include "Modules/ModuleManager.h"
 #include "RemoteSessionsManager.h"
+#include "SessionInfo.h"
 #include "Settings/ChaosVDGeneralSettings.h"
 #include "Trace/ChaosVDTraceModule.h"
 #include "Trace/StoreClient.h"
diff --git a/Engine/Plugins/ChaosVD/Source/ChaosVD/Private/Widgets/ChaosVDRecordingControls.cpp b/Engine/Plugins/ChaosVD/Source/ChaosVD/Private/Widgets/ChaosVDRecordingControls.cpp
--- a/Engine/Plugins/ChaosVD/Source/ChaosVD/Private/Widgets/ChaosVDRecordingControls.cpp
+++ b/Engine/Plugins/ChaosVD/Source/ChaosVD/Private/Widgets/ChaosVDRecordingControls.cpp
@@ -20,6 +20,7 @@
 #include "SChaosVDNameListPicker.h"
 #include "SEnumCombo.h"
 #include "Settings/ChaosVDGeneralSettings.h"
+#include "SessionInfo.h"
 #include "SocketSubsystem.h"
 #include "StatusBarSubsystem.h"
 #include "ToolMenu.h"
diff --git a/Engine/Plugins/ChaosVD/Source/ChaosVD/Private/Widgets/ChaosVDRecordingControls.h b/Engine/Plugins/ChaosVD/Source/ChaosVD/Private/Widgets/ChaosVDRecordingControls.h
--- a/Engine/Plugins/ChaosVD/Source/ChaosVD/Private/Widgets/ChaosVDRecordingControls.h
+++ b/Engine/Plugins/ChaosVD/Source/ChaosVD/Private/Widgets/ChaosVDRecordingControls.h
@@ -5,6 +5,7 @@
 #include "StatusBarSubsystem.h"
 #include "ChaosVDEngine.h"
 #include "RemoteSessionsManager.h"
+#include "SessionInfo.h"
 #include "Templates/SharedPointer.h"
 #include "Widgets/SCompoundWidget.h"
 
diff --git a/Engine/Source/Runtime/Experimental/ChaosSolverEngine/Private/Chaos/ChaosVDRemoteSessionsManager.cpp b/Engine/Source/Runtime/Experimental/ChaosSolverEngine/Private/Chaos/ChaosVDRemoteSessionsManager.cpp
--- a/Engine/Source/Runtime/Experimental/ChaosSolverEngine/Private/Chaos/ChaosVDRemoteSessionsManager.cpp
+++ b/Engine/Source/Runtime/Experimental/ChaosSolverEngine/Private/Chaos/ChaosVDRemoteSessionsManager.cpp
@@ -7,6 +7,7 @@
 #include "ChaosVisualDebugger/ChaosVDOptionalDataChannel.h"
 #include "ChaosVisualDebugger/ChaosVisualDebuggerTrace.h"
 #include "MessageEndpointBuilder.h"
+#include "SessionInfo.h"
 
 void FChaosVDRemoteSessionsHandler::OnCreatingMessageEndpoint(
 	const TSharedPtr<UE::TraceBasedDebuggers::FRemoteSessionsManager>& SessionsManager
diff --git a/Engine/Source/Runtime/Experimental/ChaosSolverEngine/Public/Chaos/ChaosVDRemoteSessionsManager.h b/Engine/Source/Runtime/Experimental/ChaosSolverEngine/Public/Chaos/ChaosVDRemoteSessionsManager.h
--- a/Engine/Source/Runtime/Experimental/ChaosSolverEngine/Public/Chaos/ChaosVDRemoteSessionsManager.h
+++ b/Engine/Source/Runtime/Experimental/ChaosSolverEngine/Public/Chaos/ChaosVDRemoteSessionsManager.h
@@ -13,6 +13,7 @@
 
 #include "ChaosVDRecordingDetails.h"
 #include "RemoteSessionsManager.h"
+#include "SessionInfo.h"
 #include "ChaosVDRemoteSessionsManager.generated.h"
 
 #define UE_API CHAOSSOLVERENGINE_API
diff --git a/Engine/Source/Runtime/Experimental/TraceBasedDebuggers/Private/RemoteSessionsManager.cpp b/Engine/Source/Runtime/Experimental/TraceBasedDebuggers/Private/RemoteSessionsManager.cpp
--- a/Engine/Source/Runtime/Experimental/TraceBasedDebuggers/Private/RemoteSessionsManager.cpp
+++ b/Engine/Source/Runtime/Experimental/TraceBasedDebuggers/Private/RemoteSessionsManager.cpp
@@ -8,6 +8,7 @@
 #include "MessageEndpoint.h"
 #include "MessageEndpointBuilder.h"
 #include "Misc/App.h"
+#include "SessionInfo.h"
 
 #include UE_INLINE_GENERATED_CPP_BY_NAME(RemoteSessionsManager)
 
diff --git a/Engine/Source/Runtime/Experimental/TraceBasedDebuggers/Public/RemoteSessionsManager.h b/Engine/Source/Runtime/Experimental/TraceBasedDebuggers/Public/RemoteSessionsManager.h
--- a/Engine/Source/Runtime/Experimental/TraceBasedDebuggers/Public/RemoteSessionsManager.h
+++ b/Engine/Source/Runtime/Experimental/TraceBasedDebuggers/Public/RemoteSessionsManager.h
@@ -9,7 +9,6 @@
 #include "Misc/Guid.h"
 #include "Templates/SharedPointer.h"
 #include "TraceDataRelayTransport.h"
-#include "StructUtils/InstancedStruct.h"
 
 #include "RemoteSessionsManager.generated.h"
 
@@ -22,9 +21,13 @@ UE_API DECLARE_LOG_CATEGORY_EXTERN(LogTraceBasedDebuggers, Log, All);
 
 namespace UE::TraceBasedDebuggers
 {
+struct FFullSessionInfoRequestMessage;
+struct FFullSessionInfoResponseMessage;
 struct FRemoteSessionsManager;
+struct FSessionInfo;
 struct FTraceConnectionDetails;
 
+
 USTRUCT(MinimalAPI)
 struct FSessionPing
 {
@@ -91,56 +94,6 @@ struct FTraceConnectionDetailsMessage
 	FTraceConnectionDetails TraceDetails;
 };
 
-USTRUCT(MinimalAPI)
-struct FFullSessionInfoRequestMessage
-{
-	GENERATED_BODY()
-};
-
-USTRUCT(MinimalAPI)
-struct FFullSessionInfoResponseMessage
-{
-	GENERATED_BODY()
-
-	template <typename DataType>
-	DataType* GetDebuggerData()
-	{
-		for (FInstancedStruct& InstancedDebuggerData : DebuggerSpecificData)
-		{
-			if (DataType* DebuggerData = InstancedDebuggerData.GetMutablePtr<DataType>())
-			{
-				return DebuggerData;
-			}
-		}
-		return nullptr;
-	}
-
-	template <typename DataType>
-	const DataType* GetDebuggerData() const
-	{
-		for (const FInstancedStruct& InstancedDebuggerData : DebuggerSpecificData)
-		{
-			if (const DataType* DebuggerData = InstancedDebuggerData.GetPtr<DataType>())
-			{
-				return DebuggerData;
-			}
-		}
-		return nullptr;
-	}
-
-	UPROPERTY()
-	FGuid InstanceId;
-
-	UPROPERTY()
-	FGuid DebuggerId = FGuid();
-
-	UPROPERTY()
-	TArray<FInstancedStruct> DebuggerSpecificData;
-
-	UPROPERTY()
-	bool bIsRecording = false;
-};
-
 USTRUCT(MinimalAPI)
 struct FStartRecordingCommandMessage
 {
@@ -166,141 +119,6 @@ enum class ERemoteSessionAttributes
 };
 ENUM_CLASS_FLAGS(ERemoteSessionAttributes)
 
-UENUM()
-enum class ERemoteSessionReadyState : uint8
-{
-	/** The session is ready to execute commands */
-	Ready,
-	/** We are executing a command in the session we expect to take a while without hearing anything from the target. */
-	Busy
-};
-
-/**
- * Session object that contains all the information needed to communicate with a remote instance, and the state of that instance
- */
-struct FSessionInfo
-{
-	explicit FSessionInfo() : InstanceId(FGuid()),
-		SessionTypeAttributes(ERemoteSessionAttributes::CanExpire | ERemoteSessionAttributes::SupportsDataChannelChange)
-	{
-	}
-
-	virtual ~FSessionInfo() = default;
-
-protected:
-
-	explicit FSessionInfo(ERemoteSessionAttributes InSessionTypeAttributes) : SessionTypeAttributes(InSessionTypeAttributes)
-	{
-	}
-
-public:
-
-	template <typename DataType>
-	DataType* GetDebuggerData()
-	{
-		for (FInstancedStruct& InstancedDebuggerData : DebuggerSpecificSessionData)
-		{
-			if (DataType* DebuggerData = InstancedDebuggerData.GetMutablePtr<DataType>())
-			{
-				return DebuggerData;
-			}
-		}
-		return nullptr;
-	}
-
-	/** Append or update data associated to a given debugger type */
-	template <typename DataType>
-	void SetDebuggerData(DataType&& NewSessionData)
-	{
-		if (DataType* ChaosData = GetDebuggerData<DataType>())
-		{
-			*ChaosData = MoveTemp(NewSessionData);
-		}
-		else
-		{
-			DebuggerSpecificSessionData.Emplace(FInstancedStruct::Make(MoveTemp(NewSessionData)));
-		}
-	}
-
-	FGuid InstanceId;
-	FString SessionName;
-	FMessageAddress Address;
-	FDateTime LastPingTime;
-	EBuildTargetType BuildTargetType = EBuildTargetType::Unknown;
-	ERemoteSessionReadyState ReadyState = ERemoteSessionReadyState::Ready;
-
-	FRecordingStatusMessage LastKnownRecordingState;
-
-	UE_API const FTraceConnectionDetails& GetTraceConnectionDetails();
-
-	/** @return Whether any debugger, regardless of its type, is currently recording */
-	UE_API virtual bool IsAnyDebuggerRecording() const;
-
-	/** 
-	 * Indicates whether a given debugger type is currently recording.
-	 * @param DebuggerTypeId The identifier of the debugger type
-	 * @return Whether the provided debugger type id is currently recording
-	 */
-	UE_API virtual bool IsRecording(const FGuid& DebuggerTypeId) const;
-
-	UE_API virtual uint64 GetBufferedBytesNum() const;
-
-	ERemoteSessionAttributes GetSessionTypeAttributes() const
-	{
-		return SessionTypeAttributes;
-	}
-
-	UE_API void SetReceivedBytesPerSecond(uint64 InNewBytesPerSecond);
-	UE_API uint64 GetReceivedBytesPerSecond() const;
-
-protected:
-
-	TArray<FInstancedStruct> DebuggerSpecificSessionData;
-
-	FTraceConnectionDetails LastKnownConnectionDetails;
-
-	uint64 ReceivedBytesPerSecond = 0;
-
-	const ERemoteSessionAttributes SessionTypeAttributes;
-
-	friend struct FRemoteSessionsManager;
-};
-
-/**
- * Session object that is able to control and provide information about multiple session objects.
- * Used so the UI can use the same API to control multiple sessions like it does for single sessions.
- */
-struct FMultiSessionInfo : FSessionInfo
-{
-	explicit FMultiSessionInfo() : FSessionInfo(ERemoteSessionAttributes::IsMultiSessionWrapper)
-	{
-	}
-
-	virtual ~FMultiSessionInfo() override = default;
-
-	UE_API virtual bool IsRecording(const FGuid& DebuggerTypeId) const override;
-	UE_API virtual bool IsAnyDebuggerRecording() const override;
-
-	UE_API virtual uint64 GetBufferedBytesNum() const override;
-
-	template<typename TCallback>
-	void EnumerateInnerSessions(const TCallback& Callback) const
-	{
-		for (const TPair<FGuid, TWeakPtr<FSessionInfo>>& InnerSessionWithID : InnerSessionsByInstanceID)
-		{
-			if (const TSharedPtr<FSessionInfo> SessionPtr = InnerSessionWithID.Value.Pin())
-			{
-				if (!Callback(SessionPtr.ToSharedRef()))
-				{
-					return;
-				}
-			}
-		}
-	}
-
-	TMap<FGuid, TWeakPtr<FSessionInfo>> InnerSessionsByInstanceID;
-};
-
 DECLARE_MULTICAST_DELEGATE_OneParam(FRecordingStateChangeDelegate, const TSharedPtr<FSessionInfo>&);
 DECLARE_MULTICAST_DELEGATE_OneParam(FSessionLifetimeDelegate, const TSharedPtr<FSessionInfo>&);
 
diff --git a/Engine/Source/Runtime/Experimental/TraceBasedDebuggers/Public/SessionInfo.h b/Engine/Source/Runtime/Experimental/TraceBasedDebuggers/Public/SessionInfo.h
new file mode 100644
--- /dev/null
+++ b/Engine/Source/Runtime/Experimental/TraceBasedDebuggers/Public/SessionInfo.h
@@ -0,0 +1,201 @@
+// Copyright Epic Games, Inc. All Rights Reserved.
+
+#pragma once
+
+#include "RemoteSessionsManager.h"
+#include "StructUtils/InstancedStruct.h"
+
+#include "SessionInfo.generated.h"
+
+#define UE_API TRACEBASEDDEBUGGERS_API
+
+namespace UE::TraceBasedDebuggers
+{
+
+USTRUCT(MinimalAPI)
+struct FFullSessionInfoRequestMessage
+{
+	GENERATED_BODY()
+};
+
+USTRUCT(MinimalAPI)
+struct FFullSessionInfoResponseMessage
+{
+	GENERATED_BODY()
+
+	template <typename DataType>
+	DataType* GetDebuggerData()
+	{
+		for (FInstancedStruct& InstancedDebuggerData : DebuggerSpecificData)
+		{
+			if (DataType* DebuggerData = InstancedDebuggerData.GetMutablePtr<DataType>())
+			{
+				return DebuggerData;
+			}
+		}
+		return nullptr;
+	}
+
+	template <typename DataType>
+	const DataType* GetDebuggerData() const
+	{
+		for (const FInstancedStruct& InstancedDebuggerData : DebuggerSpecificData)
+		{
+			if (const DataType* DebuggerData = InstancedDebuggerData.GetPtr<DataType>())
+			{
+				return DebuggerData;
+			}
+		}
+		return nullptr;
+	}
+
+	UPROPERTY()
+	FGuid InstanceId;
+
+	UPROPERTY()
+	FGuid DebuggerId = FGuid();
+
+	UPROPERTY()
+	TArray<FInstancedStruct> DebuggerSpecificData;
+
+	UPROPERTY()
+	bool bIsRecording = false;
+};
+
+UENUM()
+enum class ERemoteSessionReadyState : uint8
+{
+	/** The session is ready to execute commands */
+	Ready,
+	/** We are executing a command in the session we expect to take a while without hearing anything from the target. */
+	Busy
+};
+
+/**
+* Session object that contains all the information needed to communicate with a remote instance, and the state of that instance
+*/
+struct FSessionInfo
+{
+	explicit FSessionInfo() : InstanceId(FGuid()),
+	SessionTypeAttributes(ERemoteSessionAttributes::CanExpire | ERemoteSessionAttributes::SupportsDataChannelChange)
+	{
+	}
+
+	virtual ~FSessionInfo() = default;
+
+protected:
+
+	explicit FSessionInfo(ERemoteSessionAttributes InSessionTypeAttributes) : SessionTypeAttributes(InSessionTypeAttributes)
+	{
+	}
+
+public:
+
+	template <typename DataType>
+	DataType* GetDebuggerData()
+	{
+		for (FInstancedStruct& InstancedDebuggerData : DebuggerSpecificSessionData)
+		{
+			if (DataType* DebuggerData = InstancedDebuggerData.GetMutablePtr<DataType>())
+			{
+				return DebuggerData;
+			}
+		}
+		return nullptr;
+	}
+
+	/** Append or update data associated to a given debugger type */
+	template <typename DataType>
+	void SetDebuggerData(DataType&& NewSessionData)
+	{
+		if (DataType* ChaosData = GetDebuggerData<DataType>())
+		{
+			*ChaosData = MoveTemp(NewSessionData);
+		}
+		else
+		{
+			DebuggerSpecificSessionData.Emplace(FInstancedStruct::Make(MoveTemp(NewSessionData)));
+		}
+	}
+
+	FGuid InstanceId;
+	FString SessionName;
+	FMessageAddress Address;
+	FDateTime LastPingTime;
+	EBuildTargetType BuildTargetType = EBuildTargetType::Unknown;
+	ERemoteSessionReadyState ReadyState = ERemoteSessionReadyState::Ready;
+
+	FRecordingStatusMessage LastKnownRecordingState;
+
+	UE_API const FTraceConnectionDetails& GetTraceConnectionDetails();
+
+	/** @return Whether any debugger, regardless of its type, is currently recording */
+	UE_API virtual bool IsAnyDebuggerRecording() const;
+
+	/** 
+		* Indicates whether a given debugger type is currently recording.
+		* @param DebuggerTypeId The identifier of the debugger type
+		* @return Whether the provided debugger type id is currently recording
+		*/
+	UE_API virtual bool IsRecording(const FGuid& DebuggerTypeId) const;
+
+	UE_API virtual uint64 GetBufferedBytesNum() const;
+
+	ERemoteSessionAttributes GetSessionTypeAttributes() const
+	{
+		return SessionTypeAttributes;
+	}
+
+	UE_API void SetReceivedBytesPerSecond(uint64 InNewBytesPerSecond);
+	UE_API uint64 GetReceivedBytesPerSecond() const;
+
+protected:
+
+	TArray<FInstancedStruct> DebuggerSpecificSessionData;
+
+	FTraceConnectionDetails LastKnownConnectionDetails;
+
+	uint64 ReceivedBytesPerSecond = 0;
+
+	const ERemoteSessionAttributes SessionTypeAttributes;
+
+	friend struct FRemoteSessionsManager;
+};
+
+/**
+	* Session object that is able to control and provide information about multiple session objects.
+	* Used so the UI can use the same API to control multiple sessions like it does for single sessions.
+	*/
+struct FMultiSessionInfo : FSessionInfo
+{
+	explicit FMultiSessionInfo() : FSessionInfo(ERemoteSessionAttributes::IsMultiSessionWrapper)
+	{
+	}
+
+	virtual ~FMultiSessionInfo() override = default;
+
+	UE_API virtual bool IsRecording(const FGuid& DebuggerTypeId) const override;
+	UE_API virtual bool IsAnyDebuggerRecording() const override;
+
+	UE_API virtual uint64 GetBufferedBytesNum() const override;
+
+	template<typename TCallback>
+	void EnumerateInnerSessions(const TCallback& Callback) const
+	{
+		for (const TPair<FGuid, TWeakPtr<FSessionInfo>>& InnerSessionWithID : InnerSessionsByInstanceID)
+		{
+			if (const TSharedPtr<FSessionInfo> SessionPtr = InnerSessionWithID.Value.Pin())
+			{
+				if (!Callback(SessionPtr.ToSharedRef()))
+				{
+					return;
+				}
+			}
+		}
+	}
+
+	TMap<FGuid, TWeakPtr<FSessionInfo>> InnerSessionsByInstanceID;
+};
+
+} // namespace UE::TraceBasedDebuggers
+#undef UE_API
