Better align the periodic logging of IAX to make it easier to extract.

#rnx

#Synced-CL 48686588

diff --git a/Engine/Source/Runtime/Experimental/IoStore/OnDemand/Private/Statistics.cpp b/Engine/Source/Runtime/Experimental/IoStore/OnDemand/Private/Statistics.cpp
--- a/Engine/Source/Runtime/Experimental/IoStore/OnDemand/Private/Statistics.cpp
+++ b/Engine/Source/Runtime/Experimental/IoStore/OnDemand/Private/Statistics.cpp
@@ -686,7 +686,7 @@ void FOnDemandIoBackendStats::PrintPeriodicLogging()
 
 		if (BackendStatus.IsCacheEnabled())
 		{
-			UE_LOG(LogIas, Log, TEXT("CacheStats: Count=%d, CachedKiB=%d, WrittenKiB=%d, ReadKiB=%d, RejectedKiB=%d, Get=%d, Error=%d, Put=%d, PutReject=%d, PutExisting=%d"),
+			UE_LOG(LogIas, Log, TEXT("IASCache: Count=%d, CachedKiB=%d, WrittenKiB=%d, ReadKiB=%d, RejectedKiB=%d, Get=%d, Error=%d, Put=%d, PutReject=%d, PutExisting=%d"),
 				(int32)GCacheCachedEntries.Get(), (int32)CCachedKiB, (int32)CWrittenKiB, (int32)CReadKiB, (int32)CRejectedKiB, CGetCount, CErrorCount, CPutCount, CPutRejectCount, CPutExistingCount);
 		}
 		else
@@ -696,7 +696,7 @@ void FOnDemandIoBackendStats::PrintPeriodicLogging()
 
 		auto HttpLog = [](const TCHAR* Title, EHttpRequestType Type)
 			{
-				UE_LOG(LogIas, Log, TEXT("%s - HttpStats: DownloadedKiB=%d, Get=%d, Retry=%d, Cancel=%d, Error=%d, CurPending=%d, CurDurationMsAvg=%d, CurDurationMsMax=%d"),
+				UE_LOG(LogIas, Log, TEXT("%sHttp : DownloadedKiB=%d, Get=%d, Retry=%d, Cancel=%d, Error=%d, CurPending=%d, CurDurationMsAvg=%d, CurDurationMsMax=%d"),
 					Title,
 					(int32)((float)GHttpDownloadedBytes[(uint8)Type].Get() * OneOver1024),
 					(int32)GHttpGetCount[(uint8)Type].Get(),
