JsonObjectGraph: Add DisableNativeSerializers option to disable running the native serializers.
It reduces size and improves performance in scenarios where the native data is less important.

#rnx

#Synced-CL 48718024

diff --git a/Engine/Source/Runtime/Experimental/JsonObjectGraph/Private/JsonStringifyImpl.cpp b/Engine/Source/Runtime/Experimental/JsonObjectGraph/Private/JsonStringifyImpl.cpp
--- a/Engine/Source/Runtime/Experimental/JsonObjectGraph/Private/JsonStringifyImpl.cpp
+++ b/Engine/Source/Runtime/Experimental/JsonObjectGraph/Private/JsonStringifyImpl.cpp
@@ -415,26 +415,31 @@ void FJsonStringifyImpl::WriteNativeObjectData()
 	Writer->WriteValue(UE_JSON_OBJECT_FLAGS_KEY_TCHAR, Object->GetFlags() & RF_Load);
 	Writer->WriteObjectEnd();
 
-	// This is tricky, we have no good mechanism for detecting whether an object wants a structured
-	// serialization or a traditional stream serialization. So I've decided to 'try' the structured
-	// serializer and if it writes nothing we fall back to the stream (FArchive) serializer. Hopefully
-	// some day we get rid of FStructuredArchive as it does not spark joy.
+	// Native serializers may produce huge amounts of data that add a size and performance overhead.
+	// They can be optionally disabled in scenarios where this data is less important.
+	if (!(WriteOptions.Flags & EJsonStringifyFlags::DisableNativeSerializers))
+	{
+		// This is tricky, we have no good mechanism for detecting whether an object wants a structured
+		// serialization or a traditional stream serialization. So I've decided to 'try' the structured
+		// serializer and if it writes nothing we fall back to the stream (FArchive) serializer. Hopefully
+		// some day we get rid of FStructuredArchive as it does not spark joy.
 #if WITH_TEXT_ARCHIVE_SUPPORT
-	TArray<uint8> StructuredData = StructuredDataToJson(Object, Writer->GetIndentLevel());
-	if (!StructuredData.IsEmpty())
-	{
-		Writer->WriteIdentifierPrefix(UE_JSON_OBJECT_STRUCTURED_DATA_KEY_TCHAR);
-		Writer->WriteJsonRaw(FAnsiStringView((ANSICHAR*)StructuredData.GetData(), StructuredData.Num()));
-	}
-	else
+		TArray<uint8> StructuredData = StructuredDataToJson(Object, Writer->GetIndentLevel());
+		if (!StructuredData.IsEmpty())
+		{
+			Writer->WriteIdentifierPrefix(UE_JSON_OBJECT_STRUCTURED_DATA_KEY_TCHAR);
+			Writer->WriteJsonRaw(FAnsiStringView((ANSICHAR*)StructuredData.GetData(), StructuredData.Num()));
+		}
+		else
 #endif
-	{
-		// no useful native structured data, write the native serial data:
-		TArray<uint8> SerialData = SerialDataToJson(Object, Writer->GetIndentLevel());
-		if (!SerialData.IsEmpty())
 		{
-			Writer->WriteIdentifierPrefix(UE_JSON_OBJECT_SERIAL_DATA_KEY_TCHAR);
-			Writer->WriteJsonRaw(FAnsiStringView((ANSICHAR*)SerialData.GetData(), SerialData.Num()));
+			// no useful native structured data, write the native serial data:
+			TArray<uint8> SerialData = SerialDataToJson(Object, Writer->GetIndentLevel());
+			if (!SerialData.IsEmpty())
+			{
+				Writer->WriteIdentifierPrefix(UE_JSON_OBJECT_SERIAL_DATA_KEY_TCHAR);
+				Writer->WriteJsonRaw(FAnsiStringView((ANSICHAR*)SerialData.GetData(), SerialData.Num()));
+			}
 		}
 	}
 
diff --git a/Engine/Source/Runtime/Experimental/JsonObjectGraph/Public/JsonObjectGraph/Stringify.h b/Engine/Source/Runtime/Experimental/JsonObjectGraph/Public/JsonObjectGraph/Stringify.h
--- a/Engine/Source/Runtime/Experimental/JsonObjectGraph/Public/JsonObjectGraph/Stringify.h
+++ b/Engine/Source/Runtime/Experimental/JsonObjectGraph/Public/JsonObjectGraph/Stringify.h
@@ -18,6 +18,8 @@ enum class EJsonStringifyFlags : uint8
 	FilterEditorOnlyData = 1 << 0,
 	/** Disables delta encoding such that all properties are encoded in the Json, rather than only changes from the objects' archetypes */
 	DisableDeltaEncoding = 1 << 1,
+	/** Disables running the native serializers, reduces size and improves performance when the native data is less important */
+	DisableNativeSerializers = 1 << 2,
 };
 ENUM_CLASS_FLAGS(EJsonStringifyFlags);
 
