Complete on-demand HTTP client requests for responses with no content.

#rnx

#Synced-CL 48708293

diff --git a/Engine/Source/Runtime/Experimental/IoStore/OnDemand/Private/OnDemandHttpClient.cpp b/Engine/Source/Runtime/Experimental/IoStore/OnDemand/Private/OnDemandHttpClient.cpp
--- a/Engine/Source/Runtime/Experimental/IoStore/OnDemand/Private/OnDemandHttpClient.cpp
+++ b/Engine/Source/Runtime/Experimental/IoStore/OnDemand/Private/OnDemandHttpClient.cpp
@@ -197,23 +197,27 @@ TIoStatusOr<FMultiEndpointHttpClientResponse> FMultiEndpointHttpClient::Get(FAns
 	FEventLoop Loop;
 	FIoBuffer Body;
 	TStringBuilder<128> Reason;
+	HTTP::FTicketPerf::FSample PerfSample;
 	uint32 StatusCode = 0;
 
 	const uint32 MaxAttempts = Config.MaxRetryCount == -1 ? 3u : static_cast<uint32>(Config.MaxRetryCount);
-
-	const uint64 StartTime = FPlatformTime::Cycles64();
-
 	for (uint32 Attempt = 0; Attempt <= MaxAttempts; ++Attempt)
 	{
-		Loop.Send(Loop.Request("GET", Url, &Params), [&Body, &Reason, &StatusCode](const FTicketStatus& Status)
+		Loop.Send(Loop.Request("GET", Url, &Params), [&Body, &Reason, &PerfSample, &StatusCode](const FTicketStatus& Status)
 			{
 				if (Status.GetId() == FTicketStatus::EId::Response)
 				{
 					Status.GetResponse().SetDestination(&Body);
 					StatusCode = Status.GetResponse().GetStatusCode();
+					PerfSample = Status.GetPerf().GetSample();
 					return;
 				}
 
+				if (Status.GetId() == FTicketStatus::EId::Content)
+				{
+					PerfSample = Status.GetPerf().GetSample();
+				}
+
 				if (Status.GetId() == FTicketStatus::EId::Error)
 				{
 					Reason << Status.GetError().Reason;
@@ -230,7 +234,7 @@ TIoStatusOr<FMultiEndpointHttpClientResponse> FMultiEndpointHttpClient::Get(FAns
 			FMultiEndpointHttpClientResponse Response
 			{
 				.Body = MoveTemp(Body),
-				.DurationMilliseconds = uint64(FPlatformTime::ToMilliseconds64(FPlatformTime::Cycles64() - StartTime)),
+				.Sample = PerfSample,
 				.StatusCode = StatusCode,
 				.RetryCount = Attempt,
 			};
@@ -430,14 +434,17 @@ FMultiEndpointHttpClient::FHttpTicketId FMultiEndpointHttpClient::IssueRequest(F
 
 	auto HttpSink = [this, Request = MoveTemp(Request)](const FTicketStatus& TicketStatus) mutable
 	{
+		bool bCompleted	= false;
+		bool bCanceled	= false;
+		bool bError		= false;
+
 		switch (TicketStatus.GetId())
 		{
 			case FTicketStatus::EId::Response:
 			{
-				FResponse& HttpResponse = TicketStatus.GetResponse();
-
-				Request.StatusCode = HttpResponse.GetStatusCode();
-				Request.bIsChunkedTransfer = HttpResponse.GetContentLength() == -1;
+				FResponse& HttpResponse		= TicketStatus.GetResponse();
+				Request.StatusCode			= HttpResponse.GetStatusCode();
+				Request.bIsChunkedTransfer	= HttpResponse.GetContentLength() == -1;
 
 				if (EnumHasAnyFlags(Request.Flags, EMultiEndpointRequestFlags::ResponseHeaders))
 				{
@@ -460,46 +467,75 @@ FMultiEndpointHttpClient::FHttpTicketId FMultiEndpointHttpClient::IssueRequest(F
 					Request.ResponseMessage = FString(HttpResponse.GetStatusMessage().TrimEnd());
 				}
 
-				HttpResponse.SetDestination(Request.bIsChunkedTransfer == false ? &Request.Body : &Request.Chunk);
-
+				// Returns true if the response has no content, i.e. 204, 205 and 304
+				bCompleted = TicketStatus.IsComplete();
+				if (!bCompleted)
+				{
+					HttpResponse.SetDestination(Request.bIsChunkedTransfer == false ? &Request.Body : &Request.Chunk);
+				}
 				break;
 			}
 			case FTicketStatus::EId::Content:
 			{
-				// ::DataSize of zero means that all the chunks have been transfered and FRequest::Body should be complete
-				if (Request.bIsChunkedTransfer && Request.Chunk.DataSize() != 0)
+				bCompleted = true;
+				if (Request.bIsChunkedTransfer)
 				{
-					if (IsHttpStatusOk(Request.StatusCode) || Config.bResponseBodyOnError)
+					// If the current chunk size is zero that means all the chunks have been transfered and FRequest::Body should be complete
+					bCompleted = Request.Chunk.GetSize() == 0;
+					if (!bCompleted)
 					{
-						// Could consider using FRequest::Range to presize FRequest::Body and copy the chunks
-						// into it rather than resizing each time.
-						Request.Body = JoinIoBuffers(Request.Body, Request.Chunk);
+						if (IsHttpStatusOk(Request.StatusCode) || Config.bResponseBodyOnError)
+						{
+							// Could consider using FRequest::Range to presize FRequest::Body and copy the chunks
+							// into it rather than resizing each time.
+							Request.Body = JoinIoBuffers(Request.Body, Request.Chunk);
+						}
 					}
-					return;
 				}
+				break;
 			}
-			case FTicketStatus::EId::Error:
 			case FTicketStatus::EId::Cancelled:
 			{
-#if DO_CHECK
-				check(Request.SinkCounter++ == Request.RetryCount);
-#endif //DO_CHECK
+				bCompleted = bCanceled = true;
+				break;
+			}
+			case FTicketStatus::EId::Error:
+			{
+				bCompleted = bError = true;
+				break;
+			}
+			default:
+			{
+				checkNoReentry();
+				bCompleted = bCanceled = true;
+				UE_LOG(LogHttpIoDispatcher, Error, TEXT("Unexpected HTTP ticket status"));
+				break;
+			}
+		}
 
-				Log(Request, TicketStatus);
+		// The callback is triggered multiple times, i.e. first the response and then the content
+		if (!bCompleted)
+		{
+			return;
+		}
 
-				// We only want to retry the request if there was an internal or server problem. There is no point retrying a 404 error for example/
-				const bool bError = IsHttpServerError(Request.StatusCode) || TicketStatus.GetId() == FTicketStatus::EId::Error;
+		// Only retry if there was a connection error or a server error
+		bool bRetry = false;
+		if (!bCanceled && (bError || IsHttpServerError(Request.StatusCode)))
+		{
+			bRetry = Request.RetryCount < GetRetryLimitForRequest(Request);
+		}
 
-				if (bError && Request.RetryCount < GetRetryLimitForRequest(Request))
-				{
-					RetryRequest(MoveTemp(Request));
-				}
-				else
-				{
-					CompleteRequest(MoveTemp(Request), TicketStatus);
-				}
-				break;
-			}
+		if (bRetry)
+		{
+#if DO_CHECK
+			Request.SinkCounter++;
+#endif
+			RetryRequest(MoveTemp(Request));
+		}
+		else
+		{
+			CompleteRequest(MoveTemp(Request), TicketStatus);
 		}
 	};
 
@@ -518,10 +554,14 @@ void FMultiEndpointHttpClient::CompleteRequest(FRequest&& Request, const UE::IoS
 {
 	using namespace UE::IoStore::HTTP;
 
+#if DO_CHECK
+	check(Request.SinkCounter == Request.RetryCount);
+#endif
+	Log(Request, TicketStatus);
+
 	FMultiEndpointHttpClientResponse Response
 	{
 		.Headers				= MoveTemp(Request.ResponseHeaders),
-		.DurationMilliseconds	= uint64(FPlatformTime::ToMilliseconds64(FPlatformTime::Cycles64() - Request.StartTime)),
 		.StatusCode				= Request.StatusCode,
 		.RetryCount				= Request.RetryCount,
 		.HostIndex				= Request.Host,
@@ -655,7 +695,6 @@ void FMultiEndpointHttpClient::ProcessRetryAttempts(uint32 TicketCount)
 			FMultiEndpointHttpClientResponse Response
 			{
 				.Reason = TEXT("Canceled"),
-				.DurationMilliseconds = uint64(FPlatformTime::ToMilliseconds64(FPlatformTime::Cycles64() - Request.StartTime)),
 				.StatusCode = Request.StatusCode,
 				.RetryCount = Request.RetryCount - 1, // Reduce retry count by 1 as this retry attempt didn't really do anything
 				.HostIndex = Request.Host,
diff --git a/Engine/Source/Runtime/Experimental/IoStore/OnDemand/Private/OnDemandHttpClient.h b/Engine/Source/Runtime/Experimental/IoStore/OnDemand/Private/OnDemandHttpClient.h
--- a/Engine/Source/Runtime/Experimental/IoStore/OnDemand/Private/OnDemandHttpClient.h
+++ b/Engine/Source/Runtime/Experimental/IoStore/OnDemand/Private/OnDemandHttpClient.h
@@ -93,7 +93,6 @@ struct FMultiEndpointHttpClientResponse
 	TArray<FAnsiString>			Headers;
 	FString						Reason;
 	HTTP::FTicketPerf::FSample	Sample;
-	uint64						DurationMilliseconds = 0;
 	uint32						StatusCode = 0;
 	uint32						RetryCount = 0;
 	int32						HostIndex = INDEX_NONE;
