Dataflow graph class cleanup
- Moved many method bodies in the cpp file
- Deprecated the template version of AddNode as they were not needed and in favor of a non-template SharedPtr more general version
- better comments all around
- Updated the Node factory to properly use the right version of AddNode

#Synced-CL 48774242

diff --git a/Engine/Source/Runtime/Experimental/Dataflow/Core/Private/Dataflow/DataflowGraph.cpp b/Engine/Source/Runtime/Experimental/Dataflow/Core/Private/Dataflow/DataflowGraph.cpp
--- a/Engine/Source/Runtime/Experimental/Dataflow/Core/Private/Dataflow/DataflowGraph.cpp
+++ b/Engine/Source/Runtime/Experimental/Dataflow/Core/Private/Dataflow/DataflowGraph.cpp
@@ -32,6 +32,92 @@ namespace UE::Dataflow
 		DisabledNodes.Reset();
 	}
 
+	const TArray<FGraph::FNodeSharedPtr>& FGraph::GetFilteredNodes(const FName& NodeFilter) const
+	{
+		static TArray<FNodeSharedPtr> EmptyNodes;
+		if (const TArray<FNodeSharedPtr>* FoundNodes = FilteredNodes.Find(NodeFilter))
+		{
+			return *FoundNodes;
+		}
+		return EmptyNodes;
+	}
+
+	FGraph::FNodeSharedPtr FGraph::AddNode(FNodeSharedPtr NewNode)
+	{
+		if (NewNode)
+		{
+			Nodes.AddUnique(NewNode);
+			for (const FName& RegisteredType : RegisteredFilters)
+			{
+				if (NewNode->IsA(RegisteredType))
+				{
+					FilteredNodes.FindOrAdd(RegisteredType).Add(NewNode);
+				}
+			}
+		}
+		return NewNode;
+	}
+
+	FGraph::FNodeSharedPtr FGraph::FindBaseNode(FGuid InGuid)
+	{
+		for (const FNodeSharedPtr& Node : Nodes)
+		{
+			if (Node && Node->GetGuid() == InGuid)
+			{
+				return Node;
+			}
+		}
+		return {};
+	}
+
+	FGraph::FConstNodeSharedPtr FGraph::FindBaseNode(FGuid InGuid) const
+	{
+		for (const FNodeSharedPtr& Node : Nodes)
+		{
+			if (Node && Node->GetGuid() == InGuid)
+			{
+				return Node;
+			}
+		}
+		return {};
+	}
+
+	FGraph::FNodeSharedPtr FGraph::FindBaseNode(FName InName)
+	{
+		for (const FNodeSharedPtr& Node : Nodes)
+		{
+			if (Node && Node->GetName().IsEqual(InName))
+			{
+				return Node;
+			}
+		}
+		return {};
+	}
+
+	FGraph::FConstNodeSharedPtr FGraph::FindBaseNode(FName InName) const
+	{
+		for (const FNodeSharedPtr& Node : Nodes)
+		{
+			if (Node && Node->GetName().IsEqual(InName))
+			{
+				return Node;
+			}
+		}
+		return {};
+	}
+
+	FGraph::FNodeSharedPtr FGraph::FindFilteredNode(const FName& NodeFilter, FName InName) const
+	{
+		for (const FNodeSharedPtr& Node : GetFilteredNodes(NodeFilter))
+		{
+			if (Node && Node->GetName().IsEqual(InName))
+			{
+				return Node;
+			}
+		}
+		return {};
+	}
+
 	void FGraph::RemoveNode(TSharedPtr<FDataflowNode> Node)
 	{
 		for (FDataflowOutput* Output : Node->GetOutputs())
@@ -66,7 +152,7 @@ namespace UE::Dataflow
 		{
 			if (Node->IsA(RegisteredType))
 			{
-				if(TArray< TSharedPtr<FDataflowNode> >* FoundNodes = FilteredNodes.Find(RegisteredType))
+				if(TArray<FNodeSharedPtr>* FoundNodes = FilteredNodes.Find(RegisteredType))
 				{
 					FoundNodes->Remove(Node);
 				}
diff --git a/Engine/Source/Runtime/Experimental/Dataflow/Core/Private/Dataflow/DataflowNodeFactory.cpp b/Engine/Source/Runtime/Experimental/Dataflow/Core/Private/Dataflow/DataflowNodeFactory.cpp
--- a/Engine/Source/Runtime/Experimental/Dataflow/Core/Private/Dataflow/DataflowNodeFactory.cpp
+++ b/Engine/Source/Runtime/Experimental/Dataflow/Core/Private/Dataflow/DataflowNodeFactory.cpp
@@ -88,16 +88,16 @@ namespace UE::Dataflow
 	{ 
 		if (ClassMap.Contains(Param.Type))
 		{
-			TUniquePtr<FDataflowNode> Node = ClassMap[Param.Type](Param);
-			if(Node->HasValidConnections())
+			TUniquePtr<FDataflowNode> NewNode = ClassMap[Param.Type](Param);
+			if(NewNode->HasValidConnections())
 			{
-				ParametersMap[Param.Type].ToolTip = Node->GetToolTip();
+				ParametersMap[Param.Type].ToolTip = NewNode->GetToolTip();
 
-				return Graph.AddNode(MoveTemp(Node));
+				return Graph.AddNode(TSharedPtr<FDataflowNode>(NewNode.Release()));
 			}
 			
 			const FText ErrorTitle = FText::FromString("Node Factory");
-			const FString ErrorMessageString = FString::Printf(TEXT("Cannot create Node %s. Node Type %s is not well defined."), *Node->GetName().ToString(), *Node->GetDisplayName().ToString());
+			const FString ErrorMessageString = FString::Printf(TEXT("Cannot create Node %s. Node Type %s is not well defined."), *NewNode->GetName().ToString(), *NewNode->GetDisplayName().ToString());
 			const FText ErrorMessage = FText::FromString(ErrorMessageString);
 			FMessageDialog::Debugf(ErrorMessage, ErrorTitle);
 		}
diff --git a/Engine/Source/Runtime/Experimental/Dataflow/Core/Public/Dataflow/DataflowGraph.h b/Engine/Source/Runtime/Experimental/Dataflow/Core/Public/Dataflow/DataflowGraph.h
--- a/Engine/Source/Runtime/Experimental/Dataflow/Core/Public/Dataflow/DataflowGraph.h
+++ b/Engine/Source/Runtime/Experimental/Dataflow/Core/Public/Dataflow/DataflowGraph.h
@@ -11,6 +11,8 @@
 
 #include "DataflowGraph.generated.h"
 
+#define UE_API DATAFLOWCORE_API
+
 struct FDataflowConnection;
 
 namespace UE::Dataflow
@@ -35,13 +37,18 @@ public:
 
 namespace UE::Dataflow
 {
-	struct FLink {
+	/**
+	* A link structure defines a connection between one output and one input of two separate Dataflow nodes
+	* The Guids are the unique identifiers of the node and inputs/outputs in the dataflow graph (see FGraph class below)
+	*/
+	struct FLink
+	{
 		FGuid InputNode;
 		FGuid Input;
 		FGuid OutputNode;
 		FGuid Output;
 
-		FLink() {}
+		FLink() = default;
 
 		FLink(FGuid InOutputNode, FGuid InOutput, FGuid InInputNode, FGuid InInput)
 			: InputNode(InInputNode), Input(InInput)
@@ -64,17 +71,22 @@ namespace UE::Dataflow
 	};
 
 
-	//
-	//
-	//
+	/**
+	* A Dataflow graph is a collection of interconnected nodes
+	* Each node has aunique Guid and name and a list of inputs and outputs ( see FDataflowNode )
+	*/
 	class FGraph
 	{
+	public:
+		using FNodeSharedPtr = TSharedPtr<FDataflowNode>;
+		using FConstNodeSharedPtr = TSharedPtr<const FDataflowNode>;
 
+	private:
 		FGuid  Guid;
-		TArray< TSharedPtr<FDataflowNode> > Nodes;
-		TMap<FName, TArray<TSharedPtr<FDataflowNode>>> FilteredNodes;
-		TArray< FLink > Connections;
-		TSet< FName > DisabledNodes;
+		TArray<FNodeSharedPtr> Nodes;
+		TMap<FName, TArray<FNodeSharedPtr>> FilteredNodes;
+		TArray<FLink> Connections;
+		TSet<FName> DisabledNodes;
 		
 		/** Node filter type that could be used for fast access*/
 		static TSet<FName> RegisteredFilters;
@@ -83,116 +95,79 @@ namespace UE::Dataflow
 		friend DATAFLOWCORE_API void RegisterNodeFilter(const FName& NodeFilter);
 		
 	public:
-		DATAFLOWCORE_API FGraph(FGuid InGuid = FGuid::NewGuid());
+		UE_API FGraph(FGuid InGuid = FGuid::NewGuid());
 		virtual ~FGraph() = default;
 		
-		const TArray< TSharedPtr<FDataflowNode> >& GetFilteredNodes(const FName& NodeFilter) const
-		{
-			static TArray< TSharedPtr<FDataflowNode> > EmptyNodes;
-			if(const TArray< TSharedPtr<FDataflowNode> >* FoundNodes = FilteredNodes.Find(NodeFilter))
-			{
-				return *FoundNodes;
-			}
-			return EmptyNodes;
-		}
-		const TArray< TSharedPtr<FDataflowNode> >& GetNodes() const { return Nodes; }
-		TArray< TSharedPtr<FDataflowNode> >& GetNodes() { return Nodes; }
-		int NumNodes() { return Nodes.Num(); }
-		
-		template<class T> TSharedPtr<T> AddNode(T* InNode)
+		/** Get the list of node matching a specific named filter (@see RegisterNodeFilter) */
+		UE_API const TArray<FNodeSharedPtr>& GetFilteredNodes(const FName& NodeFilter) const;
+
+		/** 
+		* Get all the nodes of the graph ( const method )
+		* this includes the node in all subgraphs
+		*/
+		const TArray<FNodeSharedPtr>& GetNodes() const { return Nodes; }
+
+		/**
+		* Get all the nodes of the graph ( non-const method )
+		* this includes the node in all subgraphs
+		*/
+		TArray<FNodeSharedPtr>& GetNodes() { return Nodes; }
+
+		/** Get the number of nodes in the graph including the ones in the subgraphs */
+		int32 NumNodes() { return Nodes.Num(); }
+
+		template<class T> 
+		UE_DEPRECATED(5.8, "Use the Shared pointer parameter version of AddNode instead")
+		TSharedPtr<T> AddNode(T* InNode)
 		{
 			TSharedPtr<T> NewNode(InNode);
-			Nodes.AddUnique(NewNode);
-			for(const FName& RegisteredType : RegisteredFilters)
-			{
-				if (NewNode->IsA(RegisteredType))
-				{
-					FilteredNodes.FindOrAdd(RegisteredType).Add(NewNode);
-				}
-			}
+			AddNode(NewNode);
 			return NewNode;
 		}
 
-		template<class T> TSharedPtr<T> AddNode(TUniquePtr<T> &&InNode)
+		template<class T>
+		UE_DEPRECATED(5.8, "Use the Shared pointer parameter version of AddNode instead")
+		TSharedPtr<T> AddNode(TUniquePtr<T> &&InNode)
 		{
 			TSharedPtr<T> NewNode(InNode.Release());
-			Nodes.AddUnique(NewNode);
-			for(const FName& RegisteredType : RegisteredFilters)
-			{
-				if (NewNode->IsA(RegisteredType))
-				{
-					FilteredNodes.FindOrAdd(RegisteredType).Add(NewNode);
-				}
-			}
+			AddNode(NewNode);
 			return NewNode;
 		}
 
-		TSharedPtr<FDataflowNode> FindBaseNode(FGuid InGuid)
-		{
-			for (TSharedPtr<FDataflowNode> Node : Nodes)
-			{
-				if (Node->GetGuid() == InGuid)
-				{
-					return Node;
-				}
-			}
-			return TSharedPtr<FDataflowNode>(nullptr);
-		}
+		/** Add a new Dataflow node to the graph */
+		UE_API FNodeSharedPtr AddNode(FNodeSharedPtr NewNode);
 
-		TSharedPtr<const FDataflowNode> FindBaseNode(FGuid InGuid) const
-		{
-			for (TSharedPtr<FDataflowNode> Node : Nodes)
-			{
-				if (Node->GetGuid() == InGuid)
-				{
-					return Node;
-				}
-			}
-			return TSharedPtr<FDataflowNode>(nullptr);
-		}
+		/** Find a node by its unique identifier (non-const method) */
+		UE_API FNodeSharedPtr FindBaseNode(FGuid InGuid);
 
-		TSharedPtr<FDataflowNode> FindBaseNode(FName InName)
-		{
-			for (TSharedPtr<FDataflowNode> Node : Nodes)
-			{
-				if (Node->GetName().IsEqual(InName))
-				{
-					return Node;
-				}
-			}
-			return TSharedPtr<FDataflowNode>(nullptr);
-		}
+		/** Find a node by its unique identifier (const method) */
+		UE_API FConstNodeSharedPtr FindBaseNode(FGuid InGuid) const;
 
-		TSharedPtr<const FDataflowNode> FindBaseNode(FName InName) const
-		{
-			for (TSharedPtr<FDataflowNode> Node : Nodes)
-			{
-				if (Node->GetName().IsEqual(InName))
-				{
-					return Node;
-				}
-			}
-			return TSharedPtr<const FDataflowNode>(nullptr);
-		}
+		/** Find a node by its name (non-const method) */
+		UE_API FNodeSharedPtr FindBaseNode(FName InName);
 
-		TSharedPtr<FDataflowNode> FindFilteredNode(const FName& NodeFilter, FName InName) const
-		{
-			for (TSharedPtr<FDataflowNode> Node : GetFilteredNodes(NodeFilter))
-			{
-				if (Node->GetName().IsEqual(InName))
-				{
-					return Node;
-				}
-			}
-			return TSharedPtr<FDataflowNode>(nullptr);
-		}
+		/** Find a node by its name (const method) */
+		UE_API FConstNodeSharedPtr FindBaseNode(FName InName) const;
+
+		/** Find a node by name from a named filter (@see RegisterNodeFilter method) */
+		UE_API FNodeSharedPtr FindFilteredNode(const FName& NodeFilter, FName InName) const;
 
-		DATAFLOWCORE_API void RemoveNode(TSharedPtr<FDataflowNode> Node);
+		/** 
+		* Remove a specific node from the graph
+		* Removal of the node will automatically disconnect the inputs and outputs 
+		* this will also remove the node from any filters it may be in 
+		*/
+		UE_API void RemoveNode(FNodeSharedPtr Node);
 
+		/**
+		* Get all the graph connections
+		* This includes all the subgraph connections as well
+		*/
 		const TArray<FLink>& GetConnections() const { return Connections; }
-		DATAFLOWCORE_API void ClearConnections(FDataflowConnection* ConnectionBase);
-		DATAFLOWCORE_API void ClearConnections(FDataflowInput* Input);
-		DATAFLOWCORE_API void ClearConnections(FDataflowOutput* Output);
+
+		UE_API void ClearConnections(FDataflowConnection* ConnectionBase);
+		UE_API void ClearConnections(FDataflowInput* Input);
+		UE_API void ClearConnections(FDataflowOutput* Output);
 
 		enum class EConnectType : uint8
 		{
@@ -202,27 +177,27 @@ namespace UE::Dataflow
 			OUTPUT_PROMOTION, // output can be changed to adapt the input type
 		};
 
-		DATAFLOWCORE_API bool CanConnect(const FDataflowOutput& Output, const FDataflowInput& Input) const;
-		DATAFLOWCORE_API EConnectType GetConnectType(const FDataflowOutput& Output, const FDataflowInput& Input) const;
-		DATAFLOWCORE_API bool Connect(FDataflowOutput& Output, FDataflowInput& Input);
+		UE_API bool CanConnect(const FDataflowOutput& Output, const FDataflowInput& Input) const;
+		UE_API EConnectType GetConnectType(const FDataflowOutput& Output, const FDataflowInput& Input) const;
+		UE_API bool Connect(FDataflowOutput& Output, FDataflowInput& Input);
 
-		DATAFLOWCORE_API bool Connect(FDataflowConnection* ConnectionA, FDataflowConnection* ConnectionB);
-		DATAFLOWCORE_API void Connect(FDataflowOutput* OutputConnection, FDataflowInput* InputConnection);
-		DATAFLOWCORE_API void Disconnect(FDataflowOutput* OutputConnection, FDataflowInput* InputConnection);
+		UE_API bool Connect(FDataflowConnection* ConnectionA, FDataflowConnection* ConnectionB);
+		UE_API void Connect(FDataflowOutput* OutputConnection, FDataflowInput* InputConnection);
+		UE_API void Disconnect(FDataflowOutput* OutputConnection, FDataflowInput* InputConnection);
 
-		DATAFLOWCORE_API void AddReferencedObjects(FReferenceCollector& Collector);
+		UE_API void AddReferencedObjects(FReferenceCollector& Collector);
 
-		DATAFLOWCORE_API virtual void Serialize(FArchive& Ar, UObject* OwningObject);
+		UE_API virtual void Serialize(FArchive& Ar, UObject* OwningObject);
 		const TSet<FName>& GetDisabledNodes() const { return DisabledNodes; }
 
-		DATAFLOWCORE_API static void SerializeForSaving(FArchive& Ar, FGraph* InGraph, TArray<TSharedPtr<FDataflowNode>>& InNodes, TArray<FLink>& InConnections);
-		DATAFLOWCORE_API static void SerializeForLoading(FArchive& Ar, FGraph* InGraph, UObject* OwningObject);
+		UE_API static void SerializeForSaving(FArchive& Ar, FGraph* InGraph, TArray<TSharedPtr<FDataflowNode>>& InNodes, TArray<FLink>& InConnections);
+		UE_API static void SerializeForLoading(FArchive& Ar, FGraph* InGraph, UObject* OwningObject);
 
 	private:
 		void Reset();
 	};
 
-	DATAFLOWCORE_API void RegisterNodeFilter(const FName& NodeFilter);
+	UE_API void RegisterNodeFilter(const FName& NodeFilter);
 }
 
 
@@ -238,7 +213,7 @@ inline FArchive& operator<<(Chaos::FChaosArchive& Ar, UE::Dataflow::FLink& Value
 	return Ar;
 }
 
-
+#undef UE_API
 
 
 
diff --git a/Engine/Source/Runtime/Experimental/Dataflow/Core/Public/Dataflow/DataflowNodeFactory.h b/Engine/Source/Runtime/Experimental/Dataflow/Core/Public/Dataflow/DataflowNodeFactory.h
--- a/Engine/Source/Runtime/Experimental/Dataflow/Core/Public/Dataflow/DataflowNodeFactory.h
+++ b/Engine/Source/Runtime/Experimental/Dataflow/Core/Public/Dataflow/DataflowNodeFactory.h
@@ -133,7 +133,7 @@ namespace UE::Dataflow
 
 		template<class T> TSharedPtr<T> NewNode(FGraph& Graph, const FNewNodeParameters& Param)
 		{
-			return Graph.AddNode(new T(Param.Name, Param.Guid));
+			return Graph.AddNode(MakeShared<T>(Param.Name, Param.Guid));
 		}
 
 		DATAFLOWCORE_API void RegisterDefaultNodes();
