Build health fixes:
- non-unity fix by adding missing header in SConnectionStatusBar
- added FLogCategoryAlias typedef to fix targets with NO_LOGGING set to 1

#Synced-CL 48780745

diff --git a/Engine/Plugins/ChaosVD/Source/ChaosVD/Private/Widgets/SConnectionStatusBar.cpp b/Engine/Plugins/ChaosVD/Source/ChaosVD/Private/Widgets/SConnectionStatusBar.cpp
--- a/Engine/Plugins/ChaosVD/Source/ChaosVD/Private/Widgets/SConnectionStatusBar.cpp
+++ b/Engine/Plugins/ChaosVD/Source/ChaosVD/Private/Widgets/SConnectionStatusBar.cpp
@@ -6,6 +6,7 @@
 #include "ChaosVDRecordingDetails.h"
 #include "Components/HorizontalBox.h"
 #include "RemoteSessionsManager.h"
+#include "SessionInfo.h"
 #include "Widgets/Images/SImage.h"
 #include "Widgets/SOverlay.h"
 #include "Widgets/Text/STextBlock.h"
diff --git a/Engine/Source/Runtime/Experimental/TraceBasedDebuggers/Private/EngineEditorBridge.cpp b/Engine/Source/Runtime/Experimental/TraceBasedDebuggers/Private/EngineEditorBridge.cpp
--- a/Engine/Source/Runtime/Experimental/TraceBasedDebuggers/Private/EngineEditorBridge.cpp
+++ b/Engine/Source/Runtime/Experimental/TraceBasedDebuggers/Private/EngineEditorBridge.cpp
@@ -71,7 +71,7 @@ void FEngineEditorBridge::SetDataRelayTransportInstance(const TSharedPtr<IDataRe
 	RelayTraceDataTransportInstance = InExternalTraceRelayInstance;
 }
 
-FEngineEditorBridge::FEngineEditorBridge(const FLogCategoryBase& LogCategory)
+FEngineEditorBridge::FEngineEditorBridge(const FLogCategoryAlias& LogCategory)
 	: RemoteSessionsManager(SharedRemoteSessionsManager.IsValid() ? SharedRemoteSessionsManager.Pin().ToSharedRef() : MakeShared<FRemoteSessionsManager>())
 	, LogCategory(LogCategory)
 {
diff --git a/Engine/Source/Runtime/Experimental/TraceBasedDebuggers/Private/TraceBasedDebuggerRuntime.cpp b/Engine/Source/Runtime/Experimental/TraceBasedDebuggers/Private/TraceBasedDebuggerRuntime.cpp
--- a/Engine/Source/Runtime/Experimental/TraceBasedDebuggers/Private/TraceBasedDebuggerRuntime.cpp
+++ b/Engine/Source/Runtime/Experimental/TraceBasedDebuggers/Private/TraceBasedDebuggerRuntime.cpp
@@ -35,7 +35,7 @@ static FAutoConsoleVariable CVarMaxPendingBytesInRelayWriter(
 	536870912,
 	TEXT("How many bytes the relay writer can hold before overflowing, triggering a close. Defaults to ~512MB"));
 
-FRuntimeModule::FRuntimeModule(const FLogCategoryBase& InLogCategory, const FGuid& InDebuggerTypeID)
+FRuntimeModule::FRuntimeModule(const FLogCategoryAlias& InLogCategory, const FGuid& InDebuggerTypeID)
 	: LogCategory(InLogCategory)
 	, DebuggerGuid(InDebuggerTypeID)
 {
diff --git a/Engine/Source/Runtime/Experimental/TraceBasedDebuggers/Public/EngineEditorBridge.h b/Engine/Source/Runtime/Experimental/TraceBasedDebuggers/Public/EngineEditorBridge.h
--- a/Engine/Source/Runtime/Experimental/TraceBasedDebuggers/Public/EngineEditorBridge.h
+++ b/Engine/Source/Runtime/Experimental/TraceBasedDebuggers/Public/EngineEditorBridge.h
@@ -27,8 +27,14 @@ class FEngineEditorBridge
 {
 public:
 
+#if NO_LOGGING
+	typedef FNoLoggingCategory FLogCategoryAlias;
+#else
+	typedef FLogCategoryBase FLogCategoryAlias;
+#endif
+
 	FEngineEditorBridge() = delete;
-	UE_API explicit FEngineEditorBridge(const FLogCategoryBase& LogCategory);
+	UE_API explicit FEngineEditorBridge(const FLogCategoryAlias& LogCategory);
 
 	virtual ~FEngineEditorBridge() = default;
 
@@ -96,7 +102,7 @@ private:
 	/** Handler created by derived classes to handle more message types */
 	TSharedPtr<IRemoteSessionsHandler> RemoteSessionsHandler;
 
-	const FLogCategoryBase& LogCategory;
+	const FLogCategoryAlias& LogCategory;
 
 	FTSTicker::FDelegateHandle RecordingStatusUpdateHandle;
 
diff --git a/Engine/Source/Runtime/Experimental/TraceBasedDebuggers/Public/TraceBasedDebuggerRuntime.h b/Engine/Source/Runtime/Experimental/TraceBasedDebuggers/Public/TraceBasedDebuggerRuntime.h
--- a/Engine/Source/Runtime/Experimental/TraceBasedDebuggers/Public/TraceBasedDebuggerRuntime.h
+++ b/Engine/Source/Runtime/Experimental/TraceBasedDebuggers/Public/TraceBasedDebuggerRuntime.h
@@ -22,9 +22,15 @@ DECLARE_DELEGATE_RetVal(uint64, FExternalRelayBufferQueueSizeDelegate)
 
 class FRuntimeModule
 {
+#if NO_LOGGING
+	typedef FNoLoggingCategory FLogCategoryAlias;
+#else
+	typedef FLogCategoryBase FLogCategoryAlias;
+#endif
+
 protected:
 	FRuntimeModule() = delete;
-	UE_API explicit FRuntimeModule(const FLogCategoryBase& InLogCategory, const FGuid& InDebuggerTypeID);
+	UE_API explicit FRuntimeModule(const FLogCategoryAlias& InLogCategory, const FGuid& InDebuggerTypeID);
 
 	UE_API virtual ~FRuntimeModule();
 
@@ -170,7 +176,7 @@ protected:
 	 */
 	FString HostAddressCommand = TEXT("DebuggerHost");
 
-	const FLogCategoryBase& LogCategory;
+	const FLogCategoryAlias& LogCategory;
 
 private:
 
