[TraceBasedDebuggers] updated RelayTraceDataWriter to add proper support for LogCategory in NO_LOGGING targets

#Synced-CL 48784988

diff --git a/Engine/Source/Runtime/Experimental/TraceBasedDebuggers/Private/RelayTraceDataWriter.cpp b/Engine/Source/Runtime/Experimental/TraceBasedDebuggers/Private/RelayTraceDataWriter.cpp
--- a/Engine/Source/Runtime/Experimental/TraceBasedDebuggers/Private/RelayTraceDataWriter.cpp
+++ b/Engine/Source/Runtime/Experimental/TraceBasedDebuggers/Private/RelayTraceDataWriter.cpp
@@ -9,7 +9,7 @@
 
 namespace UE::TraceBasedDebuggers
 {
-	FRelayTraceDataWriter::FRelayTraceDataWriter(const FLogCategoryBase& LogCategory)
+	FRelayTraceDataWriter::FRelayTraceDataWriter(const FLogCategoryAlias& LogCategory)
 		: LogCategory(LogCategory)
 	{
 	}
diff --git a/Engine/Source/Runtime/Experimental/TraceBasedDebuggers/Public/RelayTraceDataWriter.h b/Engine/Source/Runtime/Experimental/TraceBasedDebuggers/Public/RelayTraceDataWriter.h
--- a/Engine/Source/Runtime/Experimental/TraceBasedDebuggers/Public/RelayTraceDataWriter.h
+++ b/Engine/Source/Runtime/Experimental/TraceBasedDebuggers/Public/RelayTraceDataWriter.h
@@ -22,8 +22,14 @@ namespace UE::TraceBasedDebuggers
  */
 struct FRelayTraceDataWriter
 {
+#if NO_LOGGING
+	typedef FNoLoggingCategory FLogCategoryAlias;
+#else
+	typedef FLogCategoryBase FLogCategoryAlias;
+#endif
+
 	FRelayTraceDataWriter() = delete;
-	UE_API explicit FRelayTraceDataWriter(const FLogCategoryBase& LogCategory);
+	UE_API explicit FRelayTraceDataWriter(const FLogCategoryAlias& LogCategory);
 
 	UE_API ~FRelayTraceDataWriter();
 
@@ -109,7 +115,7 @@ private:
 
 	bool CanWriteData(uint32 NewDataSize) const;
 
-	const FLogCategoryBase& LogCategory;
+	const FLogCategoryAlias& LogCategory;
 	TSpscQueue<TUniquePtr<FBufferArchive>> PendingBunchesToRelay;
 	TUniquePtr<FBufferArchive> CurrentPendingBunch;
 	std::atomic<int32> PendingBunchesCount = 0;
