Geometry Collection. Fix for crashes when using undo/redo.
Fracture operations have both GC object and component in the transaction. On undo/redo we the component can recreate scene proxy before the object has recreated the RenderData.
This leads to scene proxy being created with an inconsistent state.
Added an additional copy of the render data GUID as a GC object UProperty. We use that to validate if the render data is out of date due to some behind the scenes property update, such as those that come during an undo/redo.
This is a short term fix. A longer term fix would be to have everything that the scene proxy needs to ref inside the RenderData.
#Synced-CL 48777116

diff --git a/Engine/Source/Runtime/Experimental/GeometryCollectionEngine/Private/GeometryCollection/GeometryCollectionObject.cpp b/Engine/Source/Runtime/Experimental/GeometryCollectionEngine/Private/GeometryCollection/GeometryCollectionObject.cpp
--- a/Engine/Source/Runtime/Experimental/GeometryCollectionEngine/Private/GeometryCollection/GeometryCollectionObject.cpp
+++ b/Engine/Source/Runtime/Experimental/GeometryCollectionEngine/Private/GeometryCollection/GeometryCollectionObject.cpp
@@ -152,6 +152,7 @@ UGeometryCollection::UGeometryCollection(const FObjectInitializer& ObjectInitial
 #if WITH_EDITOR
 	SimulationDataGuid = StateGuid;
 	RenderDataGuid = StateGuid;
+	ValidRenderDataGuid = StateGuid;
 	bStripOnCook = GeometryCollectionAssetForceStripOnCook;
 #endif
 	PhysicsMaterial = GEngine? GEngine->DefaultPhysMaterial: nullptr;
@@ -1519,6 +1520,7 @@ void UGeometryCollection::RebuildRenderData()
 
 		PropagateMarkDirtyToComponents();
 		RenderDataGuid = StateGuid;
+		ValidRenderDataGuid = StateGuid;
 	}
 }
 
@@ -2041,7 +2043,7 @@ void UGeometryCollection::AppendToClassSchema(FAppendToClassSchemaContext& Conte
 	Context.Update(RenderDataVersionString.GetData(), RenderDataVersionString.Len() * sizeof(RenderDataVersionString[0]));
 }
 
-void UGeometryCollection::EnsureDataIsCooked(bool bInitResources, bool bIsTransacting, bool bIsPersistant, bool bAllowCopyFromDDC)
+void UGeometryCollection::EnsureDataIsCooked(bool bInitResources, bool bIsTransacting, bool bIsPersistent, bool bAllowCopyFromDDC)
 {
 	EnsureSimulationDataIsCooked(bIsTransacting, bAllowCopyFromDDC);
 
@@ -2093,6 +2095,8 @@ void UGeometryCollection::EnsureRenderDataIsCooked(bool bInitResources)
 		}
 
 		LastBuiltRenderDataGuid = StateGuid;
+		RenderDataGuid = StateGuid;
+		ValidRenderDataGuid = StateGuid;
 	}
 }
 
@@ -2103,9 +2107,9 @@ void UGeometryCollection::PreSave(FObjectPreSaveContext SaveContext)
 #if WITH_EDITOR
 	constexpr bool bInitResources = false; 
 	constexpr bool bIsTransacting = false;
-	constexpr bool bIsPersistant = false; // note that this has no effect on the call below  
+	constexpr bool bIsPersistent = false; // note that this has no effect on the call below  
 	constexpr bool bAllowCopyFromDDC = false;
-	EnsureDataIsCooked(bInitResources, bIsTransacting, bIsPersistant, bAllowCopyFromDDC);
+	EnsureDataIsCooked(bInitResources, bIsTransacting, bIsPersistent, bAllowCopyFromDDC);
 #endif
 
 	Super::PreSave(SaveContext);
@@ -2204,12 +2208,20 @@ bool UGeometryCollection::IsReadyForFinishDestroy()
 
 bool UGeometryCollection::HasMeshData() const
 {
-	return RenderData != nullptr && RenderData->bHasMeshData;
+	return RenderData != nullptr &&
+#if WITH_EDITOR
+		RenderDataGuid == ValidRenderDataGuid &&
+#endif
+		RenderData->bHasMeshData;
 }
 
 bool UGeometryCollection::HasNaniteData() const
 {
-	return RenderData != nullptr && RenderData->bHasNaniteData;
+	return RenderData != nullptr && 
+#if WITH_EDITOR
+		RenderDataGuid == ValidRenderDataGuid &&
+#endif
+		RenderData->bHasNaniteData;
 }
 
 uint32 UGeometryCollection::GetNaniteResourceID() const
diff --git a/Engine/Source/Runtime/Experimental/GeometryCollectionEngine/Public/GeometryCollection/GeometryCollectionObject.h b/Engine/Source/Runtime/Experimental/GeometryCollectionEngine/Public/GeometryCollection/GeometryCollectionObject.h
--- a/Engine/Source/Runtime/Experimental/GeometryCollectionEngine/Public/GeometryCollection/GeometryCollectionObject.h
+++ b/Engine/Source/Runtime/Experimental/GeometryCollectionEngine/Public/GeometryCollection/GeometryCollectionObject.h
@@ -413,7 +413,7 @@ public:
 
 #if WITH_EDITOR
 	GEOMETRYCOLLECTIONENGINE_API static void AppendToClassSchema(FAppendToClassSchemaContext& Context);
-	GEOMETRYCOLLECTIONENGINE_API void EnsureDataIsCooked(bool bInitResources, bool bIsTransacting, bool bIsPersistant, bool bAllowCopyFromDDC = true);
+	GEOMETRYCOLLECTIONENGINE_API void EnsureDataIsCooked(bool bInitResources, bool bIsTransacting, bool bIsPersistent, bool bAllowCopyFromDDC = true);
 #endif
 
 	/** Accessors for internal geometry collection */
@@ -970,6 +970,17 @@ private:
 	FGuid RenderDataGuid;
 #endif
 
+#if WITH_EDITORONLY_DATA
+	/** 
+	 * Serialized version of RenderDataGuid. 
+	 * This is used to determine when the stored RenderData mismatches the current GC object state.
+	 * This can happen temporarily in editor during a undo/redo transaction when a GC component processing the PostUndo() before the GC object.
+	 * todo: Remove this fix and move all relevent state to the RenderData so that it is the only source of truth for render scene proxies.
+	 */
+	UPROPERTY()
+	FGuid ValidRenderDataGuid;
+#endif
+
 	// cached root index for faster queries
 	UPROPERTY(VisibleAnywhere, Category = "Clustering")
 	int32 RootIndex = INDEX_NONE;
