Chaos - add a Union to the root geometry for the first added geometry

#rnx
#Synced-CL 48771487

diff --git a/Engine/Source/Runtime/Experimental/Chaos/Private/Chaos/ParticleHandle.cpp b/Engine/Source/Runtime/Experimental/Chaos/Private/Chaos/ParticleHandle.cpp
--- a/Engine/Source/Runtime/Experimental/Chaos/Private/Chaos/ParticleHandle.cpp
+++ b/Engine/Source/Runtime/Experimental/Chaos/Private/Chaos/ParticleHandle.cpp
@@ -78,12 +78,17 @@ namespace Chaos
 	template <typename T, int d>
 	void Chaos::TGeometryParticle<T, d>::MergeGeometry(TArray<Chaos::FImplicitObjectPtr>&& Objects)
 	{
-		ensure(MNonFrequentData.Read().GetGeometry());
+		// We always have a union at the root. Add one if these are the first objects to be added
+		if (MNonFrequentData.Read().GetGeometry() == nullptr)
+		{
+			Chaos::FImplicitObjectPtr Union = MakeImplicitObjectPtr<Chaos::FImplicitObjectUnion>(MoveTemp(Objects));
+			SetGeometry(Union);
+			return;
+		}
 
 		// we only support FImplicitObjectUnion
-		ensure(MNonFrequentData.Read().GetGeometry()->GetType() == FImplicitObjectUnion::StaticType());
-
-		if (MNonFrequentData.Read().GetGeometry()->GetType() == FImplicitObjectUnion::StaticType())
+		check(MNonFrequentData.Read().GetGeometry());
+		if (ensure(MNonFrequentData.Read().GetGeometry()->GetType() == FImplicitObjectUnion::StaticType()))
 		{
 			// Only adding to the root union - shallow copy allowed here.
 			ModifyGeometry(EGeometryAccess::ShallowCopy,
@@ -105,8 +110,14 @@ namespace Chaos
 			return;
 		}
 
+		// We always have a union at the root. Add one if these are the first objects to be added
 		const FImplicitObjectRef CurrentGeometry = GetGeometry();
-		if (ensure(CurrentGeometry != nullptr))
+		if (CurrentGeometry == nullptr)
+		{
+			Chaos::FImplicitObjectPtr Union = MakeImplicitObjectPtr<Chaos::FImplicitObjectUnion>(MoveTemp(Objects));
+			SetGeometry(Union);
+		}
+		else
 		{
 			if (ensure(CurrentGeometry->GetType() == FImplicitObjectUnion::StaticType()))
 			{
