Chaos API - minor fixes and tweaks

#rnx
#Synced-CL 48771487

diff --git a/Engine/Plugins/Experimental/ChaosRigidPhysicsAsync/Source/ChaosRigidPhysicsAsync/Private/ChaosRigidPhysicsAsync/RigidBodyContainerAsync.cpp b/Engine/Plugins/Experimental/ChaosRigidPhysicsAsync/Source/ChaosRigidPhysicsAsync/Private/ChaosRigidPhysicsAsync/RigidBodyContainerAsync.cpp
--- a/Engine/Plugins/Experimental/ChaosRigidPhysicsAsync/Source/ChaosRigidPhysicsAsync/Private/ChaosRigidPhysicsAsync/RigidBodyContainerAsync.cpp
+++ b/Engine/Plugins/Experimental/ChaosRigidPhysicsAsync/Source/ChaosRigidPhysicsAsync/Private/ChaosRigidPhysicsAsync/RigidBodyContainerAsync.cpp
@@ -100,6 +100,7 @@ namespace Chaos::Rigids::Async
 
 	FRigidBodyPoolAsyncPT::~FRigidBodyPoolAsyncPT()
 	{
+		// All bodies should have been destroyed before the pool is destroyed
 		check(BodyPool.GetNumAllocated() == 0);
 		check(IdToBodyMap.IsEmpty());
 	}
diff --git a/Engine/Source/Runtime/Experimental/RigidPhysics/Public/RigidPhysics/RigidBody.h b/Engine/Source/Runtime/Experimental/RigidPhysics/Public/RigidPhysics/RigidBody.h
--- a/Engine/Source/Runtime/Experimental/RigidPhysics/Public/RigidPhysics/RigidBody.h
+++ b/Engine/Source/Runtime/Experimental/RigidPhysics/Public/RigidPhysics/RigidBody.h
@@ -483,7 +483,7 @@ namespace UE::Physics
 	template<typename ContextType>
 	void TRigidBodyContainerPtrImpl<ContextType>::DestroyBody(TRigidBodyPtr<FContext>&& InBodyPtr)
 	{
-		BodyContainerRaw->DestroyBody(InBodyPtr.Ge());
+		BodyContainerRaw->DestroyBody(InBodyPtr.Get());
 	}
 
 	template<typename ContextType>
@@ -495,8 +495,11 @@ namespace UE::Physics
 	template<typename ContextType>
 	void TRigidScenePtrImpl<ContextType>::DestroyBody(TRigidBodyPtr<ContextType>&& InBody)
 	{
-		SceneRaw->DestroyBody(InBody.Get());
-		InBody.Reset();
+		if (InBody.IsValid())
+		{
+			SceneRaw->DestroyBody(InBody.Get());
+			InBody.Reset();
+		}
 	}
 
 }
diff --git a/Engine/Source/Runtime/Experimental/RigidPhysics/Public/RigidPhysics/RigidObjectPtr.h b/Engine/Source/Runtime/Experimental/RigidPhysics/Public/RigidPhysics/RigidObjectPtr.h
--- a/Engine/Source/Runtime/Experimental/RigidPhysics/Public/RigidPhysics/RigidObjectPtr.h
+++ b/Engine/Source/Runtime/Experimental/RigidPhysics/Public/RigidPhysics/RigidObjectPtr.h
@@ -98,6 +98,7 @@ namespace UE::Physics
 		// Extract the object interface (internal use only)
 		UE_INTERNAL FInterface* Get() const
 		{
+			CheckIsValid();
 			return Impl.Get();
 		}
 
diff --git a/Engine/Source/Runtime/Experimental/RigidPhysics/Public/RigidPhysics/RigidTyped.h b/Engine/Source/Runtime/Experimental/RigidPhysics/Public/RigidPhysics/RigidTyped.h
--- a/Engine/Source/Runtime/Experimental/RigidPhysics/Public/RigidPhysics/RigidTyped.h
+++ b/Engine/Source/Runtime/Experimental/RigidPhysics/Public/RigidPhysics/RigidTyped.h
@@ -77,10 +77,10 @@ namespace UE::Physics
 		RIGIDPHYSICS_API virtual ~IRigidTyped();
 
 		// Get the TypeId of this class
-		static const FRigidTypeId& GetStaticTypeId();
+		RIGIDPHYSICS_API static const FRigidTypeId& GetStaticTypeId();
 
 		// Get the leaf TypeId of this object
-		virtual const FRigidTypeId& GetTypeId() const;
+		RIGIDPHYSICS_API virtual const FRigidTypeId& GetTypeId() const;
 
 		const FString& GetTypeName() const
 		{
