Start collecting statistics on the amount of data evicted from the IAS Cache.

#rnx

- We record each pruned item as it is removed, this will allow us to track how many items have been evicted from the cache if we choose to.
-- Currently we treat pruning (entries being overwritten) and eviction (items being discarded due to errors) as the same stat.
-- We do not record the evictions if the cache is dropped.
- Fixed a minor potential issue in FDiskCache::Prune where Candidate was being used after it was removed from the TMap.
- Reorganized the on screen display a little.
-- The number of items in the cache are now grouped with the total size of entries in the cache. I have removed the text from this part of the display as the numbers are fairly obvious 'cached:' was just noise.
-- Added the new evicted stat as well as the amount of data written to the cache.

#Synced-CL 48686588

diff --git a/Engine/Source/Runtime/Experimental/IoStore/OnDemand/Private/JournaledCache.cpp b/Engine/Source/Runtime/Experimental/IoStore/OnDemand/Private/JournaledCache.cpp
--- a/Engine/Source/Runtime/Experimental/IoStore/OnDemand/Private/JournaledCache.cpp
+++ b/Engine/Source/Runtime/Experimental/IoStore/OnDemand/Private/JournaledCache.cpp
@@ -794,6 +794,8 @@ bool FDiskCache::Evict(uint64 Key)
 	const TPair<uint64, FMapEntry>& Element = DataMap.Get(Id);
 	MappedBytes -= Element.Value.Size;
 
+	FOnDemandIoBackendStats::Get()->OnCacheEvictedBytes(Element.Value.Size);
+
 	DataMap.Remove(Id);
 
 	FOnDemandIoBackendStats::Get()->OnCachePersistedBytes(MappedBytes, DataMap.Num());
@@ -938,8 +940,11 @@ void FDiskCache::Prune(uint64 DataBase, uint32 Size)
 			continue;
 		}
 
+		const int64 CandidateSize = Candidate.Size;
 		Iter.RemoveCurrent();
-		MappedBytes -= Candidate.Size;
+
+		MappedBytes -= CandidateSize;
+		FOnDemandIoBackendStats::Get()->OnCacheEvictedBytes(CandidateSize);
 
 		Overage = FMath::Max(Overage, B[1] - A[1]);
 		BytesRemoved += int32(B[1] - B[0]);
diff --git a/Engine/Source/Runtime/Experimental/IoStore/OnDemand/Private/Statistics.cpp b/Engine/Source/Runtime/Experimental/IoStore/OnDemand/Private/Statistics.cpp
--- a/Engine/Source/Runtime/Experimental/IoStore/OnDemand/Private/Statistics.cpp
+++ b/Engine/Source/Runtime/Experimental/IoStore/OnDemand/Private/Statistics.cpp
@@ -368,6 +368,7 @@ FCounterAtomicInt		GCachePutRejectCount(TEXT("Ias/CachePutRejectCount"), TraceCo
 FCounterAtomicInt		GCacheCachedBytes(TEXT("Ias/CacheCachedBytes"), TraceCounterDisplayHint_Memory);
 FCounterAtomicInt		GCacheCachedEntries(TEXT("Ias/CacheCachedEntries"), TraceCounterDisplayHint_Memory);
 FCounterAtomicInt		GCacheWrittenBytes(TEXT("Ias/CacheWrittenBytes"), TraceCounterDisplayHint_Memory);
+FCounterAtomicInt		GCacheEvictedBytes(TEXT("Ias/CacheEvictedBytes"), TraceCounterDisplayHint_Memory);
 FCounterFloat			GCacheSuspendedSeconds(TEXT("Ias/CacheSuspendedSeconds"), TraceCounterDisplayHint_None);
 int64					GCacheMaxBytes = 0;
 FCounterAtomicInt		GCachePendingBytes(TEXT("Ias/CachePendingBytes"), TraceCounterDisplayHint_Memory);
@@ -467,6 +468,7 @@ CSV_DEFINE_STAT(Ias, CachePutExistingCount);
 CSV_DEFINE_STAT(Ias, CachePutRejectCount);
 CSV_DEFINE_STAT(Ias, CacheCachedMB);
 CSV_DEFINE_STAT(Ias, CacheWrittenMB);
+CSV_DEFINE_STAT(Ias, CacheEvictedMB);
 CSV_DEFINE_STAT(Ias, CacheReadMB);
 CSV_DEFINE_STAT(Ias, CacheRejectedMB);
 // http stat totals
@@ -632,6 +634,7 @@ void FOnDemandIoBackendStats::UpdateCSVValues()
 	int32 CPutRejectCount = int32(GCachePutRejectCount.Get());
 	float CCachedKiB = (float)GCacheCachedBytes.Get() * OneOver1024;
 	float CWrittenKiB = (float)GCacheWrittenBytes.Get() * OneOver1024;
+	float CEvictedKiB = (float)GCacheEvictedBytes.Get() * OneOver1024;
 	float CReadKiB = (float)GCacheReadBytes.Get() * OneOver1024;
 	float CRejectedKiB = (float)GCacheRejectBytes.Get() * OneOver1024;
 	CSV_CUSTOM_STAT_DEFINED(CacheGetCount, CGetCount, ECsvCustomStatOp::Set);
@@ -641,6 +644,7 @@ void FOnDemandIoBackendStats::UpdateCSVValues()
 	CSV_CUSTOM_STAT_DEFINED(CachePutRejectCount, CPutRejectCount, ECsvCustomStatOp::Set);
 	CSV_CUSTOM_STAT_DEFINED(CacheCachedMB, CCachedKiB * OneOver1024, ECsvCustomStatOp::Set);
 	CSV_CUSTOM_STAT_DEFINED(CacheWrittenMB, CWrittenKiB * OneOver1024, ECsvCustomStatOp::Set);
+	CSV_CUSTOM_STAT_DEFINED(CacheEvictedMB, CEvictedKiB * OneOver1024, ECsvCustomStatOp::Set);
 	CSV_CUSTOM_STAT_DEFINED(CacheReadMB, CReadKiB * OneOver1024, ECsvCustomStatOp::Set);
 	CSV_CUSTOM_STAT_DEFINED(CacheRejectedMB, CRejectedKiB * OneOver1024, ECsvCustomStatOp::Set);
 
@@ -681,13 +685,15 @@ void FOnDemandIoBackendStats::PrintPeriodicLogging()
 
 		float CCachedKiB = (float)GCacheCachedBytes.Get() * OneOver1024;
 		float CWrittenKiB = (float)GCacheWrittenBytes.Get() * OneOver1024;
+		float CEvictedKiB = (float)GCacheEvictedBytes.Get() * OneOver1024;
 		float CReadKiB = (float)GCacheReadBytes.Get() * OneOver1024;
 		float CRejectedKiB = (float)GCacheRejectBytes.Get() * OneOver1024;
 
 		if (BackendStatus.IsCacheEnabled())
 		{
-			UE_LOG(LogIas, Log, TEXT("IASCache: Count=%d, CachedKiB=%d, WrittenKiB=%d, ReadKiB=%d, RejectedKiB=%d, Get=%d, Error=%d, Put=%d, PutReject=%d, PutExisting=%d"),
-				(int32)GCacheCachedEntries.Get(), (int32)CCachedKiB, (int32)CWrittenKiB, (int32)CReadKiB, (int32)CRejectedKiB, CGetCount, CErrorCount, CPutCount, CPutRejectCount, CPutExistingCount);
+			UE_LOG(LogIas, Log, TEXT("IASCache: Count=%d, CachedKiB=%d, WrittenKiB=%d, ReadKiB=%d, EvictedKiB=%d, RejectedKiB=%d, Get=%d, Error=%d, Put=%d, PutReject=%d, PutExisting=%d"),
+				(int32)GCacheCachedEntries.Get(), (int32)CCachedKiB, (int32)CWrittenKiB, (int32)CReadKiB,
+				(int32)CEvictedKiB, (int32)CRejectedKiB, CGetCount, CErrorCount, CPutCount, CPutRejectCount, CPutExistingCount);
 		}
 		else
 		{
@@ -923,11 +929,18 @@ void FOnDemandIoBackendStats::PrintOnScreenStatistics(FCoreDelegates::FSeverityM
 		{
 			const float CacheUsagePercent = GCacheMaxBytes > 0 ? 100.f * (float(GCacheCachedBytes.Get()) / float(GCacheMaxBytes)) : 0.f;
 
-			CachingText << TEXT("Cached: ") << FText::AsMemory(GCacheCachedBytes.Get()).ToString();
-			CachingText.Appendf(TEXT(" (%.1f%%)"), CacheUsagePercent);
-			CachingText << TEXT(" | Rejected: ") << FText::AsMemory(GCacheRejectBytes.Get()).ToString();
+			CachingText.Appendf(TEXT(" (%.1f%%) "), CacheUsagePercent);
+			CachingText << FText::AsMemory(GCacheCachedBytes.Get()).ToString();
+			CachingText << TEXT(" - (") << GCacheGetCount.Get() << TEXT(")");
+
+			CachingText << TEXT(" | Write: ") << FText::AsMemory(GCacheWrittenBytes.Get()).ToString();
 			CachingText << TEXT(" | Read: ") << FText::AsMemory(GCacheReadBytes.Get()).ToString();
-			CachingText << TEXT(" (") << GCacheGetCount.Get() << TEXT(")");
+			CachingText << TEXT(" | Evict: ") << FText::AsMemory(GCacheEvictedBytes.Get()).ToString();
+			if (GCacheRejectBytes.Get() > 0)
+			{
+				CachingText << TEXT(" | Rejected: ") << FText::AsMemory(GCacheRejectBytes.Get()).ToString();
+			}
+
 			CachingText << TEXT(" (Boot: ") << GCacheBootMs << TEXT("ms)");
 
 			const double SuspendedSeconds = GCacheSuspendedSeconds.Get();
@@ -1044,6 +1057,7 @@ void FOnDemandIoBackendStats::ReportEndPointAnalytics(TArray<FAnalyticsEventAttr
 			, TEXT("IasCacheUsagePercent"), CacheUsagePercent
 
 			, UE_TRACK_DELTA("IasCacheWriteBytes", GCacheWrittenBytes.Get())
+			, UE_TRACK_DELTA("IasCacheEvictedBytes", GCacheEvictedBytes.Get())
 			, UE_TRACK_DELTA("IasCacheReadBytes", GCacheReadBytes.Get())
 			, UE_TRACK_DELTA("IasCacheRejectBytes", GCacheRejectBytes.Get())
 		);
@@ -1136,6 +1150,7 @@ private:
 				, TEXT("IasCachePutCount"), Cache.PutCount.GetValue()
 
 				, TEXT("IasCacheWriteBytes"), Cache.WrittenBytes.GetValue()
+				, TEXT("IasCacheEvictedBytes"), Cache.EvictedBytes.GetValue()
 				, TEXT("IasCacheReadBytes"), Cache.ReadBytes.GetValue()
 				, TEXT("IasCacheRejectBytes"), Cache.RejectBytes.GetValue()
 			);
@@ -1210,6 +1225,7 @@ private:
 		TTrackedValue<FCounterAtomicInt, GCachePutCount>			PutCount;
 
 		TTrackedValue<FCounterAtomicInt, GCacheWrittenBytes>		WrittenBytes;
+		TTrackedValue<FCounterAtomicInt, GCacheEvictedBytes>		EvictedBytes;
 		TTrackedValue<FCounterAtomicInt, GCacheReadBytes>			ReadBytes;
 		TTrackedValue<FCounterAtomicInt, GCacheRejectBytes>			RejectBytes;
 	} Cache;
@@ -1316,6 +1332,11 @@ void FOnDemandIoBackendStats::OnCacheWriteBytes(uint64 WriteSize)
 	GCacheWrittenBytes.Add(WriteSize);
 }
 
+void FOnDemandIoBackendStats::OnCacheEvictedBytes(uint64 EvictedSize)
+{
+	GCacheEvictedBytes.Add(EvictedSize);
+}
+
 void FOnDemandIoBackendStats::OnCacheSetMaxBytes(uint64 TotalSize)
 {
 	GCacheMaxBytes = TotalSize;
diff --git a/Engine/Source/Runtime/Experimental/IoStore/OnDemand/Private/Statistics.h b/Engine/Source/Runtime/Experimental/IoStore/OnDemand/Private/Statistics.h
--- a/Engine/Source/Runtime/Experimental/IoStore/OnDemand/Private/Statistics.h
+++ b/Engine/Source/Runtime/Experimental/IoStore/OnDemand/Private/Statistics.h
@@ -85,6 +85,7 @@ public:
 	void OnCachePendingBytes(uint64 TotalSize) IAS_STATISTICS_IMPL()
 	void OnCachePersistedBytes(uint64 TotalSize, int32 NumEntries) IAS_STATISTICS_IMPL()
 	void OnCacheWriteBytes(uint64 WriteSize) IAS_STATISTICS_IMPL()
+	void OnCacheEvictedBytes(uint64 EvictedSize) IAS_STATISTICS_IMPL()
 	void OnCacheSetMaxBytes(uint64 TotalSize) IAS_STATISTICS_IMPL()
 	void OnCacheSuspended(double Seconds) IAS_STATISTICS_IMPL()
 
