Add new analytics to track how much IAS cache data is read from disk and how much from the in memory cache.

#rnx

- Add a new enum ECacheType that described the type of cache that the data was loaded from. I tried to keep the wording vague to avoid tying us to exact concepts but in general intermediate will mean in memory access and deferred will mean disk access.
- Fixed an issue where data loaded from the IAS cache on disk was being recorded by both the IAS Cache and the http dispatcher, causing the read to be recorded twice.

#Synced-CL 48686588

diff --git a/Engine/Source/Runtime/Experimental/IoStore/OnDemand/Private/JournaledCache.cpp b/Engine/Source/Runtime/Experimental/IoStore/OnDemand/Private/JournaledCache.cpp
--- a/Engine/Source/Runtime/Experimental/IoStore/OnDemand/Private/JournaledCache.cpp
+++ b/Engine/Source/Runtime/Experimental/IoStore/OnDemand/Private/JournaledCache.cpp
@@ -1448,35 +1448,38 @@ FCache::FGetToken FCache::Get(uint64 Key, FIoBuffer& OutData) const
 		return FGetToken(Key);
 	}
 
-	// Nothing's on disk, so lets try the memory cache
-	FReadScopeLock _(MemLock);
-
-
-	if (CloneCachedIOBufferData > 0)
 	{
-		// Speculative fix for a read after free happening in FIoChunkEncoding::Decode
-		if (Partial.Key == Key)
-		{
-			OutData = FIoBuffer(FIoBuffer::Clone, Partial.Data.Data(), Partial.Data.GetSize());
-		}
-		else if (const FIoBuffer* Data = MemCache.Get(Key); Data != nullptr)
+		// Nothing is on disk, so lets try the memory cache
+		FReadScopeLock _(MemLock);
+
+		if (CloneCachedIOBufferData > 0)
 		{
-			OutData = FIoBuffer(FIoBuffer::Clone, Data->Data(), Data->GetSize());
+			// Speculative fix for a read after free happening in FIoChunkEncoding::Decode
+			if (Partial.Key == Key)
+			{
+				OutData = FIoBuffer(FIoBuffer::Clone, Partial.Data.Data(), Partial.Data.GetSize());
+			}
+			else if (const FIoBuffer* Data = MemCache.Get(Key); Data != nullptr)
+			{
+				OutData = FIoBuffer(FIoBuffer::Clone, Data->Data(), Data->GetSize());
+			}
 		}
-	}
-	else
-	{
-		if (const FIoBuffer* Data = MemCache.Get(Key); Data != nullptr)
+		else
 		{
-			OutData = *Data;
-		}
+			if (const FIoBuffer* Data = MemCache.Get(Key); Data != nullptr)
+			{
+				OutData = *Data;
+			}
 
-		if (Partial.Key == Key)
-		{
-			OutData = Partial.Data;
+			if (Partial.Key == Key)
+			{
+				OutData = Partial.Data;
+			}
 		}
 	}
 
+	FOnDemandIoBackendStats::Get()->OnCacheGet(OutData.GetSize(), ECacheType::Memory);
+
 	return 0;
 }
 
@@ -1519,7 +1522,7 @@ EIoErrorCode FCache::Materialize(FGetToken Token, FIoBuffer& OutData, uint32 Off
 	EIoErrorCode Ret = DiskCache.Materialize(Key, OutData, Offset);
 	if (Ret == EIoErrorCode::Ok)
 	{
-		FOnDemandIoBackendStats::Get()->OnCacheGet(OutData.GetSize());
+		FOnDemandIoBackendStats::Get()->OnCacheGet(OutData.GetSize(), ECacheType::Disk);
 	}
 
 	return Ret;
diff --git a/Engine/Source/Runtime/Experimental/IoStore/OnDemand/Private/OnDemandHttpIoDispacher.cpp b/Engine/Source/Runtime/Experimental/IoStore/OnDemand/Private/OnDemandHttpIoDispacher.cpp
--- a/Engine/Source/Runtime/Experimental/IoStore/OnDemand/Private/OnDemandHttpIoDispacher.cpp
+++ b/Engine/Source/Runtime/Experimental/IoStore/OnDemand/Private/OnDemandHttpIoDispacher.cpp
@@ -1553,7 +1553,7 @@ bool FHttpDispatcher::TryReadFromCache(FHttpRequest* Request)
 	if (CacheStatus == EIoErrorCode::Ok)
 	{
 		InflightRequestCount.fetch_add(1, std::memory_order_relaxed);
-		FOnDemandIoBackendStats::Get()->OnCacheGet(Request->Buffer.GetSize());
+
 		UE::Tasks::Launch(UE_SOURCE_LOCATION, [this, Request]
 		{
 			CompleteRequest(*Request, Request->Buffer, FIoHttpRange(), 200, true);
@@ -1566,11 +1566,11 @@ bool FHttpDispatcher::TryReadFromCache(FHttpRequest* Request)
 	{
 		InflightRequestCount.fetch_add(1, std::memory_order_relaxed);
 		UE::Tasks::FTaskEvent OnReadyEvent(TEXT("HttpCacheReadComplete"));
+
 		Launch(UE_SOURCE_LOCATION, [this, Request]
 		{
 			if (Request->CacheStatus == EIoErrorCode::Ok)
 			{
-				FOnDemandIoBackendStats::Get()->OnCacheGet(Request->Buffer.GetSize());
 				CompleteRequest(*Request, Request->Buffer, FIoHttpRange(), 200, true);
 			}
 			else if ((Request->CacheStatus == EIoErrorCode::Cancelled) || Request->Flags.HasAny(EHttpRequestFlags::CancelRequested) || bStopRequested.load(std::memory_order_relaxed))
diff --git a/Engine/Source/Runtime/Experimental/IoStore/OnDemand/Private/Statistics.cpp b/Engine/Source/Runtime/Experimental/IoStore/OnDemand/Private/Statistics.cpp
--- a/Engine/Source/Runtime/Experimental/IoStore/OnDemand/Private/Statistics.cpp
+++ b/Engine/Source/Runtime/Experimental/IoStore/OnDemand/Private/Statistics.cpp
@@ -373,6 +373,8 @@ FCounterFloat			GCacheSuspendedSeconds(TEXT("Ias/CacheSuspendedSeconds"), TraceC
 int64					GCacheMaxBytes = 0;
 FCounterAtomicInt		GCachePendingBytes(TEXT("Ias/CachePendingBytes"), TraceCounterDisplayHint_Memory);
 FCounterAtomicInt		GCacheReadBytes(TEXT("Ias/CacheReadBytes"), TraceCounterDisplayHint_Memory);
+FCounterAtomicInt		GCacheReadMemory(TEXT("Ias/CacheReadMemoryBytes"), TraceCounterDisplayHint_Memory);
+FCounterAtomicInt		GCacheReadDisk(TEXT("Ias/CacheReadDiskBytes"), TraceCounterDisplayHint_Memory);
 FCounterAtomicInt		GCacheRejectBytes(TEXT("Ias/CachePutRejectBytes"), TraceCounterDisplayHint_Memory);
 // general stats
 double					GAnalyticsStartTiming = 0.0;
@@ -1059,6 +1061,8 @@ void FOnDemandIoBackendStats::ReportEndPointAnalytics(TArray<FAnalyticsEventAttr
 			, UE_TRACK_DELTA("IasCacheWriteBytes", GCacheWrittenBytes.Get())
 			, UE_TRACK_DELTA("IasCacheEvictedBytes", GCacheEvictedBytes.Get())
 			, UE_TRACK_DELTA("IasCacheReadBytes", GCacheReadBytes.Get())
+			, UE_TRACK_DELTA("IasCacheReadMemoryBytes", GCacheReadMemory.Get())
+			, UE_TRACK_DELTA("IasCacheReadDiskBytes", GCacheReadDisk.Get())
 			, UE_TRACK_DELTA("IasCacheRejectBytes", GCacheRejectBytes.Get())
 		);
 #undef UE_TRACK_DELTA
@@ -1152,6 +1156,8 @@ private:
 				, TEXT("IasCacheWriteBytes"), Cache.WrittenBytes.GetValue()
 				, TEXT("IasCacheEvictedBytes"), Cache.EvictedBytes.GetValue()
 				, TEXT("IasCacheReadBytes"), Cache.ReadBytes.GetValue()
+				, TEXT("IasCacheReadMemoryBytes"), Cache.ReadMemoryBytes.GetValue()
+				, TEXT("IasCacheReadDiskBytes"), Cache.ReadDiskBytes.GetValue()
 				, TEXT("IasCacheRejectBytes"), Cache.RejectBytes.GetValue()
 			);
 		}
@@ -1227,6 +1233,8 @@ private:
 		TTrackedValue<FCounterAtomicInt, GCacheWrittenBytes>		WrittenBytes;
 		TTrackedValue<FCounterAtomicInt, GCacheEvictedBytes>		EvictedBytes;
 		TTrackedValue<FCounterAtomicInt, GCacheReadBytes>			ReadBytes;
+		TTrackedValue<FCounterAtomicInt, GCacheReadMemory>			ReadMemoryBytes;
+		TTrackedValue<FCounterAtomicInt, GCacheReadDisk>			ReadDiskBytes;
 		TTrackedValue<FCounterAtomicInt, GCacheRejectBytes>			RejectBytes;
 	} Cache;
 
@@ -1294,10 +1302,22 @@ void FOnDemandIoBackendStats::OnCacheDecodeError()
 	GCacheDecodeErrorCount.Increment();
 }
 
-void FOnDemandIoBackendStats::OnCacheGet(uint64 DataSize)
+void FOnDemandIoBackendStats::OnCacheGet(uint64 DataSize, ECacheType Type)
 {
 	GCacheGetCount.Increment();
 	GCacheReadBytes.Add(DataSize);
+
+	switch(Type)
+	{
+	case ECacheType::Memory:
+		GCacheReadMemory.Add(DataSize);
+		break;
+	case ECacheType::Disk:
+		GCacheReadDisk.Add(DataSize);
+		break;
+	default:
+		checkNoEntry();
+	}
 }
 
 void FOnDemandIoBackendStats::OnCachePut()
diff --git a/Engine/Source/Runtime/Experimental/IoStore/OnDemand/Private/Statistics.h b/Engine/Source/Runtime/Experimental/IoStore/OnDemand/Private/Statistics.h
--- a/Engine/Source/Runtime/Experimental/IoStore/OnDemand/Private/Statistics.h
+++ b/Engine/Source/Runtime/Experimental/IoStore/OnDemand/Private/Statistics.h
@@ -40,6 +40,14 @@ enum class EHttpRequestType : uint8;
 struct FAvailableEndPoints;
 struct FOnDemandImmediateAnalyticHandler;
 
+enum class ECacheType : uint8
+{
+	/* The data was returned immediately from the in memory cache. */
+	Memory,
+	/* The data was returned at a later point from the disk cache. */
+	Disk,
+};
+
 #if IAS_WITH_STATISTICS
 #	define IAS_STATISTICS_IMPL(...) ;
 #else
@@ -78,7 +86,7 @@ public:
 	void OnCacheBootMs(uint64 TimeMs) IAS_STATISTICS_IMPL()
 	void OnCacheError() IAS_STATISTICS_IMPL()
 	void OnCacheDecodeError() IAS_STATISTICS_IMPL()
-	void OnCacheGet(uint64 DataSize) IAS_STATISTICS_IMPL()
+	void OnCacheGet(uint64 DataSize, ECacheType Type) IAS_STATISTICS_IMPL()
 	void OnCachePut() IAS_STATISTICS_IMPL()
 	void OnCachePutExisting(uint64 DataSize) IAS_STATISTICS_IMPL()
 	void OnCachePutReject(uint64 DataSize) IAS_STATISTICS_IMPL()
