Remove unused cvar 'ias.EnableCancelation'

#rnx

- The cvar was accidently broken when the new backend was introduced but we haven't really been using it for debugging so I see no reason not to just remove it rather than fixing it.

#Synced-CL 48708293

diff --git a/Engine/Source/Runtime/Experimental/IoStore/OnDemand/Private/OnDemandHttpIoDispacher.cpp b/Engine/Source/Runtime/Experimental/IoStore/OnDemand/Private/OnDemandHttpIoDispacher.cpp
--- a/Engine/Source/Runtime/Experimental/IoStore/OnDemand/Private/OnDemandHttpIoDispacher.cpp
+++ b/Engine/Source/Runtime/Experimental/IoStore/OnDemand/Private/OnDemandHttpIoDispacher.cpp
@@ -29,7 +29,6 @@
 
 // When enabled the cvar 'iax.InvalidUrlChance' will be enabled and allow us to simulate invalid urls
 #define UE_ALLOW_INVALID_URL_DEBUGGING !UE_BUILD_SHIPPING
-#define UE_ALLOW_DISABLE_CANCELLING 1 && !UE_BUILD_SHIPPING
 #define UE_ALLOW_CACHE_POISONING !UE_BUILD_SHIPPING
 #define UE_ALLOW_HTTP_PAUSE (1 && !NO_CVARS && !UE_BUILD_SHIPPING)
 
@@ -87,6 +86,7 @@ ENUM_CLASS_FLAGS(EHttpRequestTypeFilter);
 DECLARE_MULTICAST_DELEGATE(FOnRecreateHttpClient);
 FOnRecreateHttpClient OnRecreateHttpClient;
 
+
 ///////////////////////////////////////////////////////////////////////////////
 void OnHttpClientCVarChanged(IConsoleVariable* CVar)
 {
@@ -175,15 +175,6 @@ static FAutoConsoleVariableRef CVar_EnableRequestHandleClear(
 	TEXT("When enabled FChunkRequest::HttpRequest will be cleared when the request callback is triggered")
 );
 
-#if UE_ALLOW_DISABLE_CANCELLING
-bool GIasCancelationEnabled = true;
-static FAutoConsoleVariableRef CVar_IasCancelationEnabled(
-	TEXT("ias.EnableCancelation"),
-	GIasCancelationEnabled,
-	TEXT("Allows existing IO requests to be canceled")
-);
-#endif //UE_ALLOW_DISABLE_CANCELLING
-
 #if UE_ALLOW_CACHE_POISONING
 bool GIasPoisonCache = false;
 static FAutoConsoleVariableRef CVar_IasPoisonCache(
