Support register HTTP host groups without providing a test URL.

#rnx

#Synced-CL 48712420

diff --git a/Engine/Source/Runtime/Experimental/IoStore/OnDemand/Private/IasHostGroup.cpp b/Engine/Source/Runtime/Experimental/IoStore/OnDemand/Private/IasHostGroup.cpp
--- a/Engine/Source/Runtime/Experimental/IoStore/OnDemand/Private/IasHostGroup.cpp
+++ b/Engine/Source/Runtime/Experimental/IoStore/OnDemand/Private/IasHostGroup.cpp
@@ -221,20 +221,35 @@ void FIASHostGroup::Connect(int32 HostIndex)
 	SetPrimaryHost(HostIndex);
 }
 
-void FIASHostGroup::Disconnect()
+bool FIASHostGroup::Disconnect()
 {
+	// Never disconnect when the registered test URL is empty
+	if (Impl->TestPath.IsEmpty())
+	{
+		Impl->HttpErrorHistory.Reset();
+		return false;
+	}
+
 	Impl->bHttpEnabled = false;
 	Impl->HttpErrorHistory.Reset();
 
 	SetPrimaryHost(INDEX_NONE);
 
 	FHostGroupManager::Get().OnHostGroupDisconncted().Broadcast();
+
+	return true;
 }
 
 FIASHostGroup::EReconnectionResult FIASHostGroup::AttemptReconnection(uint32 TimeoutMs, std::atomic_bool& CancellationToken)
 {
 	if (!IsConnected())
 	{
+		if (Impl->TestPath.IsEmpty())
+		{
+			Connect(0);
+			return EReconnectionResult::Reconnected;
+		}
+
 		UE_LOGFMT(LogIas, Log, "[{HostName}] Trying to reconnect to any available endpoint...", GetName());
 
 		if (const int32 Idx = ConnectionTest(GetHostUrls(), Impl->TestPath, TimeoutMs, CancellationToken); Idx != INDEX_NONE)
@@ -282,12 +297,13 @@ bool FIASHostGroup::OnFailedResponse()
 
 	if (bAboveHighWaterMark && IsConnected())
 	{
-		Disconnect();
-
-		UE_LOG(LogIas, Warning, TEXT("[%s] Host group disabled due to high water mark of %.2f of the last %d requests reached"),
-			*Impl->Name.ToString(), GIasHttpErrorHighWater * 100.0f, GIasHttpErrorSampleCount);
+		if (Disconnect())
+		{
+			UE_LOG(LogIas, Warning, TEXT("[%s] Host group disabled due to high water mark of %.2f of the last %d requests reached"),
+				*Impl->Name.ToString(), GIasHttpErrorHighWater * 100.0f, GIasHttpErrorSampleCount);
 
-		return true;
+			return true;
+		}
 	}
 
 	return false;
diff --git a/Engine/Source/Runtime/Experimental/IoStore/OnDemand/Private/IasHostGroup.h b/Engine/Source/Runtime/Experimental/IoStore/OnDemand/Private/IasHostGroup.h
--- a/Engine/Source/Runtime/Experimental/IoStore/OnDemand/Private/IasHostGroup.h
+++ b/Engine/Source/Runtime/Experimental/IoStore/OnDemand/Private/IasHostGroup.h
@@ -37,7 +37,7 @@ public:
 	FIoStatus Resolve(TConstArrayView<FString> HostUrls);
 	
 	void Connect(int32 HostIndex);
-	void Disconnect();
+	bool Disconnect();
 
 	enum class EReconnectionResult : uint8
 	{
diff --git a/Engine/Source/Runtime/Experimental/IoStore/OnDemandCore/Private/HttpIoDispatcher.cpp b/Engine/Source/Runtime/Experimental/IoStore/OnDemandCore/Private/HttpIoDispatcher.cpp
--- a/Engine/Source/Runtime/Experimental/IoStore/OnDemandCore/Private/HttpIoDispatcher.cpp
+++ b/Engine/Source/Runtime/Experimental/IoStore/OnDemandCore/Private/HttpIoDispatcher.cpp
@@ -266,6 +266,11 @@ FIoStatus FHttpIoDispatcher::RegisterHostGroup(const FName& HostGroup, TConstArr
 	return FIoStatus(EIoErrorCode::InvalidCode);
 }
 
+FIoStatus FHttpIoDispatcher::RegisterHostGroup(const FName& HostGroup, TConstArrayView<FAnsiString> HostNames)
+{
+	return RegisterHostGroup(HostGroup, HostNames, FAnsiStringView());
+}
+
 FIoStatus FHttpIoDispatcher::RegisterHostGroup(const FName& HostGroup, FAnsiStringView HostName, FAnsiStringView TestUrl)
 {
 	FAnsiString HostNameString(HostName);
@@ -273,6 +278,11 @@ FIoStatus FHttpIoDispatcher::RegisterHostGroup(const FName& HostGroup, FAnsiStri
 	return RegisterHostGroup(HostGroup, HostNames, TestUrl);
 }
 
+FIoStatus FHttpIoDispatcher::RegisterHostGroup(const FName& HostGroup, FAnsiStringView HostName)
+{
+	return RegisterHostGroup(HostGroup, HostName, FAnsiStringView());
+}
+
 bool FHttpIoDispatcher::IsHostGroupRegistered(const FName& HostGroup)
 {
 	return GHttpIoDispatcher.IsValid() ? GHttpIoDispatcher->IsHostGroupRegistered(HostGroup) : false;
diff --git a/Engine/Source/Runtime/Experimental/IoStore/OnDemandCore/Public/IO/HttpIoDispatcher.h b/Engine/Source/Runtime/Experimental/IoStore/OnDemandCore/Public/IO/HttpIoDispatcher.h
--- a/Engine/Source/Runtime/Experimental/IoStore/OnDemandCore/Public/IO/HttpIoDispatcher.h
+++ b/Engine/Source/Runtime/Experimental/IoStore/OnDemandCore/Public/IO/HttpIoDispatcher.h
@@ -372,7 +372,9 @@ public:
 
 	UE_API static void				GetHttpStats(FOnDemandHttpStats& Out);
 	UE_API static FIoStatus			RegisterHostGroup(const FName& HostGroup, TConstArrayView<FAnsiString> HostNames, FAnsiStringView TestUrl);
+	UE_API static FIoStatus			RegisterHostGroup(const FName& HostGroup, TConstArrayView<FAnsiString> HostNames);
 	UE_API static FIoStatus			RegisterHostGroup(const FName& HostGroup, FAnsiStringView HostName, FAnsiStringView TestUrl);
+	UE_API static FIoStatus			RegisterHostGroup(const FName& HostGroup, FAnsiStringView HostName);
 	UE_API static bool				IsHostGroupRegistered(const FName& HostGroup);
 	UE_API static bool				IsHostGroupOk(const FName& HostGroup);
 
