Requests that end with a 0-sized frame would time out.

Head request that have no body, CDNs send this;

id:1 | frame:headers | flag:end of headers | flag:end of stream

Except one provider which sends;

id:1 | frame:headers | flag:end of headers
id:1 | frame:data | size:0 | flag:end of stream

Because we finish the request on headers if there's no body, this extra data frame becomes orphaned and we have no idea what "id=1" corresponds to. So we wait until a data frame comes in before moving on getting the request content such that we get to correctly see the zero-sized frame and not miss the stream-end flag.

#rnx

diff --git a/Engine/Source/Runtime/Experimental/IoStore/HttpClient/Private/TransactionTwo.inl b/Engine/Source/Runtime/Experimental/IoStore/HttpClient/Private/TransactionTwo.inl
--- a/Engine/Source/Runtime/Experimental/IoStore/HttpClient/Private/TransactionTwo.inl
+++ b/Engine/Source/Runtime/Experimental/IoStore/HttpClient/Private/TransactionTwo.inl
@@ -447,15 +447,8 @@ FOutcome FDriverNg::TickInner(FTlsPeer& Peer)
 	RecvRemaining = RecvFrame.GetSize();
 	if (RecvFrame.IsPertinent())
 	{
-		// We'll not expect higher levels to come to read zero-sized data.
-		if (RecvFrame.Type != FFrame::EType::Data || RecvRemaining)
-		{
-			uint32 Id = RecvFrame.GetId();
-			return FOutcome::Ok(Id);
-		}
-
-		check(RecvFrame.Flags & FFrame::Flag_EndStream);
-		RecvFrame = {};
+		uint32 Id = RecvFrame.GetId();
+		return FOutcome::Ok(Id);
 	}
 
 	if (RecvRemaining == 0)
@@ -482,6 +475,18 @@ FOutcome FDriverNg::Send(FTlsPeer& Peer)
 ////////////////////////////////////////////////////////////////////////////////
 FOutcome FDriverNg::RecvResponse(FTlsPeer& Peer, TArray<FHeaderNg>* Sink)
 {
+	if (RecvFrame.Type == FFrame::EType::Data)
+	{
+		// Sometimes streams of requests with no content can be closed with an
+		// empty data frame instead of via header frame's flags. By deferring
+		// RecvResponse() completion we avoid orphaning zero-sized data frames
+		if ((RecvFrame.Flags & FFrame::Flag_EndStream) && RecvFrame.GetSize() == 0)
+		{
+			RecvFrame = {};
+		}
+		return FOutcome::Ok();
+	}
+
 	ResponseHeaders = Sink;
 	ON_SCOPE_EXIT { ResponseHeaders = nullptr; };
 
@@ -500,6 +505,7 @@ FOutcome FDriverNg::RecvResponse(FTlsPeer& Peer, TArray<FHeaderNg>* Sink)
 		}
 
 		int32 EndOfHeaders = (RecvFrame.Flags & FFrame::Flag_EndHeaders);
+		int32 EndOfStream = (RecvFrame.Flags & FFrame::Flag_EndStream);
 		uint32 ThisId = RecvFrame.GetId();
 
 		if (EndOfHeaders)
@@ -508,9 +514,33 @@ FOutcome FDriverNg::RecvResponse(FTlsPeer& Peer, TArray<FHeaderNg>* Sink)
 		}
 
 		FOutcome Outcome = Tick(Peer); 
-		if (Outcome.IsError())				return Outcome;
-		if (EndOfHeaders)					return FOutcome::Ok();
-		if (Outcome.GetResult() != ThisId)	return FOutcome::WaitStream();
+		if (Outcome.IsError())
+		{
+			return Outcome;
+		}
+
+		if (EndOfStream)
+		{
+			check(EndOfHeaders);
+			return FOutcome::Ok();
+		}
+
+		if (EndOfHeaders)
+		{
+			// We shall end headers once the first data frame arrives, which may
+			// have already happened in the Tick() call.
+			if (Outcome.GetResult() == ThisId)
+			{
+				return RecvResponse(Peer, ResponseHeaders);
+			}
+
+			return FOutcome::Waiting();
+		}
+
+		if (Outcome.GetResult() != ThisId)
+		{
+			return FOutcome::WaitStream();
+		}
 	}
 }
 
