Add LineNo to IadCacheDefrag and IadCachePurge analytic payloads.

#Synced-CL 48800600

diff --git a/Engine/Source/Runtime/Experimental/IoStore/OnDemand/Private/Statistics.cpp b/Engine/Source/Runtime/Experimental/IoStore/OnDemand/Private/Statistics.cpp
--- a/Engine/Source/Runtime/Experimental/IoStore/OnDemand/Private/Statistics.cpp
+++ b/Engine/Source/Runtime/Experimental/IoStore/OnDemand/Private/Statistics.cpp
@@ -131,6 +131,29 @@ FOnDemandImmediateAnalytic MakeAnalyticsEventFromResult(const TCHAR* EventName,
 	};
 }
 
+////////////////////////////////////////////////////////////////////////////////
+int32 TryGetLineNo(const UE::UnifiedError::FError& Error)
+{
+	int32 LineNo = -1;
+	const UE::UnifiedError::IoStoreOnDemand::FInstallCacheErrorContext* ErrorCtx =
+		Error.GetErrorContext<UE::UnifiedError::IoStoreOnDemand::FInstallCacheErrorContext>();
+	if (ErrorCtx)
+	{
+		LineNo = ErrorCtx->LineNo;
+	}
+	return LineNo;
+}
+
+int32 TryGetLineNo(const FResult& Result)
+{
+	int32 LineNo = -1;
+	if (Result.HasError())
+	{
+		LineNo = TryGetLineNo(Result.GetError());
+	}
+	return LineNo;
+}
+
 ////////////////////////////////////////////////////////////////////////////////
 static int32 BytesToApproxMB(uint64 Bytes) { return int32(Bytes >> 20); }
 static int32 BytesToApproxKB(uint64 Bytes) { return int32(Bytes >> 10); }
@@ -1666,8 +1689,12 @@ void FOnDemandInstallCacheStats::OnPurge(
 		InstallCacheAnalytics.PurgedBytes += PurgedBytes;
 	}
 
+	const int32 LineNo = TryGetLineNo(Result);
+
 	FOnDemandImmediateAnalytic Event = MakeAnalyticsEventFromResult(TEXT("IadCachePurge"), Result);
-	AppendAnalyticsEventAttributeArray(Event.AnalyticsArray, TEXT("PurgedBytes"), PurgedBytes);
+	AppendAnalyticsEventAttributeArray(Event.AnalyticsArray, 
+		TEXT("PurgedBytes"), PurgedBytes,
+		TEXT("LineNo"), LineNo);
 	SendImmediateAnalytic(MoveTemp(Event));
 }
 
@@ -1706,6 +1733,8 @@ void FOnDemandInstallCacheStats::OnDefrag(
 		InstallerAnalytics.DefraggedBlockFragmentionRatio += BlocksToDefragFragmentionRatio;
 	}
 
+	const int32 LineNo = TryGetLineNo(Result);
+
 	FOnDemandImmediateAnalytic Event = MakeAnalyticsEventFromResult(TEXT("IadCacheDefrag"), Result);
 	AppendAnalyticsEventAttributeArray(Event.AnalyticsArray, 
 		TEXT("TotalFragmentedBytes"), TotalFragmentedBytes,
@@ -1713,7 +1742,8 @@ void FOnDemandInstallCacheStats::OnDefrag(
 		TEXT("TotalFragmentionRatio"), TotalFragmentionRatio,
 		TEXT("DefraggedBlocksFragmentedSize"), BlocksToDefragFragmentedSize,
 		TEXT("DefraggedBlocksTotalSize"), BlocksToDefragTotalSize,
-		TEXT("DefraggedBlocksFragmentionRatio"), BlocksToDefragFragmentionRatio);
+		TEXT("DefraggedBlocksFragmentionRatio"), BlocksToDefragFragmentionRatio,
+		TEXT("LineNo"), LineNo);
 	SendImmediateAnalytic(MoveTemp(Event));
 }
 
