Dataflow : fix potential crash in field processing when looking up a transform index from a handle
- This may happen when a partial cluster is reported in the field and do not match any direct transform index

#Synced-CL 48816559

diff --git a/Engine/Source/Runtime/Experimental/Chaos/Private/PhysicsProxy/GeometryCollectionPhysicsProxy.cpp b/Engine/Source/Runtime/Experimental/Chaos/Private/PhysicsProxy/GeometryCollectionPhysicsProxy.cpp
--- a/Engine/Source/Runtime/Experimental/Chaos/Private/PhysicsProxy/GeometryCollectionPhysicsProxy.cpp
+++ b/Engine/Source/Runtime/Experimental/Chaos/Private/PhysicsProxy/GeometryCollectionPhysicsProxy.cpp
@@ -6059,25 +6059,27 @@ void FGeometryCollectionPhysicsProxy::FieldParameterUpdateCallback(Chaos::FPBDRi
 											CurrResult >= std::numeric_limits<int8>::min());
 
 										const int8 ResultState = static_cast<int8>(CurrResult);
-										const int32 TransformIndex = HandleToTransformGroupIndex[RigidHandle];
-
-										Chaos::FVec3 InitialLinearVelocity{ 0 };
-										Chaos::FVec3 InitialAngularVelocity{ 0 };
-
-										// bUpdateViews is used here to avoid applying the velocity twice when running the Initial fields ( see InitializeBodiesPT )
-										const bool bHasInitialVelocities = (bUpdateViews && (Parameters.InitialVelocityType == EInitialVelocityTypeEnum::Chaos_Initial_Velocity_User_Defined));
-										if (bHasInitialVelocities)
+										const int32 TransformIndex = GetTransformGroupIndexFromHandle(RigidHandle);
+										if (TransformIndex != INDEX_NONE)
 										{
-											InitialLinearVelocity = InitialVelocityFacade.InitialLinearVelocityAttribute[TransformIndex];
-											InitialAngularVelocity = InitialVelocityFacade.InitialAngularVelocityAttribute[TransformIndex];
-										}
+											Chaos::FVec3 InitialLinearVelocity{ 0 };
+											Chaos::FVec3 InitialAngularVelocity{ 0 };
 
-										// Update of the handles object state. 
-										bHasStateChanged |= ReportDynamicStateResult(RigidSolver, static_cast<Chaos::EObjectStateType>(ResultState), RigidHandle,
+											// bUpdateViews is used here to avoid applying the velocity twice when running the Initial fields ( see InitializeBodiesPT )
+											const bool bHasInitialVelocities = (bUpdateViews && (Parameters.InitialVelocityType == EInitialVelocityTypeEnum::Chaos_Initial_Velocity_User_Defined));
+											if (bHasInitialVelocities)
+											{
+												InitialLinearVelocity = InitialVelocityFacade.InitialLinearVelocityAttribute[TransformIndex];
+												InitialAngularVelocity = InitialVelocityFacade.InitialAngularVelocityAttribute[TransformIndex];
+											}
+
+											// Update of the handles object state. 
+											bHasStateChanged |= ReportDynamicStateResult(RigidSolver, static_cast<Chaos::EObjectStateType>(ResultState), RigidHandle,
 												bHasInitialVelocities, InitialLinearVelocity, bHasInitialVelocities, InitialAngularVelocity);
 
-										// Update of the Collection dynamic state. It will be used just after to set the initial velocity
-										Collection.DynamicState[TransformIndex] = ResultState;
+											// Update of the Collection dynamic state. It will be used just after to set the initial velocity
+											Collection.DynamicState[TransformIndex] = ResultState;
+										}
 									}
 								}
 								if (bUpdateViews && bHasStateChanged)
@@ -6114,8 +6116,11 @@ void FGeometryCollectionPhysicsProxy::FieldParameterUpdateCallback(Chaos::FPBDRi
 											Chaos::TPBDRigidParticleHandle<Chaos::FReal, 3>* RigidHandle = ParticleHandles[Index.Sample]->CastToRigidParticle();
 											if (RigidHandle)
 											{
-												const int32 TransformIndex = HandleToTransformGroupIndex[RigidHandle];
-												InitialVelocityFacade.InitialLinearVelocityAttribute.ModifyAt(TransformIndex, FVector3f(ResultsView[Index.Result]));
+												const int32 TransformIndex = GetTransformGroupIndexFromHandle(RigidHandle);
+												if (TransformIndex != INDEX_NONE)
+												{
+													InitialVelocityFacade.InitialLinearVelocityAttribute.ModifyAt(TransformIndex, FVector3f(ResultsView[Index.Result]));
+												}
 											}
 										}
 									}
@@ -6145,8 +6150,11 @@ void FGeometryCollectionPhysicsProxy::FieldParameterUpdateCallback(Chaos::FPBDRi
 											Chaos::TPBDRigidParticleHandle<Chaos::FReal, 3>* RigidHandle = ParticleHandles[Index.Sample]->CastToRigidParticle();
 											if (RigidHandle)
 											{
-												const int32 TransformIndex = HandleToTransformGroupIndex[RigidHandle];
-												InitialVelocityFacade.InitialAngularVelocityAttribute.ModifyAt(TransformIndex, FVector3f(ResultsView[Index.Result]));
+												const int32 TransformIndex = GetTransformGroupIndexFromHandle(RigidHandle);
+												if (TransformIndex != INDEX_NONE)
+												{
+													InitialVelocityFacade.InitialAngularVelocityAttribute.ModifyAt(TransformIndex, FVector3f(ResultsView[Index.Result]));
+												}
 											}
 										}
 									}
