IAD Installer
- VisitPackageDependencies now uses a TFunctionRef for the visitor function.
- Make FCas::BlockIds private
- Fix timestamps not being removed from LastAccess map when processing the BlockDeleted journal op

#Synced-CL 48800600

diff --git a/Engine/Source/Runtime/Experimental/IoStore/OnDemand/Private/OnDemandContentInstaller.cpp b/Engine/Source/Runtime/Experimental/IoStore/OnDemand/Private/OnDemandContentInstaller.cpp
--- a/Engine/Source/Runtime/Experimental/IoStore/OnDemand/Private/OnDemandContentInstaller.cpp
+++ b/Engine/Source/Runtime/Experimental/IoStore/OnDemand/Private/OnDemandContentInstaller.cpp
@@ -4,6 +4,7 @@
 
 #include "Algo/Accumulate.h"
 #include "Async/UniqueLock.h"
+#include "Containers/RingBuffer.h"
 #include "Experimental/UnifiedError/CoreErrorTypes.h"
 #include "HAL/Platform.h"
 #include "IO/IoBuffer.h"
@@ -104,6 +105,88 @@ int32 SecondsToMillicseconds(double Seconds)
 	return Seconds > 0.0 ? int32(Seconds * 1000.0) : 0;
 }
 
+////////////////////////////////////////////////////////////////////////////////
+void VisitPackageDependencies(
+	const TSet<FPackageId>& PackageIds, 
+	bool bIncludeSoftReferences, 
+	FPackageDependencyVisitor Visitor, 
+	FPackageStoreReadScope* InReadScope)
+{
+	TRACE_CPUPROFILER_EVENT_SCOPE(FOnDemandContentInstaller::TVisitPackageDependencies);
+
+	TRingBuffer<FPackageId> Queue;
+	TSet<FPackageId>		Visitied;
+
+	Visitied.Reserve(PackageIds.Num());
+	Queue.Reserve(PackageIds.Num());
+	for (const FPackageId& PackageId : PackageIds)
+	{
+		Queue.Add(PackageId);
+	}
+
+	FPackageStore& PackageStore = FPackageStore::Get();
+	TOptional<FPackageStoreReadScope> ReadScope;
+	if (InReadScope == nullptr)
+	{
+		ReadScope.Emplace(PackageStore);
+	}
+
+	while (!Queue.IsEmpty())
+	{
+		FPackageId PackageId = Queue.PopFrontValue();
+		{
+			FName		SourcePackageName;
+			FPackageId	RedirectedToPackageId;
+			if (PackageStore.GetPackageRedirectInfo(PackageId, SourcePackageName, RedirectedToPackageId))
+			{
+				PackageId = RedirectedToPackageId;
+			}
+		}
+
+		bool bIsAlreadyInSet = false;
+		Visitied.Add(PackageId, &bIsAlreadyInSet);
+		if (bIsAlreadyInSet)
+		{
+			continue;
+		}
+
+		FPackageStoreEntry PackageStoreEntry;
+		TSharedPtr<IPackageStoreBackend> Backend;
+		const EPackageStoreEntryStatus EntryStatus = PackageStore.GetPackageStoreEntry(PackageId, NAME_None, PackageStoreEntry, &Backend);
+
+		const bool bContinue = Visitor(PackageId, EntryStatus, PackageStoreEntry, Backend);
+		if (bContinue == false)
+		{
+			break;
+		}
+
+		if (EntryStatus != EPackageStoreEntryStatus::Missing)
+		{
+			for (const FPackageId& ImportedPackageId : PackageStoreEntry.ImportedPackageIds)
+			{
+				if (!Visitied.Contains(ImportedPackageId))
+				{
+					Queue.Add(ImportedPackageId);
+				}
+			}
+
+			if (bIncludeSoftReferences)
+			{
+				TConstArrayView<FPackageId> SoftReferences;
+				TConstArrayView<uint32> Indices = PackageStore.GetSoftReferences(PackageId, SoftReferences);
+				for (uint32 Idx : Indices)
+				{
+					const FPackageId& SoftRef = SoftReferences[Idx];
+					if (!Visitied.Contains(SoftRef))
+					{
+						Queue.Add(SoftRef);
+					}
+				}
+			}
+		}
+	}
+}
+
 ////////////////////////////////////////////////////////////////////////////////
 void ResolvePackageDependencies(
 	const TSet<FPackageId>& PackageIds,
@@ -141,7 +224,7 @@ void ResolvePackageDependencies(
 		return true;
 	};
 
-	TVisitPackageDependencies(PackageIds, bIncludeSoftReferences, Visitor, InReadScope);
+	VisitPackageDependencies(PackageIds, bIncludeSoftReferences, Visitor, InReadScope);
 }
 
 ////////////////////////////////////////////////////////////////////////////////
diff --git a/Engine/Source/Runtime/Experimental/IoStore/OnDemand/Private/OnDemandContentInstaller.h b/Engine/Source/Runtime/Experimental/IoStore/OnDemand/Private/OnDemandContentInstaller.h
--- a/Engine/Source/Runtime/Experimental/IoStore/OnDemand/Private/OnDemandContentInstaller.h
+++ b/Engine/Source/Runtime/Experimental/IoStore/OnDemand/Private/OnDemandContentInstaller.h
@@ -4,7 +4,6 @@
 
 #include "Async/Mutex.h"
 #include "Containers/Set.h"
-#include "Containers/RingBuffer.h"
 #include "Experimental/UnifiedError/CoreErrorTypes.h"
 #include "IO/IoStoreOnDemand.h"
 #include "IO/OnDemandError.h"
@@ -44,83 +43,19 @@ struct FResolvedContainerChunks
 };
 
 ////////////////////////////////////////////////////////////////////////////////
-template<typename VisitorType>
-void TVisitPackageDependencies(const TSet<FPackageId>& PackageIds, bool bIncludeSoftReferences, const VisitorType& Visitor, FPackageStoreReadScope* InReadScope = nullptr)
-{
-	TRACE_CPUPROFILER_EVENT_SCOPE(FOnDemandContentInstaller::TVisitPackageDependencies);
-
-	TRingBuffer<FPackageId> Queue;
-	TSet<FPackageId>		Visitied;
-
-	Visitied.Reserve(PackageIds.Num());
-	Queue.Reserve(PackageIds.Num());
-	for (const FPackageId& PackageId : PackageIds)
-	{
-		Queue.Add(PackageId);
-	}
-
-	FPackageStore& PackageStore = FPackageStore::Get();
-	TOptional<FPackageStoreReadScope> ReadScope;
-	if (InReadScope == nullptr)
-	{
-		ReadScope.Emplace(PackageStore);
-	}
-
-	while (!Queue.IsEmpty())
-	{
-		FPackageId PackageId = Queue.PopFrontValue();
-		{
-			FName		SourcePackageName;
-			FPackageId	RedirectedToPackageId;
-			if (PackageStore.GetPackageRedirectInfo(PackageId, SourcePackageName, RedirectedToPackageId))
-			{
-				PackageId = RedirectedToPackageId;
-			}
-		}
-
-		bool bIsAlreadyInSet = false;
-		Visitied.Add(PackageId, &bIsAlreadyInSet);
-		if (bIsAlreadyInSet)
-		{
-			continue;
-		}
-
-		FPackageStoreEntry PackageStoreEntry;
-		TSharedPtr<IPackageStoreBackend> Backend;
-		const EPackageStoreEntryStatus EntryStatus = PackageStore.GetPackageStoreEntry(PackageId, NAME_None, PackageStoreEntry, &Backend);
-
-		const bool bContinue = Visitor(PackageId, EntryStatus, PackageStoreEntry, Backend);
-		if (bContinue == false)
-		{
-			break;
-		}
-
-		if (EntryStatus != EPackageStoreEntryStatus::Missing)
-		{
-			for (const FPackageId& ImportedPackageId : PackageStoreEntry.ImportedPackageIds)
-			{
-				if (!Visitied.Contains(ImportedPackageId))
-				{
-					Queue.Add(ImportedPackageId);
-				}
-			}
-
-			if (bIncludeSoftReferences)
-			{
-				TConstArrayView<FPackageId> SoftReferences;
-				TConstArrayView<uint32> Indices = PackageStore.GetSoftReferences(PackageId, SoftReferences);
-				for (uint32 Idx : Indices)
-				{
-					const FPackageId& SoftRef = SoftReferences[Idx];
-					if (!Visitied.Contains(SoftRef))
-					{
-						Queue.Add(SoftRef);
-					}
-				}
-			}
-		}
-	}
-}
+/** Returns true to continue visiting dependencies */
+using FPackageDependencyVisitor = TFunctionRef<
+	bool(FPackageId /*PackageId*/,
+		EPackageStoreEntryStatus /*EntryStatus*/,
+		const FPackageStoreEntry& /*PackageStoreEntry*/,
+		const TSharedPtr<IPackageStoreBackend>& /*Backend*/
+	)>;
+
+void VisitPackageDependencies(
+	const TSet<FPackageId>& PackageIds, 
+	bool bIncludeSoftReferences, 
+	FPackageDependencyVisitor Visitor, 
+	FPackageStoreReadScope* InReadScope = nullptr);
 
 ////////////////////////////////////////////////////////////////////////////////
 void ResolvePackageDependencies(
diff --git a/Engine/Source/Runtime/Experimental/IoStore/OnDemand/Private/OnDemandInstallCache.cpp b/Engine/Source/Runtime/Experimental/IoStore/OnDemand/Private/OnDemandInstallCache.cpp
--- a/Engine/Source/Runtime/Experimental/IoStore/OnDemand/Private/OnDemandInstallCache.cpp
+++ b/Engine/Source/Runtime/Experimental/IoStore/OnDemand/Private/OnDemandInstallCache.cpp
@@ -176,6 +176,473 @@
 };
 
 ///////////////////////////////////////////////////////////////////////////////
+struct FCasJournal
+{
+	enum class EVersion : uint32
+	{
+		Invalid = 0,
+		Initial,
+
+		LatestPlusOne,
+		Latest = LatestPlusOne - 1
+	};
+
+	enum class EErrorCode : uint32
+	{
+		None = 0,
+		Simulated = 1,
+		DefragOutOfDiskSpace = 2,
+		DefragHashMismatch = 3
+	};
+
+	struct FHeader
+	{
+		static const inline uint8 MagicSequence[16] = { 'C', 'A', 'S', 'J', 'O', 'U', 'R', 'N', 'A', 'L', 'H', 'E', 'A', 'D', 'E', 'R' };
+
+		bool			IsValid() const;
+		static int64	Size() { return sizeof(FHeader); }
+
+		uint8		Magic[16] = { 0 };
+		EVersion	Version = EVersion::Invalid;
+		uint8		Pad[12] = { 0 };
+	};
+	static_assert(sizeof(FHeader) == 32);
+
+	struct FFooter
+	{
+		static const inline uint8 MagicSequence[16] = { 'C', 'A', 'S', 'J', 'O', 'U', 'R', 'N', 'A', 'L', 'F', 'O', 'O', 'T', 'E', 'R' };
+
+		bool			IsValid() const;
+		static int64	Size() { return sizeof(FFooter); }
+
+		uint8 Magic[16] = { 0 };
+	};
+	static_assert(sizeof(FFooter) == 16);
+
+	struct FEntry
+	{
+		enum class EType : uint8
+		{
+			None = 0,
+			ChunkLocation,
+			BlockCreated,
+			BlockDeleted,
+			BlockAccess,
+			CriticalError
+		};
+
+		struct FChunkLocation
+		{
+			EType			Type = EType::ChunkLocation;
+			uint8			Pad[3] = { 0 };
+			FCasLocation	CasLocation;
+			FCasAddr		CasAddr;
+		};
+		static_assert(sizeof(FChunkLocation) == 24);
+
+		struct FBlockOperation
+		{
+			EType		Type = EType::None;
+			uint8		Pad[3] = { 0 };
+			FCasBlockId	BlockId;
+			int64		UtcTicks = 0;
+			uint8		Pad1[8] = { 0 };
+		};
+		static_assert(sizeof(FBlockOperation) == 24);
+
+		struct FCriticalError
+		{
+			EType		Type = EType::CriticalError;
+			EErrorCode	ErrorCode = EErrorCode::None;
+		};
+		static_assert(sizeof(FBlockOperation) == 24);
+
+		union
+		{
+			FChunkLocation	ChunkLocation;
+			FBlockOperation	BlockOperation;
+			FCriticalError	CriticalError;
+		};
+
+		EType			Type() const { return *reinterpret_cast<const EType*>(this); }
+		static int64	Size() { return sizeof(FEntry); }
+	};
+	static_assert(sizeof(FEntry) == 24);
+
+	struct FTransaction
+	{
+		void			ChunkLocation(const FCasLocation& Location, const FCasAddr& Addr);
+		void			BlockCreated(FCasBlockId BlockId);
+		void			BlockDeleted(FCasBlockId BlockId);
+		void			BlockAccess(FCasBlockId BlockId, int64 UtcTicks);
+		void			CriticalError(FCasJournal::EErrorCode ErrorCode);
+
+		FString			JournalFile;
+		TArray<FEntry>	Entries;
+	};
+
+	using FEntryHandler = TFunction<void(const FEntry&)>;
+
+	static FResult			Replay(const FString& JournalFile, FEntryHandler&& Handler);
+	static FResult			Create(const FString& JournalFile);
+	static FTransaction		Begin(FString&& JournalFile);
+	static FTransaction		Begin(const FString& JournalFile) { return Begin(FString(JournalFile)); }
+	static FResult			Commit(FTransaction&& Transaction);
+	static FResult			Commit(FTransaction&& Transaction, uint64& OutByteCount, uint32& OutOpCount);
+};
+
+///////////////////////////////////////////////////////////////////////////////
+static const TCHAR* GetErrorText(FCasJournal::EErrorCode ErrorCode)
+{
+	switch (ErrorCode)
+	{
+	case FCasJournal::EErrorCode::None:
+		return TEXT("None");
+	case FCasJournal::EErrorCode::Simulated:
+		return TEXT("Simulated error");
+	case FCasJournal::EErrorCode::DefragOutOfDiskSpace:
+		return TEXT("Defrag failed due to out of disk space");
+	case FCasJournal::EErrorCode::DefragHashMismatch:
+		return TEXT("Found corrupt chunk while defragging");
+	}
+
+	return TEXT("Unknown");
+}
+
+///////////////////////////////////////////////////////////////////////////////
+bool FCasJournal::FHeader::IsValid() const
+{
+	if (FMemory::Memcmp(&Magic, &FHeader::MagicSequence, sizeof(FHeader::MagicSequence)) != 0)
+	{
+		return false;
+	}
+
+	if (static_cast<uint32>(Version) > static_cast<uint32>(EVersion::Latest))
+	{
+		return false;
+	}
+
+	return true;
+}
+
+bool FCasJournal::FFooter::IsValid() const
+{
+	return FMemory::Memcmp(Magic, FFooter::MagicSequence, sizeof(FFooter::MagicSequence)) == 0;
+}
+
+FResult FCasJournal::Replay(const FString& JournalFile, FEntryHandler&& Handler)
+{
+	using namespace UE::UnifiedError::IoStoreOnDemand;
+
+	IPlatformFile& Ipf = FPlatformFileManager::Get().GetPlatformFile();
+
+	if (Ipf.FileExists(*JournalFile) == false)
+	{
+		return MakeJournalError(ECasErrorCode::ReplayJournalFailed, EIoErrorCode::NotFound, FString::Printf(TEXT("Failed to find '%s'"), *JournalFile));
+	}
+
+	TUniquePtr<IFileHandle> FileHandle(Ipf.OpenRead(*JournalFile));
+	if (FileHandle.IsValid() == false)
+	{
+		return MakeJournalError(ECasErrorCode::ReplayJournalFailed, EIoErrorCode::FileOpenFailed, FString::Printf(TEXT("Failed to open '%s'"), *JournalFile));
+	}
+
+	FHeader Header;
+	if ((FileHandle->Read(reinterpret_cast<uint8*>(&Header), FHeader::Size()) == false) || (Header.IsValid() == false))
+	{
+		return MakeJournalError(ECasErrorCode::ReplayJournalFailed, EIoErrorCode::ReadError, FString::Printf(TEXT("Failed to validate journal header in '%s'"), *JournalFile));
+	}
+
+	const int64 FileSize = FileHandle->Size();
+	const int64 EntryCount = (FileSize - FHeader::Size() - FFooter::Size()) / FEntry::Size();
+
+	if (EntryCount < 0)
+	{
+		return MakeJournalError(ECasErrorCode::ReplayJournalFailed, EIoErrorCode::FileCorrupt, FString::Printf(TEXT("Invalid journal file '%s'"), *JournalFile));
+	}
+
+	if (EntryCount == 0)
+	{
+		return MakeValue();
+	}
+
+	const int64 FooterPos = FileSize - FFooter::Size();
+	if (FooterPos < 0)
+	{
+		return MakeJournalError(ECasErrorCode::ReplayJournalFailed, EIoErrorCode::FileCorrupt, FString::Printf(TEXT("Invalid journal footer in '%s'"), *JournalFile));
+	}
+
+	const int64 EntriesPos = FileHandle->Tell();
+	if (FileHandle->Seek(FooterPos) == false)
+	{
+		return MakeJournalError(ECasErrorCode::ReplayJournalFailed, EIoErrorCode::FileSeekFailed, FString::Printf(TEXT("Failed to seek to footer offset %" INT64_FMT " in '%s'"), FooterPos, *JournalFile));
+	}
+
+	FFooter Footer;
+	if ((FileHandle->Read(reinterpret_cast<uint8*>(&Footer), FFooter::Size()) == false) || (Footer.IsValid() == false))
+	{
+		return MakeJournalError(ECasErrorCode::ReplayJournalFailed, EIoErrorCode::ReadError, FString::Printf(TEXT("Failed to validate journal footer in '%s'"), *JournalFile));
+	}
+
+	if (FileHandle->Seek(EntriesPos) == false)
+	{
+		return MakeJournalError(ECasErrorCode::ReplayJournalFailed, EIoErrorCode::FileSeekFailed, FString::Printf(TEXT("Failed to seek to entries offset %" INT64_FMT " in '%s'"), EntriesPos, *JournalFile));
+	}
+
+	TArray<FEntry> Entries;
+	Entries.SetNumZeroed(IntCastChecked<int32>(EntryCount));
+
+	if (FileHandle->Read(reinterpret_cast<uint8*>(Entries.GetData()), FEntry::Size() * EntryCount) == false)
+	{
+		return MakeJournalError(ECasErrorCode::ReplayJournalFailed, EIoErrorCode::ReadError, FString::Printf(TEXT("Failed to read journal entries in '%s'"), *JournalFile));
+	}
+
+	UE_LOG(LogIoStoreOnDemand, Log, TEXT("Replaying %" INT64_FMT " CAS journal entries of total %.2lf KiB from '%s'"),
+		EntryCount, ToKiB(FEntry::Size() * EntryCount), *JournalFile);
+
+	for (const FEntry& Entry : Entries)
+	{
+		if (Entry.Type() == FEntry::EType::CriticalError)
+		{
+			const FEntry::FCriticalError& Error = Entry.CriticalError;
+			UE_LOG(LogIoStoreOnDemand, Warning, TEXT("Found critical error entry '%s' (%d) in journal '%s'"),
+				GetErrorText(Error.ErrorCode), Error.ErrorCode, *JournalFile);
+
+			// We append "critical error" entries to the journal when we endup in an unrecoverable error state. This will cause the cache to be reset
+			// at startup
+			return MakeJournalError(ECasErrorCode::ReplayJournalFailed, EIoErrorCode::InvalidCode, FString::Printf(TEXT("Found critical error journal entry in '%s'"), *JournalFile));
+		}
+
+		Handler(Entry);
+	}
+
+	return MakeValue();
+}
+
+FResult FCasJournal::Create(const FString& JournalFile)
+{
+	using namespace UE::UnifiedError::IoStoreOnDemand;
+
+	IPlatformFile& Ipf = FPlatformFileManager::Get().GetPlatformFile();
+	Ipf.DeleteFile(*JournalFile);
+
+	TUniquePtr<IFileHandle> FileHandle(Ipf.OpenWrite(*JournalFile));
+	if (FileHandle.IsValid() == false)
+	{
+		return MakeJournalError(ECasErrorCode::CreateJournalFailed, EIoErrorCode::FileOpenFailed, FString::Printf(TEXT("Failed to create journal '%s'"), *JournalFile));
+	}
+
+	FHeader Header;
+	FMemory::Memcpy(&Header.Magic, &FHeader::MagicSequence, sizeof(FHeader::MagicSequence));
+	Header.Version = EVersion::Latest;
+
+	if (FileHandle->Write(reinterpret_cast<uint8*>(&Header), FHeader::Size()) == false)
+	{
+		return MakeJournalError(ECasErrorCode::CreateJournalFailed, EIoErrorCode::WriteError, FString::Printf(TEXT("Failed to write journal header in '%s'"), *JournalFile));
+	}
+
+	FFooter Footer;
+	FMemory::Memcpy(&Footer.Magic, &FFooter::MagicSequence, sizeof(FFooter::MagicSequence));
+	if (FileHandle->Write(reinterpret_cast<uint8*>(&Footer), FFooter::Size()) == false)
+	{
+		return MakeJournalError(ECasErrorCode::CreateJournalFailed, EIoErrorCode::WriteError, FString::Printf(TEXT("Failed to write journal footer in '%s'"), *JournalFile));
+	}
+
+	return MakeValue();
+}
+
+FCasJournal::FTransaction FCasJournal::Begin(FString&& JournalFile)
+{
+	return FTransaction
+	{
+		.JournalFile = MoveTemp(JournalFile)
+	};
+}
+
+FResult FCasJournal::Commit(FTransaction&& Transaction, uint64& OutByteCount, uint32& OutOpCount)
+{
+	using namespace UE::UnifiedError::IoStoreOnDemand;
+
+	OutByteCount = 0;
+	OutOpCount = 0;
+
+	if (Transaction.Entries.IsEmpty())
+	{
+		return MakeValue();
+	}
+
+	IPlatformFile& Ipf = FPlatformFileManager::Get().GetPlatformFile();
+
+	// Validate header and footer
+	{
+		TUniquePtr<IFileHandle> FileHandle(Ipf.OpenRead(*Transaction.JournalFile));
+		if (FileHandle.IsValid() == false)
+		{
+			FResult Result = MakeJournalError(
+				ECasErrorCode::CommitJournalFailed, EIoErrorCode::FileOpenFailed, FString::Printf(TEXT("Failed to open journal '%s'"), *Transaction.JournalFile));
+			FOnDemandInstallCacheStats::OnJournalCommit(Result, 0);
+			return Result;
+		}
+
+		const int64 FileSize = FileHandle->Size();
+		if (FileSize < FHeader::Size())
+		{
+			FResult Result = MakeJournalError(
+				ECasErrorCode::CommitJournalFailed, EIoErrorCode::FileCorrupt, FString::Printf(TEXT("Failed to validate journal header in '%s'"), *Transaction.JournalFile));
+			FOnDemandInstallCacheStats::OnJournalCommit(Result, 0);
+			return Result;
+		}
+
+		FHeader Header;
+		if ((FileHandle->Read(reinterpret_cast<uint8*>(&Header), FHeader::Size()) == false) || (Header.IsValid() == false))
+		{
+			FResult Result = MakeJournalError(
+				ECasErrorCode::CommitJournalFailed, EIoErrorCode::ReadError, FString::Printf(TEXT("Failed to validate journal header in '%s'"), *Transaction.JournalFile));
+			FOnDemandInstallCacheStats::OnJournalCommit(Result, 0);
+			return Result;
+		}
+
+		const int64 FooterPos = FileSize - FFooter::Size();
+		if (FileHandle->Seek(FooterPos) == false)
+		{
+			FResult Result = MakeJournalError(
+				ECasErrorCode::CommitJournalFailed, EIoErrorCode::FileSeekFailed, FString::Printf(TEXT("Failed to validate journal footer in '%s'"), *Transaction.JournalFile));
+			FOnDemandInstallCacheStats::OnJournalCommit(Result, 0);
+			return Result;
+		}
+
+		FFooter Footer;
+		if ((FileHandle->Read(reinterpret_cast<uint8*>(&Footer), FFooter::Size()) == false) || (Footer.IsValid() == false))
+		{
+			FResult Result = MakeJournalError(
+				ECasErrorCode::CommitJournalFailed, EIoErrorCode::ReadError, FString::Printf(TEXT("Found to validate journal footer in '%s'"), *Transaction.JournalFile));
+			FOnDemandInstallCacheStats::OnJournalCommit(Result, 0);
+			return Result;
+		}
+	}
+
+	// Append entries
+	{
+		const bool				bAppend = true;
+		TUniquePtr<IFileHandle> FileHandle(Ipf.OpenWrite(*Transaction.JournalFile, bAppend));
+		const int64				FileSize = FileHandle.IsValid() ? FileHandle->Size() : -1;
+		const int64				EntriesPos = FileSize > 0 ? FileSize - FFooter::Size() : -1;
+
+		if (EntriesPos < 0)
+		{
+			FResult Result = MakeJournalError(
+				ECasErrorCode::CommitJournalFailed, EIoErrorCode::FileOpenFailed, FString::Printf(TEXT("Failed to open journal '%s'"), *Transaction.JournalFile));
+			FOnDemandInstallCacheStats::OnJournalCommit(Result, 0);
+			return Result;
+		}
+
+		if (FileHandle->Seek(EntriesPos) == false)
+		{
+			FResult Result = MakeJournalError(
+				ECasErrorCode::CommitJournalFailed, EIoErrorCode::FileSeekFailed, FString::Printf(TEXT("Failed to seek to journal entries offset %" INT64_FMT " in '%s'"), EntriesPos, *Transaction.JournalFile));
+			FOnDemandInstallCacheStats::OnJournalCommit(Result, 0);
+			return Result;
+		}
+
+		const int64 TotalEntrySize = Transaction.Entries.Num() * FEntry::Size();
+		if (FileHandle->Write(
+			reinterpret_cast<const uint8*>(Transaction.Entries.GetData()),
+			TotalEntrySize) == false)
+		{
+			FResult Result = MakeJournalError(
+				ECasErrorCode::CommitJournalFailed, EIoErrorCode::WriteError, FString::Printf(TEXT("Failed to write journal entries to '%s'"), *Transaction.JournalFile));
+			FOnDemandInstallCacheStats::OnJournalCommit(Result, 0);
+			return Result;
+		}
+		OutOpCount++;
+		OutByteCount += TotalEntrySize;
+
+		FFooter Footer;
+		FMemory::Memcpy(&Footer.Magic, &FFooter::MagicSequence, sizeof(FFooter::MagicSequence));
+		if (FileHandle->Write(reinterpret_cast<uint8*>(&Footer), FFooter::Size()) == false)
+		{
+			FResult Result = MakeJournalError(
+				ECasErrorCode::CommitJournalFailed, EIoErrorCode::WriteError, FString::Printf(TEXT("Failed to write journal footer to '%s'"), *Transaction.JournalFile));
+			FOnDemandInstallCacheStats::OnJournalCommit(Result, 0);
+			return Result;
+		}
+		OutOpCount++;
+		OutByteCount += uint64(FFooter::Size());
+
+		if (FileHandle->Flush() == false)
+		{
+			FResult Result = MakeJournalError(
+				ECasErrorCode::CommitJournalFailed, EIoErrorCode::FileFlushFailed, FString::Printf(TEXT("Failed to flush journal entries to '%s'"), *Transaction.JournalFile));
+			FOnDemandInstallCacheStats::OnJournalCommit(Result, 0);
+			return Result;
+		}
+		OutOpCount++;
+
+		UE_LOG(LogIoStoreOnDemand, Log, TEXT("Committed %d CAS journal entries of total %.2lf KiB to '%s'"),
+			Transaction.Entries.Num(), ToKiB(TotalEntrySize), *Transaction.JournalFile);
+
+		FResult Result = MakeValue();
+		FOnDemandInstallCacheStats::OnJournalCommit(Result, TotalEntrySize);
+		return Result;
+	}
+}
+
+FResult FCasJournal::Commit(FTransaction&& Transaction)
+{
+	uint64 ByteCount = 0;
+	uint32 OpCount = 0;
+	return Commit(MoveTemp(Transaction), ByteCount, OpCount);
+}
+
+void FCasJournal::FTransaction::ChunkLocation(const FCasLocation& Location, const FCasAddr& Addr)
+{
+	Entries.AddZeroed_GetRef().ChunkLocation = FEntry::FChunkLocation
+	{
+		.CasLocation = Location,
+		.CasAddr = Addr
+	};
+}
+
+void FCasJournal::FTransaction::BlockCreated(FCasBlockId BlockId)
+{
+	Entries.AddZeroed_GetRef().BlockOperation = FEntry::FBlockOperation
+	{
+		.Type = FEntry::EType::BlockCreated,
+		.BlockId = BlockId,
+		.UtcTicks = FDateTime::UtcNow().GetTicks()
+	};
+}
+
+void FCasJournal::FTransaction::BlockDeleted(FCasBlockId BlockId)
+{
+	Entries.AddZeroed_GetRef().BlockOperation = FEntry::FBlockOperation
+	{
+		.Type = FEntry::EType::BlockDeleted,
+		.BlockId = BlockId,
+		.UtcTicks = FDateTime::UtcNow().GetTicks()
+	};
+}
+
+void FCasJournal::FTransaction::BlockAccess(FCasBlockId BlockId, int64 UtcTicks)
+{
+	Entries.AddZeroed_GetRef().BlockOperation = FEntry::FBlockOperation
+	{
+		.Type = FEntry::EType::BlockAccess,
+		.BlockId = BlockId,
+		.UtcTicks = UtcTicks
+	};
+}
+
+void FCasJournal::FTransaction::CriticalError(FCasJournal::EErrorCode ErrorCode)
+{
+	Entries.AddZeroed_GetRef().CriticalError = FEntry::FCriticalError
+	{
+		.Type = FEntry::EType::CriticalError,
+		.ErrorCode = ErrorCode
+	};
+}
+
+///////////////////////////////////////////////////////////////////////////////
 struct FCasSnapshot;
 struct FCas
 {
@@ -211,6 +678,7 @@
 	void						Compact();
 	FResult						Verify(TArray<FCasAddr>& OutAddrs);
 	void						LoadSnapshot(FCasSnapshot&& Snapshot);
+	void						ApplyJournalEntry(const FCasJournal::FEntry& JournalEntry);
 	FLastAccess					ConsumeLastAcccess()
 	{
 		FLastAccess Temp;
@@ -244,8 +712,8 @@
 	FStringView				RootDirectory;
 public: // TODO: FIXME: encapsulate these members
 	FLookup					Lookup;
+private:
 	FBlockIdHandleCounts	BlockIds;
-private:
 	FLastAccess				LastAccess;
 	FReadHandles			ReadHandles;
 	FEventRef				BlockReadsDoneEvent;
@@ -792,122 +1260,6 @@
 }
 
 ///////////////////////////////////////////////////////////////////////////////
-struct FCasJournal
-{
-	enum class EVersion : uint32
-	{
-		Invalid	= 0,
-		Initial,
-
-		LatestPlusOne,
-		Latest = LatestPlusOne - 1
-	};
-
-	enum class EErrorCode : uint32
-	{
-		None					= 0,
-		Simulated				= 1,
-		DefragOutOfDiskSpace	= 2,
-		DefragHashMismatch		= 3
-	};
-
-	struct FHeader
-	{
-		static const inline uint8 MagicSequence[16] = {'C', 'A', 'S', 'J', 'O', 'U', 'R', 'N', 'A', 'L', 'H', 'E', 'A', 'D', 'E', 'R'};
-
-		bool			IsValid() const;
-		static int64	Size() { return sizeof(FHeader); }
-
-		uint8		Magic[16] = {0};
-		EVersion	Version = EVersion::Invalid;
-		uint8		Pad[12] = {0};
-	};
-	static_assert(sizeof(FHeader) == 32);
-
-	struct FFooter
-	{
-		static const inline uint8 MagicSequence[16] = {'C', 'A', 'S', 'J', 'O', 'U', 'R', 'N', 'A', 'L', 'F', 'O', 'O', 'T', 'E', 'R'};
-
-		bool			IsValid() const;
-		static int64	Size() { return sizeof(FFooter); }
-
-		uint8 Magic[16] = {0};
-	};
-	static_assert(sizeof(FFooter) == 16);
-
-	struct FEntry
-	{
-		enum class EType : uint8
-		{
-			None = 0,
-			ChunkLocation,
-			BlockCreated,
-			BlockDeleted,
-			BlockAccess,
-			CriticalError
-		};
-
-		struct FChunkLocation
-		{
-			EType			Type = EType::ChunkLocation;
-			uint8			Pad[3]= {0};
-			FCasLocation	CasLocation;
-			FCasAddr		CasAddr;
-		};
-		static_assert(sizeof(FChunkLocation) == 24);
-
-		struct FBlockOperation
-		{
-			EType		Type = EType::None;
-			uint8		Pad[3]= {0};
-			FCasBlockId	BlockId;
-			int64		UtcTicks = 0;
-			uint8		Pad1[8]= {0};
-		};
-		static_assert(sizeof(FBlockOperation) == 24);
-
-		struct FCriticalError
-		{
-			EType		Type = EType::CriticalError;
-			EErrorCode	ErrorCode = EErrorCode::None;
-		};
-		static_assert(sizeof(FBlockOperation) == 24);
-
-		union
-		{
-			FChunkLocation	ChunkLocation;
-			FBlockOperation	BlockOperation;
-			FCriticalError	CriticalError;
-		};
-
-		EType			Type() const { return *reinterpret_cast<const EType*>(this); }
-		static int64	Size() { return sizeof(FEntry); }
-	};
-	static_assert(sizeof(FEntry) == 24);
-
-	struct FTransaction
-	{
-		void			ChunkLocation(const FCasLocation& Location, const FCasAddr& Addr);
-		void			BlockCreated(FCasBlockId BlockId);
-		void			BlockDeleted(FCasBlockId BlockId);
-		void			BlockAccess(FCasBlockId BlockId, int64 UtcTicks);
-		void			CriticalError(FCasJournal::EErrorCode ErrorCode);
-
-		FString			JournalFile;
-		TArray<FEntry>	Entries;
-	};
-
-	using FEntryHandler		= TFunction<void(const FEntry&)>;
-
-	static FResult			Replay(const FString& JournalFile, FEntryHandler&& Handler);
-	static FResult			Create(const FString& JournalFile);
-	static FTransaction		Begin(FString&& JournalFile);
-	static FTransaction		Begin(const FString& JournalFile) { return Begin(FString(JournalFile)); }
-	static FResult			Commit(FTransaction&& Transaction);
-	static FResult			Commit(FTransaction&& Transaction, uint64& OutByteCount, uint32& OutOpCount);
-};
-
-///////////////////////////////////////////////////////////////////////////////
 static void JournalLastAccess(FCasJournal::FTransaction& Transaction, const FCas::FLastAccess& LastAccess)
 {
 	for (const TPair<FCasBlockId, int64>& Kv : LastAccess)
@@ -921,354 +1273,3 @@
 }
 
 ///////////////////////////////////////////////////////////////////////////////
-static const TCHAR* GetErrorText(FCasJournal::EErrorCode ErrorCode)
-{
-	switch (ErrorCode)
-	{
-		case FCasJournal::EErrorCode::None:
-			return TEXT("None");
-		case FCasJournal::EErrorCode::Simulated:
-			return TEXT("Simulated error");
-		case FCasJournal::EErrorCode::DefragOutOfDiskSpace:
-			return TEXT("Defrag failed due to out of disk space");
-		case FCasJournal::EErrorCode::DefragHashMismatch:
-			return TEXT("Found corrupt chunk while defragging");
-	}
-
-	return TEXT("Unknown");
-}
-
-///////////////////////////////////////////////////////////////////////////////
-bool FCasJournal::FHeader::IsValid() const
-{
-	if (FMemory::Memcmp(&Magic, &FHeader::MagicSequence, sizeof(FHeader::MagicSequence)) != 0)
-	{
-		return false;
-	}
-
-	if (static_cast<uint32>(Version) > static_cast<uint32>(EVersion::Latest))
-	{
-		return false;
-	}
-
-	return true;
-}
-
-bool FCasJournal::FFooter::IsValid() const
-{
-	return FMemory::Memcmp(Magic, FFooter::MagicSequence, sizeof(FFooter::MagicSequence)) == 0;
-}
-
-FResult FCasJournal::Replay(const FString& JournalFile, FEntryHandler&& Handler)
-{
-	using namespace UE::UnifiedError::IoStoreOnDemand;
-
-	IPlatformFile& Ipf = FPlatformFileManager::Get().GetPlatformFile();
-
-	if (Ipf.FileExists(*JournalFile) == false)
-	{
-		return MakeJournalError(ECasErrorCode::ReplayJournalFailed, EIoErrorCode::NotFound, FString::Printf(TEXT("Failed to find '%s'"), *JournalFile));
-	}
-
-	TUniquePtr<IFileHandle> FileHandle(Ipf.OpenRead(*JournalFile));
-	if (FileHandle.IsValid() == false)
-	{
-		return MakeJournalError(ECasErrorCode::ReplayJournalFailed, EIoErrorCode::FileOpenFailed, FString::Printf(TEXT("Failed to open '%s'"), *JournalFile));
-	}
-
-	FHeader Header;
-	if ((FileHandle->Read(reinterpret_cast<uint8*>(&Header), FHeader::Size()) == false) || (Header.IsValid() == false))
-	{
-		return MakeJournalError(ECasErrorCode::ReplayJournalFailed, EIoErrorCode::ReadError, FString::Printf(TEXT("Failed to validate journal header in '%s'"), *JournalFile));
-	}
-
-	const int64 FileSize	= FileHandle->Size();
-	const int64 EntryCount	= (FileSize - FHeader::Size() - FFooter::Size()) / FEntry::Size();
-
-	if (EntryCount < 0)
-	{
-		return MakeJournalError(ECasErrorCode::ReplayJournalFailed, EIoErrorCode::FileCorrupt, FString::Printf(TEXT("Invalid journal file '%s'"), *JournalFile));
-	}
-
-	if (EntryCount == 0)
-	{
-		return MakeValue();
-	}
-
-	const int64 FooterPos = FileSize - FFooter::Size();
-	if (FooterPos < 0)
-	{
-		return MakeJournalError(ECasErrorCode::ReplayJournalFailed, EIoErrorCode::FileCorrupt, FString::Printf(TEXT("Invalid journal footer in '%s'"), *JournalFile));
-	}
-
-	const int64 EntriesPos = FileHandle->Tell();
-	if (FileHandle->Seek(FooterPos) == false)
-	{
-		return MakeJournalError(ECasErrorCode::ReplayJournalFailed, EIoErrorCode::FileSeekFailed, FString::Printf(TEXT("Failed to seek to footer offset %" INT64_FMT " in '%s'"), FooterPos, *JournalFile));
-	}
-
-	FFooter Footer;
-	if ((FileHandle->Read(reinterpret_cast<uint8*>(&Footer), FFooter::Size()) == false) || (Footer.IsValid() == false))
-	{
-		return MakeJournalError(ECasErrorCode::ReplayJournalFailed, EIoErrorCode::ReadError, FString::Printf(TEXT("Failed to validate journal footer in '%s'"), *JournalFile));
-	}
-
-	if (FileHandle->Seek(EntriesPos) == false)
-	{
-		return MakeJournalError(ECasErrorCode::ReplayJournalFailed, EIoErrorCode::FileSeekFailed, FString::Printf(TEXT("Failed to seek to entries offset %" INT64_FMT " in '%s'"), EntriesPos, *JournalFile));
-	}
-
-	TArray<FEntry> Entries;
-	Entries.SetNumZeroed(IntCastChecked<int32>(EntryCount));
-
-	if (FileHandle->Read(reinterpret_cast<uint8*>(Entries.GetData()), FEntry::Size() * EntryCount) == false)
-	{
-		return MakeJournalError(ECasErrorCode::ReplayJournalFailed, EIoErrorCode::ReadError, FString::Printf(TEXT("Failed to read journal entries in '%s'"), *JournalFile));
-	}
-
-	UE_LOG(LogIoStoreOnDemand, Log, TEXT("Replaying %" INT64_FMT " CAS journal entries of total %.2lf KiB from '%s'"),
-		EntryCount, ToKiB(FEntry::Size() * EntryCount), *JournalFile);
-
-	for (const FEntry& Entry : Entries)
-	{
-		if (Entry.Type() == FEntry::EType::CriticalError)
-		{
-			const FEntry::FCriticalError& Error = Entry.CriticalError;
-			UE_LOG(LogIoStoreOnDemand, Warning, TEXT("Found critical error entry '%s' (%d) in journal '%s'"),
-				GetErrorText(Error.ErrorCode), Error.ErrorCode, *JournalFile);
-
-			// We append "critical error" entries to the journal when we endup in an unrecoverable error state. This will cause the cache to be reset
-			// at startup
-			return MakeJournalError(ECasErrorCode::ReplayJournalFailed, EIoErrorCode::InvalidCode, FString::Printf(TEXT("Found critical error journal entry in '%s'"), *JournalFile));
-		}
-
-		Handler(Entry);
-	}
-
-	return MakeValue();
-}
-
-FResult FCasJournal::Create(const FString& JournalFile)
-{
-	using namespace UE::UnifiedError::IoStoreOnDemand;
-
-	IPlatformFile& Ipf = FPlatformFileManager::Get().GetPlatformFile();
-	Ipf.DeleteFile(*JournalFile);
-
-	TUniquePtr<IFileHandle> FileHandle(Ipf.OpenWrite(*JournalFile));
-	if (FileHandle.IsValid() == false)
-	{
-		return MakeJournalError(ECasErrorCode::CreateJournalFailed, EIoErrorCode::FileOpenFailed, FString::Printf(TEXT("Failed to create journal '%s'"), *JournalFile));
-	}
-
-	FHeader Header;
-	FMemory::Memcpy(&Header.Magic, &FHeader::MagicSequence, sizeof(FHeader::MagicSequence));
-	Header.Version = EVersion::Latest;
-
-	if (FileHandle->Write(reinterpret_cast<uint8*>(&Header), FHeader::Size()) == false)
-	{
-		return MakeJournalError(ECasErrorCode::CreateJournalFailed, EIoErrorCode::WriteError, FString::Printf(TEXT("Failed to write journal header in '%s'"), *JournalFile));
-	}
-
-	FFooter Footer;
-	FMemory::Memcpy(&Footer.Magic, &FFooter::MagicSequence, sizeof(FFooter::MagicSequence));
-	if (FileHandle->Write(reinterpret_cast<uint8*>(&Footer), FFooter::Size()) == false)
-	{
-		return MakeJournalError(ECasErrorCode::CreateJournalFailed, EIoErrorCode::WriteError, FString::Printf(TEXT("Failed to write journal footer in '%s'"), *JournalFile));
-	}
-
-	return MakeValue();
-}
-
-FCasJournal::FTransaction FCasJournal::Begin(FString&& JournalFile)
-{
-	return FTransaction
-	{
-		.JournalFile = MoveTemp(JournalFile)
-	};
-}
-
-FResult FCasJournal::Commit(FTransaction&& Transaction, uint64& OutByteCount, uint32& OutOpCount)
-{
-	using namespace UE::UnifiedError::IoStoreOnDemand;
-
-	OutByteCount	= 0;
-	OutOpCount		= 0;
-
-	if (Transaction.Entries.IsEmpty())
-	{
-		return MakeValue();
-	}
-
-	IPlatformFile& Ipf = FPlatformFileManager::Get().GetPlatformFile();
-
-	// Validate header and footer
-	{
-		TUniquePtr<IFileHandle> FileHandle(Ipf.OpenRead(*Transaction.JournalFile));
-		if (FileHandle.IsValid() == false)
-		{
-			FResult Result = MakeJournalError(
-				ECasErrorCode::CommitJournalFailed, EIoErrorCode::FileOpenFailed, FString::Printf(TEXT("Failed to open journal '%s'"), *Transaction.JournalFile));
-			FOnDemandInstallCacheStats::OnJournalCommit(Result, 0);
-			return Result;
-		}
-
-		const int64 FileSize = FileHandle->Size();
-		if (FileSize < FHeader::Size())
-		{
-			FResult Result = MakeJournalError(
-				ECasErrorCode::CommitJournalFailed, EIoErrorCode::FileCorrupt, FString::Printf(TEXT("Failed to validate journal header in '%s'"), *Transaction.JournalFile));
-			FOnDemandInstallCacheStats::OnJournalCommit(Result, 0);
-			return Result;
-		}
-
-		FHeader Header;
-		if ((FileHandle->Read(reinterpret_cast<uint8*>(&Header), FHeader::Size()) == false) || (Header.IsValid() == false))
-		{
-			FResult Result = MakeJournalError(
-				ECasErrorCode::CommitJournalFailed, EIoErrorCode::ReadError, FString::Printf(TEXT("Failed to validate journal header in '%s'"), *Transaction.JournalFile));
-			FOnDemandInstallCacheStats::OnJournalCommit(Result, 0);
-			return Result;
-		}
-
-		const int64 FooterPos = FileSize - FFooter::Size();
-		if (FileHandle->Seek(FooterPos) == false)
-		{
-			FResult Result = MakeJournalError(
-				ECasErrorCode::CommitJournalFailed, EIoErrorCode::FileSeekFailed, FString::Printf(TEXT("Failed to validate journal footer in '%s'"), *Transaction.JournalFile));
-			FOnDemandInstallCacheStats::OnJournalCommit(Result, 0);
-			return Result;
-		}
-
-		FFooter Footer;
-		if ((FileHandle->Read(reinterpret_cast<uint8*>(&Footer), FFooter::Size()) == false) || (Footer.IsValid() == false))
-		{
-			FResult Result = MakeJournalError(
-				ECasErrorCode::CommitJournalFailed, EIoErrorCode::ReadError, FString::Printf(TEXT("Found to validate journal footer in '%s'"), *Transaction.JournalFile));
-			FOnDemandInstallCacheStats::OnJournalCommit(Result, 0);
-			return Result;
-		}
-	}
-
-	// Append entries
-	{
-		const bool				bAppend = true;
-		TUniquePtr<IFileHandle> FileHandle(Ipf.OpenWrite(*Transaction.JournalFile, bAppend));
-		const int64				FileSize	= FileHandle.IsValid() ? FileHandle->Size() : -1;
-		const int64				EntriesPos	= FileSize > 0 ? FileSize - FFooter::Size() : -1;
-
-		if (EntriesPos < 0) 
-		{
-			FResult Result = MakeJournalError(
-				ECasErrorCode::CommitJournalFailed, EIoErrorCode::FileOpenFailed, FString::Printf(TEXT("Failed to open journal '%s'"), *Transaction.JournalFile));
-			FOnDemandInstallCacheStats::OnJournalCommit(Result, 0);
-			return Result;
-		}
-
-		if (FileHandle->Seek(EntriesPos) == false)
-		{
-			FResult Result = MakeJournalError(
-				ECasErrorCode::CommitJournalFailed, EIoErrorCode::FileSeekFailed, FString::Printf(TEXT("Failed to seek to journal entries offset %" INT64_FMT " in '%s'"), EntriesPos, *Transaction.JournalFile));
-			FOnDemandInstallCacheStats::OnJournalCommit(Result, 0);
-			return Result;
-		}
-
-		const int64 TotalEntrySize = Transaction.Entries.Num() * FEntry::Size();
-		if (FileHandle->Write(
-			reinterpret_cast<const uint8*>(Transaction.Entries.GetData()),
-			TotalEntrySize) == false)
-		{
-			FResult Result = MakeJournalError(
-				ECasErrorCode::CommitJournalFailed, EIoErrorCode::WriteError, FString::Printf(TEXT("Failed to write journal entries to '%s'"), *Transaction.JournalFile));
-			FOnDemandInstallCacheStats::OnJournalCommit(Result, 0);
-			return Result;
-		}
-		OutOpCount++;
-		OutByteCount += TotalEntrySize;
-
-		FFooter Footer;
-		FMemory::Memcpy(&Footer.Magic, &FFooter::MagicSequence, sizeof(FFooter::MagicSequence));
-		if (FileHandle->Write(reinterpret_cast<uint8*>(&Footer), FFooter::Size()) == false)
-		{
-			FResult Result = MakeJournalError(
-				ECasErrorCode::CommitJournalFailed, EIoErrorCode::WriteError, FString::Printf(TEXT("Failed to write journal footer to '%s'"), *Transaction.JournalFile));
-			FOnDemandInstallCacheStats::OnJournalCommit(Result, 0);
-			return Result;
-		}
-		OutOpCount++;
-		OutByteCount += uint64(FFooter::Size());
-
-		if (FileHandle->Flush() == false)
-		{
-			FResult Result = MakeJournalError(
-				ECasErrorCode::CommitJournalFailed, EIoErrorCode::FileFlushFailed, FString::Printf(TEXT("Failed to flush journal entries to '%s'"), *Transaction.JournalFile));
-			FOnDemandInstallCacheStats::OnJournalCommit(Result, 0);
-			return Result;
-		}
-		OutOpCount++;
-
-		UE_LOG(LogIoStoreOnDemand, Log, TEXT("Committed %d CAS journal entries of total %.2lf KiB to '%s'"),
-			Transaction.Entries.Num(), ToKiB(TotalEntrySize), *Transaction.JournalFile);
-
-		FResult Result = MakeValue();
-		FOnDemandInstallCacheStats::OnJournalCommit(Result, TotalEntrySize);
-		return Result;
-	}
-}
-
-FResult FCasJournal::Commit(FTransaction&& Transaction)
-{
-	uint64 ByteCount	= 0;
-	uint32 OpCount		= 0;
-	return Commit(MoveTemp(Transaction), ByteCount, OpCount);
-}
-
-void FCasJournal::FTransaction::ChunkLocation(const FCasLocation& Location, const FCasAddr& Addr)
-{
-	Entries.AddZeroed_GetRef().ChunkLocation = FEntry::FChunkLocation
-	{
-		.CasLocation	= Location,
-		.CasAddr		= Addr
-	};
-}
-
-void FCasJournal::FTransaction::BlockCreated(FCasBlockId BlockId)
-{
-	Entries.AddZeroed_GetRef().BlockOperation = FEntry::FBlockOperation
-	{
-		.Type		= FEntry::EType::BlockCreated,
-		.BlockId	= BlockId,
-		.UtcTicks	= FDateTime::UtcNow().GetTicks()
-	};
-}
-
-void FCasJournal::FTransaction::BlockDeleted(FCasBlockId BlockId)
-{
-	Entries.AddZeroed_GetRef().BlockOperation = FEntry::FBlockOperation
-	{
-		.Type		= FEntry::EType::BlockDeleted,
-		.BlockId	= BlockId,
-		.UtcTicks	= FDateTime::UtcNow().GetTicks()
-	};
-}
-
-void FCasJournal::FTransaction::BlockAccess(FCasBlockId BlockId, int64 UtcTicks)
-{
-	Entries.AddZeroed_GetRef().BlockOperation = FEntry::FBlockOperation
-	{
-		.Type		= FEntry::EType::BlockAccess,
-		.BlockId	= BlockId,
-		.UtcTicks	= UtcTicks
-	};
-}
-
-void FCasJournal::FTransaction::CriticalError(FCasJournal::EErrorCode ErrorCode)
-{
-	Entries.AddZeroed_GetRef().CriticalError = FEntry::FCriticalError
-	{
-		.Type		= FEntry::EType::CriticalError,
-		.ErrorCode	= ErrorCode
-	};
-}
-
-///////////////////////////////////////////////////////////////////////////////
@@ -1625,5 +1626,47 @@
 }
 
 ///////////////////////////////////////////////////////////////////////////////
+void FCas::ApplyJournalEntry(const FCasJournal::FEntry& JournalEntry)
+{
+	switch (JournalEntry.Type())
+	{
+	case FCasJournal::FEntry::EType::ChunkLocation:
+	{
+		const FCasJournal::FEntry::FChunkLocation& ChunkLocation = JournalEntry.ChunkLocation;
+		if (ChunkLocation.CasLocation.IsValid())
+		{
+			FCasLocation& Loc = Lookup.FindOrAdd(ChunkLocation.CasAddr);
+			Loc = ChunkLocation.CasLocation;
+		}
+		else
+		{
+			Lookup.Remove(ChunkLocation.CasAddr);
+		}
+		break;
+	}
+	case FCasJournal::FEntry::EType::BlockCreated:
+	{
+		const FCasJournal::FEntry::FBlockOperation& Op = JournalEntry.BlockOperation;
+		BlockIds.Add(Op.BlockId, 0);
+		break;
+	}
+	case FCasJournal::FEntry::EType::BlockDeleted:
+	{
+		const FCasJournal::FEntry::FBlockOperation& Op = JournalEntry.BlockOperation;
+		BlockIds.Remove(Op.BlockId);
+		LastAccess.Remove(Op.BlockId);
+		break;
+	}
+	case FCasJournal::FEntry::EType::BlockAccess:
+	{
+		const FCasJournal::FEntry::FBlockOperation& Op = JournalEntry.BlockOperation;
+		constexpr const bool bDirty = false;
+		check(TrackAccessIf(ECasTrackAccessType::Always, Op.BlockId, Op.UtcTicks, bDirty));
+		break;
+	}
+	};
+}
+
+///////////////////////////////////////////////////////////////////////////////
 class FOnDemandInstallCache final
 	: public IOnDemandInstallCache 
@@ -1869,44 +1912,23 @@
 		// Replay the journal 
 		InitResult = FCasJournal::Replay(JournalFile, [this](const FCasJournal::FEntry& JournalEntry)
 		{
+			Cas.ApplyJournalEntry(JournalEntry);
+
 			switch(JournalEntry.Type())
 			{
-			case FCasJournal::FEntry::EType::ChunkLocation:
-			{
-				const FCasJournal::FEntry::FChunkLocation& ChunkLocation = JournalEntry.ChunkLocation;
-				if (ChunkLocation.CasLocation.IsValid())
-				{
-					FCasLocation& Loc = Cas.Lookup.FindOrAdd(ChunkLocation.CasAddr);
-					Loc = ChunkLocation.CasLocation;
-				}
-				else
-				{
-					Cas.Lookup.Remove(ChunkLocation.CasAddr);
-				}
-				break;
-			}
 			case FCasJournal::FEntry::EType::BlockCreated:
 			{
 				const FCasJournal::FEntry::FBlockOperation& Op = JournalEntry.BlockOperation;
 				CurrentBlock = Op.BlockId;
-				Cas.BlockIds.Add(Op.BlockId, 0);
 				break;
 			}
 			case FCasJournal::FEntry::EType::BlockDeleted:
 			{
 				const FCasJournal::FEntry::FBlockOperation& Op = JournalEntry.BlockOperation;
-				Cas.BlockIds.Remove(Op.BlockId);
 				FCasBlockId MaybeCurrentBlock = Op.BlockId;
 				CurrentBlock.compare_exchange_strong(MaybeCurrentBlock, FCasBlockId::Invalid);
 				break;
 			}
-			case FCasJournal::FEntry::EType::BlockAccess:
-			{
-				const FCasJournal::FEntry::FBlockOperation& Op = JournalEntry.BlockOperation;
-				constexpr const bool bDirty = false;
-				check(Cas.TrackAccessIf(ECasTrackAccessType::Always, Op.BlockId, Op.UtcTicks, bDirty));
-				break;
-			}
 			};
 		});
 	}
diff --git a/Engine/Source/Runtime/Experimental/IoStore/OnDemand/Private/OnDemandIoStore.cpp b/Engine/Source/Runtime/Experimental/IoStore/OnDemand/Private/OnDemandIoStore.cpp
--- a/Engine/Source/Runtime/Experimental/IoStore/OnDemand/Private/OnDemandIoStore.cpp
+++ b/Engine/Source/Runtime/Experimental/IoStore/OnDemand/Private/OnDemandIoStore.cpp
@@ -858,7 +858,7 @@ FIoStatus FOnDemandIoStore::GetIsOnDemand(TConstArrayView<FOnDemandIsOnDemandArg
 
 		const bool bIncludeSoftReferences = EnumHasAnyFlags(Args.Options, EOnDemandIsOnDemandOptions::IncludeSoftReferences);
 
-		TVisitPackageDependencies(TSet<FPackageId>(Args.PackageIds), bIncludeSoftReferences,
+		VisitPackageDependencies(TSet<FPackageId>(Args.PackageIds), bIncludeSoftReferences,
 			[PackageVisitor, Index](
 				FPackageId PackageId,
 				EPackageStoreEntryStatus EntryStatus,
