[TraceBasedDebuggers] converted FFNTraceBasedDebuggersModule from a module to a simple helper base class since that class was meant to be abstract (fixing CIS Issue 1149922)
- FSecureSocketTraceRelayTransport is now shareable between all classes inheriting from FFNTraceBasedDebuggersModule
- changed FNTraceBasedDebuggersAuth loading phase to "PostDefault" since it requires the modules deriving from FFNTraceBasedDebuggersModule to be initialized first to create the data relay transport
- made method GetSessionsManager non-static in order to avoid initialization order issue; it now enforces the EngineEditorBridge instance

#Synced-CL 48800164

diff --git a/Engine/Source/Runtime/Experimental/TraceBasedDebuggers/Private/EngineEditorBridge.cpp b/Engine/Source/Runtime/Experimental/TraceBasedDebuggers/Private/EngineEditorBridge.cpp
--- a/Engine/Source/Runtime/Experimental/TraceBasedDebuggers/Private/EngineEditorBridge.cpp
+++ b/Engine/Source/Runtime/Experimental/TraceBasedDebuggers/Private/EngineEditorBridge.cpp
@@ -1,4 +1,4 @@
-﻿// Copyright Epic Games, Inc. All Rights Reserved.
+// Copyright Epic Games, Inc. All Rights Reserved.
 
 #include "EngineEditorBridge.h"
 
@@ -11,9 +11,9 @@ namespace UE::TraceBasedDebuggers
 
 TWeakPtr<FRemoteSessionsManager> FEngineEditorBridge::SharedRemoteSessionsManager;
 
-TSharedPtr<FRemoteSessionsManager> FEngineEditorBridge::GetSessionsManager()
+TSharedPtr<FRemoteSessionsManager> FEngineEditorBridge::GetSessionsManager() const
 {
-	return SharedRemoteSessionsManager.Pin();
+	return RemoteSessionsManager;
 }
 
 void FEngineEditorBridge::HandleRecordingStarted()
diff --git a/Engine/Source/Runtime/Experimental/TraceBasedDebuggers/Public/EngineEditorBridge.h b/Engine/Source/Runtime/Experimental/TraceBasedDebuggers/Public/EngineEditorBridge.h
--- a/Engine/Source/Runtime/Experimental/TraceBasedDebuggers/Public/EngineEditorBridge.h
+++ b/Engine/Source/Runtime/Experimental/TraceBasedDebuggers/Public/EngineEditorBridge.h
@@ -1,4 +1,4 @@
-﻿// Copyright Epic Games, Inc. All Rights Reserved.
+// Copyright Epic Games, Inc. All Rights Reserved.
 
 #pragma once
 
@@ -41,7 +41,7 @@ public:
 	UE_API void Initialize();
 	UE_API void TearDown();
 
-	UE_API static TSharedPtr<FRemoteSessionsManager> GetSessionsManager();
+	UE_API TSharedPtr<FRemoteSessionsManager> GetSessionsManager() const;
 
 	/** 
 	 * Derived classes can use this method to register a remote sessions handler that adds support
