Fix issues when reloading dataflow asset where subgraph notifications will be processed by nodes that are no longer part of it

#Synced-CL 48772056

diff --git a/Engine/Source/Runtime/Experimental/Dataflow/Engine/Private/Dataflow/DataflowSubGraphNodes.cpp b/Engine/Source/Runtime/Experimental/Dataflow/Engine/Private/Dataflow/DataflowSubGraphNodes.cpp
--- a/Engine/Source/Runtime/Experimental/Dataflow/Engine/Private/Dataflow/DataflowSubGraphNodes.cpp
+++ b/Engine/Source/Runtime/Experimental/Dataflow/Engine/Private/Dataflow/DataflowSubGraphNodes.cpp
@@ -462,7 +462,7 @@ void FDataflowCallSubGraphNode::OnSubGraphChanged(const struct FEdGraphEditActio
 
 void FDataflowCallSubGraphNode::OnSomeSubGraphsChanged(const UDataflow* InDataflowAsset, const FGuid& InSubGraphGuid, UE::Dataflow::ESubGraphChangedReason InReason)
 {
-	if (InSubGraphGuid == SubGraphGuid)
+	if (InSubGraphGuid == SubGraphGuid && InDataflowAsset == DataflowAssetWeakPtr)
 	{
 		switch (InReason)
 		{
@@ -489,9 +489,12 @@ void FDataflowCallSubGraphNode::OnSomeSubGraphsChanged(const UDataflow* InDatafl
 
 void FDataflowCallSubGraphNode::OnSubGraphLoaded(const UDataflowSubGraph& InSubGraph)
 {
-	if (InSubGraph.GetSubGraphGuid() == SubGraphGuid)
+	if (InSubGraph.GetOuter() == DataflowAssetWeakPtr.Get())
 	{
-		SetSubGraphGuid(SubGraphGuid);
+		if (InSubGraph.GetSubGraphGuid() == SubGraphGuid)
+		{
+			SetSubGraphGuid(SubGraphGuid);
+		}
 	}
 }
 
