Dataflow : Add GraphTopology Guid to dataflow graph class
- did some cleanup and added a lot of pointer checks as well that were missing

#Synced-CL 48816559

diff --git a/Engine/Source/Runtime/Core/Public/UObject/FortniteMainBranchObjectVersions.inl b/Engine/Source/Runtime/Core/Public/UObject/FortniteMainBranchObjectVersions.inl
--- a/Engine/Source/Runtime/Core/Public/UObject/FortniteMainBranchObjectVersions.inl
+++ b/Engine/Source/Runtime/Core/Public/UObject/FortniteMainBranchObjectVersions.inl
@@ -718,6 +718,8 @@ FORTNITE_MAIN_VERSION(LiveLinkCompressSubSectionData, 234)
 // Added option to quantize chaos tri mesh BVH's per-face bounds
 FORTNITE_MAIN_VERSION(OptionalQuantizePerFaceBoundsInChaosTriMeshBVH, 235)
 
+// Add Dataflow Graph Topology Guid to UE::Dataflow:FGraph
+FORTNITE_MAIN_VERSION(DataflowAddGraphTopologyGuid, 236)
 
 // -----<new versions can be added above this line>-------------------------------------------------
 #undef FORTNITE_MAIN_VERSION
diff --git a/Engine/Source/Runtime/Experimental/Dataflow/Core/Private/Dataflow/DataflowGraph.cpp b/Engine/Source/Runtime/Experimental/Dataflow/Core/Private/Dataflow/DataflowGraph.cpp
--- a/Engine/Source/Runtime/Experimental/Dataflow/Core/Private/Dataflow/DataflowGraph.cpp
+++ b/Engine/Source/Runtime/Experimental/Dataflow/Core/Private/Dataflow/DataflowGraph.cpp
@@ -17,6 +17,31 @@ DEFINE_LOG_CATEGORY_STATIC(DATAFLOW_LOG, Error, All);
 
 namespace UE::Dataflow
 {
+	namespace Private
+	{
+		FString GetConnectionFullName(const FDataflowConnection& Connection)
+		{
+			static const FName InvalidName("[Invalid]");
+
+			const FName NodeName = Connection.GetOwningNode()
+				? Connection.GetOwningNode()->GetName()
+				: InvalidName;
+			const FName ConnectionName = Connection.GetName();
+
+			return FString::Format(TEXT("{0}:{1}"), { NodeName.ToString(), ConnectionName.ToString() });
+		}
+
+		FLink MakeConnectionLink(FDataflowOutput& Output, FDataflowInput& Input)
+		{
+			return FLink(
+				Output.GetOwningNode()->GetGuid(),
+				Output.GetGuid(),
+				Input.GetOwningNode()->GetGuid(),
+				Input.GetGuid()
+			);
+		}
+	}
+
 	TSet<FName> FGraph::RegisteredFilters = {};
 
 	FGraph::FGraph(FGuid InGuid)
@@ -30,6 +55,15 @@ namespace UE::Dataflow
 		FilteredNodes.Reset();
 		Connections.Reset();
 		DisabledNodes.Reset();
+
+		// Empty graphs have a zero topology Guid
+		GraphTopologyGuid = FGuid();
+	}
+
+	void FGraph::OnChangeGraphTopology()
+	{
+		// Make sure we get a new topology Guid 
+		GraphTopologyGuid = FGuid::NewGuid();
 	}
 
 	const TArray<FGraph::FNodeSharedPtr>& FGraph::GetFilteredNodes(const FName& NodeFilter) const
@@ -54,6 +88,7 @@ namespace UE::Dataflow
 					FilteredNodes.FindOrAdd(RegisteredType).Add(NewNode);
 				}
 			}
+			OnChangeGraphTopology();
 		}
 		return NewNode;
 	}
@@ -120,108 +155,79 @@ namespace UE::Dataflow
 
 	void FGraph::RemoveNode(TSharedPtr<FDataflowNode> Node)
 	{
-		for (FDataflowOutput* Output : Node->GetOutputs())
+		if (Node)
 		{
-			if (Output)
+			for (FDataflowOutput* Output : Node->GetOutputs())
+			{
+				ClearConnections(Output);
+			}
+			for (FDataflowInput* Input : Node->GetInputs())
 			{
-				for (FDataflowInput* Input : Output->GetConnectedInputs())
+				ClearConnections(Input);
+			}
+
+			Nodes.Remove(Node);
+			OnChangeGraphTopology();
+
+			for (const FName& RegisteredType : RegisteredFilters)
+			{
+				if (Node->IsA(RegisteredType))
 				{
-					if (Input)
+					if (TArray<FNodeSharedPtr>* FoundNodes = FilteredNodes.Find(RegisteredType))
 					{
-						Disconnect(Output, Input);
+						FoundNodes->Remove(Node);
 					}
 				}
 			}
 		}
-		for (FDataflowInput* Input : Node->GetInputs())
+	}
+
+	void FGraph::ClearConnections(FDataflowConnection* InConnection)
+	{
+		if (InConnection)
 		{
-			if (Input)
+			// Todo(dataflow) : do this without triggering a invalidation. 
+			//            or implement a better sync for the EdGraph and DataflowGraph
+			if (InConnection->GetDirection() == FPin::EDirection::INPUT)
 			{
+				FDataflowInput* Input = static_cast<FDataflowInput*>(InConnection);
 				TArray<FDataflowOutput*> Outputs = Input->GetConnectedOutputs();
 				for (FDataflowOutput* Output : Outputs)
 				{
-					if (Output)
-					{
-						Disconnect(Output, Input);
-					}
+					Disconnect(Output, Input);
 				}
 			}
-		}
-		Nodes.Remove(Node);
-		for(const FName& RegisteredType : RegisteredFilters)
-		{
-			if (Node->IsA(RegisteredType))
+			else if (InConnection->GetDirection() == FPin::EDirection::OUTPUT)
 			{
-				if(TArray<FNodeSharedPtr>* FoundNodes = FilteredNodes.Find(RegisteredType))
+				FDataflowOutput* Output = static_cast<FDataflowOutput*>(InConnection);
+				TArray<FDataflowInput*> Inputs = Output->GetConnectedInputs();
+				for (FDataflowInput* Input : Inputs)
 				{
-					FoundNodes->Remove(Node);
+					Disconnect(Output, Input);
 				}
 			}
 		}
 	}
 
-	void FGraph::ClearConnections(FDataflowConnection* Connection)
+	void FGraph::ClearConnections(FDataflowInput* InConnection)
 	{
-		// Todo(dataflow) : do this without triggering a invalidation. 
-		//            or implement a better sync for the EdGraph and DataflowGraph
-		if (Connection->GetDirection() == FPin::EDirection::INPUT)
-		{
-			FDataflowInput* ConnectionIn = static_cast<FDataflowInput*>(Connection);
-			TArray<FDataflowOutput*> BaseOutputs = ConnectionIn->GetConnectedOutputs();
-			for (FDataflowOutput* Output : BaseOutputs)
-			{
-				Disconnect(Output, ConnectionIn);
-			}
-		}
-		else if (Connection->GetDirection() == FPin::EDirection::OUTPUT)
+		if (InConnection)
 		{
-			FDataflowOutput* ConnectionOut = static_cast<FDataflowOutput*>(Connection);
-			TArray<FDataflowInput*> BaseInputs = ConnectionOut->GetConnectedInputs();
-			for (FDataflowInput* Input : BaseInputs)
+			for (FDataflowOutput* Output : InConnection->GetConnectedOutputs())
 			{
-				Disconnect(ConnectionOut, Input);
+				Disconnect(Output, InConnection);
 			}
 		}
 	}
 
-	void FGraph::ClearConnections(FDataflowInput* InConnection)
-	{
-		for (FDataflowOutput* Output : InConnection->GetConnectedOutputs())
-		{
-			Disconnect(Output, InConnection);
-		}
-	}
-
 	void FGraph::ClearConnections(FDataflowOutput* OutConnection)
 	{
-		for (FDataflowInput* Input : OutConnection->GetConnectedInputs())
-		{
-			Disconnect(OutConnection, Input);
-		}
-	}
-
-	namespace Private
-	{
-		FString GetConnectionFullName(const FDataflowConnection& Connection)
-		{
-			static const FName InvalidName("Invalid");
-
-			const FName NodeName = Connection.GetOwningNode()
-				? Connection.GetOwningNode()->GetName()
-				: InvalidName;
-			const FName ConnectionName = Connection.GetName();
-
-			return FString::Format(TEXT("{0}:{1}"), { NodeName.ToString(), ConnectionName.ToString() });
-		}
-
-		FLink MakeConnectionLink(FDataflowOutput& Output, FDataflowInput& Input)
+		if (OutConnection)
 		{
-			return FLink(
-				Output.GetOwningNode()->GetGuid(),
-				Output.GetGuid(),
-				Input.GetOwningNode()->GetGuid(),
-				Input.GetGuid()
-			);
+			for (FDataflowInput* Input : OutConnection->GetConnectedInputs())
+			{
+				Disconnect(OutConnection, Input);
+			}
 		}
 	}
 
@@ -390,22 +396,27 @@ namespace UE::Dataflow
 			Output.AddConnection(&Input);
 			Input.AddConnection(&Output);
 			Connections.Add(Private::MakeConnectionLink(Output, Input));
+
+			OnChangeGraphTopology();
 		}
 		return true;
 	}
 
 	void FGraph::Disconnect(FDataflowOutput* OutputConnection, FDataflowInput* InputConnection)
 	{
-		UE_LOG(LogChaosDataflow, Verbose, TEXT("FGraph::Disconnect(): Disconnecting output [%s:%s] from input [%s:%s]"),
-			OutputConnection->GetOwningNode() ? *OutputConnection->GetOwningNode()->GetName().ToString() : TEXT("Invalid"),
-			*OutputConnection->GetName().ToString(),
-			InputConnection->GetOwningNode() ? *InputConnection->GetOwningNode()->GetName().ToString() : TEXT("Invalid"),
-			*InputConnection->GetName().ToString());
-		OutputConnection->RemoveConnection(InputConnection);
-		InputConnection->RemoveConnection(OutputConnection);
-		Connections.RemoveSwap(FLink(
-			OutputConnection->GetOwningNode()->GetGuid(), OutputConnection->GetGuid(),
-			InputConnection->GetOwningNode()->GetGuid(), InputConnection->GetGuid()));
+		if (OutputConnection && InputConnection)
+		{
+			UE_LOG(LogChaosDataflow, Verbose, TEXT("FGraph::Disconnect(): Disconnecting output [%s] from input [%s]"),
+				*Private::GetConnectionFullName(*OutputConnection),
+				*Private::GetConnectionFullName(*InputConnection)
+			);
+
+			OutputConnection->RemoveConnection(InputConnection);
+			InputConnection->RemoveConnection(OutputConnection);
+			Connections.RemoveSwap(Private::MakeConnectionLink(*OutputConnection, *InputConnection));
+
+			OnChangeGraphTopology();
+		}
 	}
 
 	void FGraph::AddReferencedObjects(FReferenceCollector& Collector)
@@ -434,10 +445,14 @@ namespace UE::Dataflow
 
 	void FGraph::SerializeForSaving(FArchive& Ar, FGraph* InGraph, TArray<TSharedPtr<FDataflowNode>>& InNodes, TArray<FLink>& InConnections)
 	{
+		check(InGraph);
+
 		FGuid ArGuid;
 		FName ArType, ArName;
 		int32 ArNum = InNodes.Num();
 
+		Ar << InGraph->GraphTopologyGuid;
+
 		Ar << ArNum;
 		for (TSharedPtr<FDataflowNode> Node : InNodes)
 		{
@@ -491,11 +506,14 @@ namespace UE::Dataflow
 
 	void FGraph::SerializeForLoading(FArchive& Ar, FGraph* InGraph, UObject* OwningObject)
 	{
+		check(InGraph);
+
 		InGraph->Reset();
 
 		const bool bDataflowSeparateInputOutputSerialization = (Ar.CustomVer(FUE5MainStreamObjectVersion::GUID) >= FUE5MainStreamObjectVersion::DataflowSeparateInputOutputSerialization);
 		const bool bDataflowAnyTypeSupport = (Ar.CustomVer(FFortniteMainBranchObjectVersion::GUID) >= FFortniteMainBranchObjectVersion::DataflowAnyTypeSupport);
 		const bool bDataflowTemplateTypeFix = (Ar.CustomVer(FFortniteMainBranchObjectVersion::GUID) >= FFortniteMainBranchObjectVersion::DataflowTemplatedTypeFix);
+		const bool bDataflowAddGraphTopologyGuid = (Ar.CustomVer(FFortniteMainBranchObjectVersion::GUID) >= FFortniteMainBranchObjectVersion::DataflowAddGraphTopologyGuid);
 
 		FGuid ArGuid;
 		FName ArType, ArName;
@@ -522,6 +540,15 @@ namespace UE::Dataflow
 				return false;
 			};
 
+		if (bDataflowAddGraphTopologyGuid)
+		{
+			Ar << InGraph->GraphTopologyGuid;
+		}
+		else
+		{
+			InGraph->GraphTopologyGuid = FGuid::NewGuid();
+		}
+
 		Ar << ArNum;
 		for (int32 Ndx = ArNum; Ndx > 0; Ndx--)
 		{
diff --git a/Engine/Source/Runtime/Experimental/Dataflow/Core/Public/Dataflow/DataflowGraph.h b/Engine/Source/Runtime/Experimental/Dataflow/Core/Public/Dataflow/DataflowGraph.h
--- a/Engine/Source/Runtime/Experimental/Dataflow/Core/Public/Dataflow/DataflowGraph.h
+++ b/Engine/Source/Runtime/Experimental/Dataflow/Core/Public/Dataflow/DataflowGraph.h
@@ -87,6 +87,12 @@ namespace UE::Dataflow
 		TMap<FName, TArray<FNodeSharedPtr>> FilteredNodes;
 		TArray<FLink> Connections;
 		TSet<FName> DisabledNodes;
+
+		/** 
+		* This guid changes whenever the topology of the graph changes
+		* (Add/Remove nodes, Connect/Disconnect input and/or outputs)
+		*/
+		FGuid GraphTopologyGuid;
 		
 		/** Node filter type that could be used for fast access*/
 		static TSet<FName> RegisteredFilters;
@@ -193,8 +199,11 @@ namespace UE::Dataflow
 		UE_API static void SerializeForSaving(FArchive& Ar, FGraph* InGraph, TArray<TSharedPtr<FDataflowNode>>& InNodes, TArray<FLink>& InConnections);
 		UE_API static void SerializeForLoading(FArchive& Ar, FGraph* InGraph, UObject* OwningObject);
 
+		const FGuid& GetGraphTopologyGuid() const { return GraphTopologyGuid; }
+
 	private:
 		void Reset();
+		void OnChangeGraphTopology();
 	};
 
 	UE_API void RegisterNodeFilter(const FName& NodeFilter);
