Physics object interface : fix crash when calling SetLinearEtherDrag or SetAngularEtherDrag form the physics thread

#Synced-CL 48816559

diff --git a/Engine/Source/Runtime/Experimental/Chaos/Private/Chaos/PhysicsObjectInterface.cpp b/Engine/Source/Runtime/Experimental/Chaos/Private/Chaos/PhysicsObjectInterface.cpp
--- a/Engine/Source/Runtime/Experimental/Chaos/Private/Chaos/PhysicsObjectInterface.cpp
+++ b/Engine/Source/Runtime/Experimental/Chaos/Private/Chaos/PhysicsObjectInterface.cpp
@@ -1257,26 +1257,38 @@ namespace Chaos
 	template<EThreadContext Id>
 	void Chaos::FWritePhysicsObjectInterface<Id>::SetLinearEtherDrag(TArrayView<const FPhysicsObjectHandle> InObjects, float InLinearDrag)
 	{
-		if (Chaos::FPBDRigidsSolver* RigidSolver = Chaos::FPhysicsObjectInterface::GetSolver(InObjects))
+		if (InObjects.IsEmpty())
 		{
-			//put onto physics thread
-			RigidSolver->EnqueueCommandImmediate([InObjects = TArray<FPhysicsObjectHandle>{ InObjects }, InLinearDrag, RigidSolver]() {
-				for (const FPhysicsObjectHandle Object : InObjects)
-				{
-					if (!Object)
-					{
-						continue;
-					}
+			return;
+		}
 
-					if (TThreadParticle<EThreadContext::Internal>* Particle = Object->GetParticle<EThreadContext::Internal>())
-					{
-						if (Chaos::TThreadRigidParticle<EThreadContext::Internal>* Rigid = Particle->CastToRigidParticle())
-						{
-							Rigid->SetLinearEtherDrag(InLinearDrag);
-						}
-					}
+		for (const FPhysicsObjectHandle Object : InObjects)
+		{
+			if (!Object)
+			{
+				continue;
+			}
+
+			if (TThreadParticle<Id>* Particle = Object->GetParticle<Id>())
+			{
+				if (Chaos::TThreadRigidParticle<Id>* Rigid = Particle->CastToRigidParticle())
+				{
+					Rigid->SetLinearEtherDrag(InLinearDrag);
 				}
-				});
+			}
+		}
+			
+		if constexpr (Id == EThreadContext::External)
+		{
+			if (Chaos::FPBDRigidsSolver* RigidSolver = Chaos::FPhysicsObjectInterface::GetSolver(InObjects))
+			{
+				//put onto physics thread
+				RigidSolver->EnqueueCommandImmediate(
+					[InObjects = TArray<FPhysicsObjectHandle>{ InObjects }, InLinearDrag]() {
+						Chaos::FWritePhysicsObjectInterface_Internal Interface = Chaos::FPhysicsObjectInternalInterface::GetWrite();
+						Interface.SetLinearEtherDrag(InObjects, InLinearDrag);
+					});
+			}
 		}
 	}
 
@@ -1284,26 +1296,38 @@ namespace Chaos
 	template<EThreadContext Id>
 	void Chaos::FWritePhysicsObjectInterface<Id>::SetAngularEtherDrag(TArrayView<const FPhysicsObjectHandle> InObjects, float InAngularDrag)
 	{
-		if (Chaos::FPBDRigidsSolver* RigidSolver = Chaos::FPhysicsObjectInterface::GetSolver(InObjects))
+		if (InObjects.IsEmpty())
 		{
-			//put onto physics thread
-			RigidSolver->EnqueueCommandImmediate([InObjects = TArray<FPhysicsObjectHandle>{ InObjects }, InAngularDrag, RigidSolver]() {
-				for (const FPhysicsObjectHandle Object : InObjects)
-				{
-					if (!Object)
-					{
-						continue;
-					}
+			return;
+		}
 
-					if (TThreadParticle<EThreadContext::Internal>* Particle = Object->GetParticle<EThreadContext::Internal>())
-					{
-						if (Chaos::TThreadRigidParticle<EThreadContext::Internal>* Rigid = Particle->CastToRigidParticle())
-						{
-							Rigid->SetAngularEtherDrag(InAngularDrag);
-						}
-					}
+		for (const FPhysicsObjectHandle Object : InObjects)
+		{
+			if (!Object)
+			{
+				continue;
+			}
+
+			if (TThreadParticle<Id>* Particle = Object->GetParticle<Id>())
+			{
+				if (Chaos::TThreadRigidParticle<Id>* Rigid = Particle->CastToRigidParticle())
+				{
+					Rigid->SetAngularEtherDrag(InAngularDrag);
 				}
-			});
+			}
+		}
+
+		if constexpr (Id == EThreadContext::External)
+		{
+			if (Chaos::FPBDRigidsSolver* RigidSolver = Chaos::FPhysicsObjectInterface::GetSolver(InObjects))
+			{
+				//put onto physics thread
+				RigidSolver->EnqueueCommandImmediate(
+					[InObjects = TArray<FPhysicsObjectHandle>{ InObjects }, InAngularDrag]() {
+						Chaos::FWritePhysicsObjectInterface_Internal Interface = Chaos::FPhysicsObjectInternalInterface::GetWrite();
+						Interface.SetAngularEtherDrag(InObjects, InAngularDrag);
+					});
+			}
 		}
 	}
 
