Log error when HTTP I/O dispatcher fails to shutdown.

#rnx

#Synced-CL 48816662

diff --git a/Engine/Source/Runtime/Experimental/IoStore/OnDemand/Private/OnDemandIoStore.cpp b/Engine/Source/Runtime/Experimental/IoStore/OnDemand/Private/OnDemandIoStore.cpp
--- a/Engine/Source/Runtime/Experimental/IoStore/OnDemand/Private/OnDemandIoStore.cpp
+++ b/Engine/Source/Runtime/Experimental/IoStore/OnDemand/Private/OnDemandIoStore.cpp
@@ -215,7 +215,11 @@ FOnDemandIoStore::~FOnDemandIoStore()
 	if (FHttpIoDispatcher::IsInitialized())
 	{
 		FHttpIoDispatcher::OnHostGroupRegistered().RemoveAll(this);
-		FHttpIoDispatcher::Shutdown();
+		if (FIoStatus Status = FHttpIoDispatcher::Shutdown(); Status.IsOk() == false)
+		{
+			UE_LOG(LogIoStoreOnDemand, Error, TEXT("Failed to shutdown HTTP I/O dispatcher, reason: %s"),
+				*Status.ToString());
+		}
 	}
 }
 
