Changed opaque I/O HTTP request handle to be an explicit reference counted type.
- changed opaque ptr type to IIoHttpRequestHandle that inherits from IRefCountedObject
- updated Cancel/UpdatePriority to take a strong reference. This fixes an issue where
a request handle might be deleted in the callback function.
- updated the batch request API with nodiscard since the batch relies on the caller to keep a reference to the requests
- this change also fixes an issue during shutdown where the global dispatcher is no longer needed for deallocating requests

#rnx

#Synced-CL 48874164

diff --git a/Engine/Source/Runtime/Experimental/IoStore/OnDemand/Private/OnDemandHttpIoDispacher.cpp b/Engine/Source/Runtime/Experimental/IoStore/OnDemand/Private/OnDemandHttpIoDispacher.cpp
--- a/Engine/Source/Runtime/Experimental/IoStore/OnDemand/Private/OnDemandHttpIoDispacher.cpp
+++ b/Engine/Source/Runtime/Experimental/IoStore/OnDemand/Private/OnDemandHttpIoDispacher.cpp
@@ -430,40 +430,55 @@ ENUM_CLASS_FLAGS(EHttpRequestFlags);
 struct FHttpRequest final
 	: TConcurrentLinearObject<FHttpRequest, FHttpRequestBlockAllocationTag>
 	, TIntrusiveListElement<FHttpRequest>
+	, IIoHttpRequestHandle
 {
-	using FFlags				= TAtomicFlags<EHttpRequestFlags>;
-	static uint32				NextSeqNo;
-
-	EHttpRequestType			GetHttpRequestType() const { return Options.GetCategory() == 1 ? EHttpRequestType::Installed : EHttpRequestType::Streaming; }
-	void						AddRef() { RefCount.fetch_add(1, std::memory_order_relaxed); }
-	inline void					Release();
-
-	FIoHash						ChunkHash;
-	FIoHash						CacheKey; //TODO make this smaller the IAS cache only uses 64 bit cache keys
-	FIoHttpRequestCompleted 	OnCompleted;
-	FIoHttpHeaders				Headers;
-	FIoHttpOptions				Options;
-	FIoBuffer					Buffer;
-	FIoRelativeUrl				RelativeUrl;
-	FHttpRequest*				Next = nullptr;
-	FName						HostGroupName;
-	EIoErrorCode				CacheStatus = EIoErrorCode::Unknown;
-	FFlags						Flags;
-	std::atomic<int32>			RefCount{0};
-	std::atomic<EIoErrorCode>	CompletionStatus{EIoErrorCode::Unknown};
-	uint32						SeqNo = ++NextSeqNo;
-
-	TSharedPtr<FCancellationToken> CancellationToken;
+	using FRefCountValue			= FReturnedRefCountValue;
+	using FFlags					= TAtomicFlags<EHttpRequestFlags>;
+	static uint32					NextSeqNo;
+
+	EHttpRequestType				GetHttpRequestType() const		{ return Options.GetCategory() == 1 ? EHttpRequestType::Installed : EHttpRequestType::Streaming; }
+	virtual EIoErrorCode			GetStatus() const override		{ return CompletionStatus.load(std::memory_order_acquire); }
+
+	virtual void					AddRef() const override			{ RefCount.fetch_add(1, std::memory_order_acq_rel); }
+	virtual FRefCountValue			GetRefCount() const override	{ return FRefCountValue(RefCount.load(std::memory_order_relaxed)); }
+	virtual inline FRefCountValue	Release() const override;
+
+	FIoHash							ChunkHash;
+	FIoHash							CacheKey; //TODO make this smaller the IAS cache only uses 64 bit cache keys
+	FIoHttpRequestCompleted 		OnCompleted;
+	FIoHttpHeaders					Headers;
+	FIoHttpOptions					Options;
+	FIoBuffer						Buffer;
+	FIoRelativeUrl					RelativeUrl;
+	FHttpRequest*					Next = nullptr;
+	FName							HostGroupName;
+	EIoErrorCode					CacheStatus = EIoErrorCode::Unknown;
+	FFlags							Flags;
+	mutable std::atomic<int32>		RefCount{0};
+	std::atomic<EIoErrorCode>		CompletionStatus{EIoErrorCode::Unknown};
+	uint32							SeqNo = ++NextSeqNo;
+
+	TSharedPtr<FCancellationToken>	CancellationToken;
 };
 using FHttpRequestList = TIntrusiveList<FHttpRequest>;
 
 ////////////////////////////////////////////////////////////////////////////////
-void FHttpRequest::Release()
+FReturnedRefCountValue FHttpRequest::Release() const
 {
-	if (RefCount.fetch_sub(1, std::memory_order_acq_rel) == 1)
+	const int32 PrevRefs = RefCount.fetch_sub(1, std::memory_order_acq_rel);
+	if (PrevRefs == 1)
 	{
+		ensureMsgf(Next == nullptr, TEXT("The HTTP request must not be released while in a batch"));
+#if !UE_BUILD_SHIPPING 
+		const EHttpRequestFlags CurrentFlags = Flags.Get(std::memory_order_acq_rel);
+		ensureMsgf(CurrentFlags == EHttpRequestFlags::None || EnumHasAnyFlags(CurrentFlags, EHttpRequestFlags::Completing | EHttpRequestFlags::Completed),
+			TEXT("The HTTP request handle must not be released before the completion callback has been triggered."));
+#endif
 		delete this;
 	}
+
+	check(PrevRefs > 0);
+	return FReturnedRefCountValue(PrevRefs - 1);
 }
 
 ////////////////////////////////////////////////////////////////////////////////
@@ -998,20 +1013,18 @@ private:
 	virtual bool					IsHostGroupRegistered(const FName& HostGroup) override;
 	virtual bool					IsHostGroupOk(const FName& HostGroup) override;
 	virtual FHostGroupRegistered&	OnHostGroupRegistered() override { return HostGroupRegisterd; }
-	virtual FIoHttpRequestHandle	CreateRequest(
-										FIoHttpRequestHandle& First,
-										FIoHttpRequestHandle& Last,
+	virtual FIoHttpRequest			CreateRequest(
+										IIoHttpRequestHandle*& First,
+										IIoHttpRequestHandle*& Last,
 										const FName& HostGroupName,
 										const FIoRelativeUrl& RelativeUrl,
 										const FIoHttpOptions& Options,
 										FIoHttpHeaders&& Headers,
 										FIoHttpRequestCompleted&& OnCompleted,
 										const FIoHash* ChunkHash) override;
-	virtual void					IssueRequest(FIoHttpRequestHandle RequestHandle) override;
-	virtual void					ReleaseRequest(FIoHttpRequestHandle Handle) override;
-	virtual void					CancelRequest(FIoHttpRequestHandle Handle) override;
-	virtual void					UpdateRequestPriority(FIoHttpRequestHandle Handle, int32 NewPriority) override;
-	virtual EIoErrorCode			GetRequestStatus(FIoHttpRequestHandle Handle) override;
+	virtual void					IssueRequest(IIoHttpRequestHandle& Handle) override;
+	virtual void					CancelRequest(FIoHttpRequest Handle) override;
+	virtual void					UpdateRequestPriority(FIoHttpRequest Handle, int32 NewPriority) override;
 
 	virtual FIoStatus				CacheResponse(const FIoHttpResponse& Response) override;
 	virtual FIoStatus				EvictFromCache(const FIoHttpResponse& Response) override;
@@ -1289,9 +1302,9 @@ bool FHttpDispatcher::IsHostGroupOk(const FName& HostGroup)
 	return Existing.IsConnected();
 }
 
-FIoHttpRequestHandle FHttpDispatcher::CreateRequest(
-	FIoHttpRequestHandle& First,
-	FIoHttpRequestHandle& Last,
+FIoHttpRequest FHttpDispatcher::CreateRequest(
+	IIoHttpRequestHandle*& First,
+	IIoHttpRequestHandle*& Last,
 	const FName& HostGroupName,
 	const FIoRelativeUrl& RelativeUrl,
 	const FIoHttpOptions& Options,
@@ -1300,7 +1313,7 @@ FIoHttpRequestHandle FHttpDispatcher::CreateRequest(
 	const FIoHash* ChunkHash)
 	
 {
-	FHttpRequest* Request = new FHttpRequest();
+	FHttpRequest* Request	= new FHttpRequest();
 	Request->OnCompleted	= MoveTemp(OnCompleted);
 	Request->Headers		= MoveTemp(Headers);
 	Request->Options		= Options;
@@ -1319,32 +1332,24 @@ FIoHttpRequestHandle FHttpDispatcher::CreateRequest(
 		Request->CacheKey	= GetCacheKey(HostGroupName, Request->RelativeUrl, Request->Headers, Request->Options);
 	}
 
-	if (Last == 0)
+	if (Last == nullptr)
 	{
-		ensure(First == 0);
-		First = UPTRINT(Request);
+		ensure(First == nullptr);
+		First = Request;
 	}
 	else
 	{
-		FHttpRequest* BatchLast = reinterpret_cast<FHttpRequest*>(Last);
-		BatchLast->Next = Request;
+		static_cast<FHttpRequest*>(Last)->Next = Request;
 	}
-	Last = UPTRINT(Request);
+	Last = Request;
 
-	Request->AddRef();
-	return reinterpret_cast<FIoHttpRequestHandle>(Request);
+	return FIoHttpRequest(Request);
 }
 
-void FHttpDispatcher::IssueRequest(FIoHttpRequestHandle RequestHandle)
+void FHttpDispatcher::IssueRequest(IIoHttpRequestHandle& Handle)
 {
-	ensure(RequestHandle != 0);
-	if (RequestHandle == 0)
-	{
-		return;
-	}
-
 	bool bWakeUp				= false;
-	FHttpRequest* NextRequest	= reinterpret_cast<FHttpRequest*>(RequestHandle);
+	FHttpRequest* NextRequest	= static_cast<FHttpRequest*>(&Handle);
 
 	while (NextRequest != nullptr)
 	{
@@ -1352,6 +1357,7 @@ void FHttpDispatcher::IssueRequest(FIoHttpRequestHandle RequestHandle)
 		NextRequest				= Request->Next;
 		Request->Next			= nullptr;
 
+		// The dispatcher holds a reference to the request until completed
 		Request->AddRef();
 		Request->Flags.Add(EHttpRequestFlags::Issued);
 		if (TryReadFromCache(Request) == false)
@@ -1367,16 +1373,14 @@ void FHttpDispatcher::IssueRequest(FIoHttpRequestHandle RequestHandle)
 	}
 }
 
-void FHttpDispatcher::CancelRequest(FIoHttpRequestHandle Handle)
+void FHttpDispatcher::CancelRequest(FIoHttpRequest Handle)
 {
-	ensure(Handle != 0);
-	if (Handle == 0)
+	FHttpRequest* Request = static_cast<FHttpRequest*>(Handle.GetHandle());
+	if (Request == nullptr)
 	{
 		return;
 	}
 
-	FHttpRequest* Request = reinterpret_cast<FHttpRequest*>(Handle);
-
 	const EHttpRequestFlags CurrentFlags = Request->Flags.Get();
 	if (EnumHasAnyFlags(CurrentFlags, EHttpRequestFlags::CancelRequested | EHttpRequestFlags::Completing | EHttpRequestFlags::Completed))
 	{
@@ -1407,30 +1411,19 @@ void FHttpDispatcher::CancelRequest(FIoHttpRequestHandle Handle)
 	}
 }
 
-void FHttpDispatcher::UpdateRequestPriority(FIoHttpRequestHandle Handle, int32 NewPriority)
+void FHttpDispatcher::UpdateRequestPriority(FIoHttpRequest Handle, int32 NewPriority)
 {
-	ensure(Handle != 0);
-	if (Handle == 0)
+	FHttpRequest* Request = static_cast<FHttpRequest*>(Handle.GetHandle());
+	if (Request == nullptr)
 	{
 		return;
 	}
 
-	FHttpRequest* Request = reinterpret_cast<FHttpRequest*>(Handle);
 	UE_LOG(LogHttpIoDispatcher, Verbose, TEXT("Updating request priority, SeqNo=%u, Priority=%d, NewPriority=%d"),
 		Request->SeqNo, Request->Options.GetPriority(), NewPriority);
 	HttpQueue.Reprioritize(Request, NewPriority);
 }
 
-EIoErrorCode FHttpDispatcher::GetRequestStatus(FIoHttpRequestHandle Handle)
-{
-	if (Handle != 0)
-	{
-		return reinterpret_cast<FHttpRequest*>(Handle)->CompletionStatus.load(std::memory_order_relaxed);
-	}
-
-	return EIoErrorCode::InvalidCode;
-}
-
 FIoStatus FHttpDispatcher::CacheResponse(const FIoHttpResponse& Response)
 {
 	if (Response.IsCached())
@@ -1589,25 +1582,6 @@ bool FHttpDispatcher::TryReadFromCache(FHttpRequest* Request)
 	return false;
 }
 
-void FHttpDispatcher::ReleaseRequest(FIoHttpRequestHandle Handle)
-{
-	ensure(Handle != 0);
-	FHttpRequest* NextRequest = reinterpret_cast<FHttpRequest*>(Handle);
-	while (NextRequest != nullptr)
-	{
-		FHttpRequest* Request	= NextRequest;
-		NextRequest				= Request->Next;
-		Request->Next			= nullptr;
-
-#if !UE_BUILD_SHIPPING 
-		const EHttpRequestFlags Flags = Request->Flags.Get();
-		ensureMsgf(Flags == EHttpRequestFlags::None || EnumHasAnyFlags(Flags, EHttpRequestFlags::Completing | EHttpRequestFlags::Completed),
-			TEXT("The HTTP request handle must not be released before the completion callback has been triggered."));
-#endif
-		Request->Release();
-	}
-}
-
 bool FHttpDispatcher::ProcessHttpRequests()
 {
 	TRACE_CPUPROFILER_EVENT_SCOPE(FHttpDispatcher::ProcessHttpRequests);
diff --git a/Engine/Source/Runtime/Experimental/IoStore/OnDemandCore/Private/HttpIoDispatcher.cpp b/Engine/Source/Runtime/Experimental/IoStore/OnDemandCore/Private/HttpIoDispatcher.cpp
--- a/Engine/Source/Runtime/Experimental/IoStore/OnDemandCore/Private/HttpIoDispatcher.cpp
+++ b/Engine/Source/Runtime/Experimental/IoStore/OnDemandCore/Private/HttpIoDispatcher.cpp
@@ -92,65 +92,30 @@ FIoRelativeUrl FIoRelativeUrl::From(FAnsiStringView Url)
 }
 
 ////////////////////////////////////////////////////////////////////////////////
-FIoHttpRequest::~FIoHttpRequest()
-{
-	Release();
-}
-
 void FIoHttpRequest::Cancel()
 {
-	if (Handle != 0)
-	{
-		return GHttpIoDispatcher->CancelRequest(Handle);
-	}
+	FIoHttpRequest ToCancel(*this);
+	GHttpIoDispatcher->CancelRequest(ToCancel);
 }
 
 void FIoHttpRequest::UpdatePriorty(int32 NewPriority)
 {
-	if (Handle != 0)
-	{
-		return GHttpIoDispatcher->UpdateRequestPriority(Handle, NewPriority);
-	}
+	FIoHttpRequest ToUpdate(*this);
+	GHttpIoDispatcher->UpdateRequestPriority(ToUpdate, NewPriority);
 }
 
 EIoErrorCode FIoHttpRequest::Status() const
 {
-	if (Handle != 0)
-	{
-		return GHttpIoDispatcher->GetRequestStatus(Handle);
-	}
-
-	return EIoErrorCode::InvalidCode;
-}
-
-void FIoHttpRequest::Release()
-{
-	if (Handle != 0)
-	{
-		GHttpIoDispatcher->ReleaseRequest(Handle);
-		Handle = 0;
-	}
+	return Handle.IsValid() ? Handle->GetStatus() : EIoErrorCode::InvalidCode;
 }
 
 FIoHttpRequest&	FIoHttpRequest::operator=(FIoHttpRequest&& Other)
 {
-	Release();
-	Handle = Other.Handle;
-	Other.Handle = 0;
+	Handle = MoveTemp(Other.Handle);
 	return *this;
 }
 
 ////////////////////////////////////////////////////////////////////////////////
-FIoHttpBatch::~FIoHttpBatch()
-{
-	if (First != 0)
-	{
-		ensure(Last != 0);
-		GHttpIoDispatcher->ReleaseRequest(First);
-		First = Last = 0;
-	}
-}
-
 FIoHttpRequest FIoHttpBatch::Get(
 	const FName& HostGroup,
 	const FIoRelativeUrl& RelativeUrl,
@@ -159,7 +124,7 @@ FIoHttpRequest FIoHttpBatch::Get(
 	const FIoHash& ChunkHash,
 	FIoHttpRequestCompleted&& OnCompleted)
 {
-	FIoHttpRequestHandle Handle = GHttpIoDispatcher->CreateRequest(
+	return GHttpIoDispatcher->CreateRequest(
 		First,
 		Last,
 		HostGroup,
@@ -168,8 +133,6 @@ FIoHttpRequest FIoHttpBatch::Get(
 		MoveTemp(Headers),
 		MoveTemp(OnCompleted),
 		&ChunkHash);
-
-	return FIoHttpRequest(Handle);
 }
 
 FIoHttpRequest FIoHttpBatch::Get(
@@ -179,7 +142,7 @@ FIoHttpRequest FIoHttpBatch::Get(
 	const FIoHttpOptions& Options,
 	FIoHttpRequestCompleted&& OnCompleted)
 {
-	FIoHttpRequestHandle Handle = GHttpIoDispatcher->CreateRequest(
+	return GHttpIoDispatcher->CreateRequest(
 		First,
 		Last,
 		HostGroup,
@@ -187,8 +150,6 @@ FIoHttpRequest FIoHttpBatch::Get(
 		Options,
 		MoveTemp(Headers),
 		MoveTemp(OnCompleted));
-
-	return FIoHttpRequest(Handle);
 }
 
 FIoHttpRequest FIoHttpBatch::Get(
@@ -210,11 +171,10 @@ FIoHttpRequest FIoHttpBatch::Get(
 
 void FIoHttpBatch::Issue()
 {
-	if (First != 0)
+	if (IIoHttpRequestHandle* ToIssue = First)
 	{
-		ensure(Last != 0);
-		GHttpIoDispatcher->IssueRequest(First);
-		First = Last = 0;
+		First = Last = nullptr;
+		GHttpIoDispatcher->IssueRequest(*ToIssue);
 	}
 }
 
@@ -295,7 +255,7 @@ bool FHttpIoDispatcher::IsHostGroupOk(const FName& HostGroup)
 
 FIoHttpBatch FHttpIoDispatcher::NewBatch()
 {
-	ensure(GHttpIoDispatcher.IsValid());
+	check(GHttpIoDispatcher.IsValid());
 	return FIoHttpBatch();
 }
 
@@ -308,10 +268,10 @@ FIoHttpRequest FHttpIoDispatcher::Get(
 {
 	ensure(GHttpIoDispatcher.IsValid());
 
-	FIoHttpRequestHandle First = 0;
-	FIoHttpRequestHandle Last = 0;
+	IIoHttpRequestHandle* First = nullptr;
+	IIoHttpRequestHandle* Last = nullptr;
 
-	FIoHttpRequestHandle Handle = GHttpIoDispatcher->CreateRequest(
+	FIoHttpRequest Request = GHttpIoDispatcher->CreateRequest(
 		First,
 		Last,
 		HostGroup,
@@ -320,8 +280,12 @@ FIoHttpRequest FHttpIoDispatcher::Get(
 		MoveTemp(Headers),
 		MoveTemp(OnCompleted));
 
-	GHttpIoDispatcher->IssueRequest(First);
-	return FIoHttpRequest(First);
+	if (ensure(First != nullptr))
+	{
+		GHttpIoDispatcher->IssueRequest(*First);
+	}
+
+	return Request; 
 }
 
 FIoHttpRequest FHttpIoDispatcher::Get(
@@ -340,21 +304,22 @@ FIoHttpRequest FHttpIoDispatcher::Get(
 {
 	return Get(HostGroup, RelativeUrl, FIoHttpHeaders(), FIoHttpOptions::Default, MoveTemp(OnCompleted));
 }
- FIoStatus FHttpIoDispatcher::CacheResponse(const FIoHttpResponse& Response)
+
+FIoStatus FHttpIoDispatcher::CacheResponse(const FIoHttpResponse& Response)
 {
-	 ensure(GHttpIoDispatcher.IsValid());
+	 check(GHttpIoDispatcher.IsValid());
 	 return GHttpIoDispatcher->CacheResponse(Response);
 }
 
 FIoStatus FHttpIoDispatcher::EvictFromCache(const FIoHttpResponse& Response)
 {
-	ensure(GHttpIoDispatcher.IsValid());
+	check(GHttpIoDispatcher.IsValid());
 	return GHttpIoDispatcher->EvictFromCache(Response);
 }
 
 FHttpIoDispatcher::FHostGroupRegistered& FHttpIoDispatcher::OnHostGroupRegistered()
 {
-	ensure(GHttpIoDispatcher.IsValid());
+	check(GHttpIoDispatcher.IsValid());
 	return GHttpIoDispatcher->OnHostGroupRegistered();
 }
 
diff --git a/Engine/Source/Runtime/Experimental/IoStore/OnDemandCore/Public/IO/HttpIoDispatcher.h b/Engine/Source/Runtime/Experimental/IoStore/OnDemandCore/Public/IO/HttpIoDispatcher.h
--- a/Engine/Source/Runtime/Experimental/IoStore/OnDemandCore/Public/IO/HttpIoDispatcher.h
+++ b/Engine/Source/Runtime/Experimental/IoStore/OnDemandCore/Public/IO/HttpIoDispatcher.h
@@ -12,6 +12,7 @@
 #include "Misc/EnumClassFlags.h"
 #include "Templates/Function.h"
 #include "Templates/SharedPointer.h"
+#include "Templates/RefCounting.h"
 
 #define UE_API IOSTOREONDEMANDCORE_API
 
@@ -23,8 +24,7 @@ namespace UE
 namespace IoStore { class FOnDemandHttpStats; }
 using FOnDemandHttpStats = IoStore::FOnDemandHttpStats;
 
-/** Internal request handle type. */
-using FIoHttpRequestHandle = UPTRINT;
+class IIoHttpRequestHandle;
 
 /** Flags for controlling the behavior of an HTTP request. */
 enum class EIoHttpFlags
@@ -234,24 +234,25 @@ class FIoHttpRequest
 {
 public:
 							FIoHttpRequest() = default;
-	explicit 				FIoHttpRequest(FIoHttpRequestHandle InHandle)
+	explicit 				FIoHttpRequest(IIoHttpRequestHandle* InHandle)
 								: Handle(InHandle) { }
 							FIoHttpRequest(FIoHttpRequest&& Other)
-								: Handle(Other.Handle) { Other.Handle = 0; }
-							FIoHttpRequest(const FIoHttpRequest&) = delete;
-	UE_API					~FIoHttpRequest();
+								: Handle(MoveTemp(Other.Handle)) { }
 
-	bool					IsValid() const { return Handle != 0; }
+	bool					IsValid() const { return Handle.IsValid(); }
 	UE_API void				Cancel();
 	UE_API void				UpdatePriorty(int32 NewPriority);
 	UE_API EIoErrorCode		Status() const;
+	IIoHttpRequestHandle*	GetHandle() { return Handle.GetReference(); }
 
 	FIoHttpRequest&			operator=(const FIoHttpRequest&) = delete;
 	UE_API FIoHttpRequest&	operator=(FIoHttpRequest&& Other);
 
 private:
-	void					Release();
-	FIoHttpRequestHandle	Handle = 0;
+							FIoHttpRequest(const FIoHttpRequest& Other)
+								: Handle(Other.Handle) { }
+
+	TRefCountPtr<IIoHttpRequestHandle> Handle;
 };
 
 /** Flags describing a HTTP response. */
@@ -322,82 +323,89 @@ using FIoHttpRequestCompleted = TUniqueFunction<void(FIoHttpResponse&&)>;
 class FIoHttpBatch
 {
 public:
-							FIoHttpBatch(const FIoHttpRequest&) = delete;
-							FIoHttpBatch(FIoHttpBatch&& Other)
-								: First(Other.First), Last(Other.Last)
-							{
-								Other.First = 0;
-								Other.Last = 0;
-							}
-
-	UE_API					~FIoHttpBatch();
-
-	UE_API FIoHttpRequest	Get(const FName& HostGroup,
-								const FIoRelativeUrl& RelativeUrl,
-								FIoHttpHeaders&& Headers,
-								const FIoHttpOptions& Options,
-								const FIoHash& ChunkHash,
-								FIoHttpRequestCompleted&& OnCompleted);
-	UE_API FIoHttpRequest	Get(const FName& HostGroup,
-								const FIoRelativeUrl& RelativeUrl,
-								FIoHttpHeaders&& Headers,
-								const FIoHttpOptions& Options,
-								FIoHttpRequestCompleted&& OnCompleted);
-	UE_API FIoHttpRequest	Get(const FName& HostGroup,
-								const FIoRelativeUrl& RelativeUrl,
-								FIoHttpHeaders&& Headers,
-								FIoHttpRequestCompleted&& OnCompleted);
-	UE_API FIoHttpRequest	Get(const FName& HostGroup,
-								const FIoRelativeUrl& RelativeUrl,
-								FIoHttpRequestCompleted&& OnCompleted);
-	UE_API void				Issue();
-
-	FIoHttpBatch&			operator=(const FIoHttpBatch&) = delete;
-	FIoHttpBatch&			operator=(FIoHttpBatch&& Other) = delete;
+										FIoHttpBatch(const FIoHttpRequest&) = delete;
+										FIoHttpBatch(FIoHttpBatch&& Other)
+											: First(Other.First), Last(Other.Last)
+										{
+											Other.First = 0;
+											Other.Last = 0;
+										}
+
+	[[nodiscard]] UE_API FIoHttpRequest	Get(const FName& HostGroup,
+											const FIoRelativeUrl& RelativeUrl,
+											FIoHttpHeaders&& Headers,
+											const FIoHttpOptions& Options,
+											const FIoHash& ChunkHash,
+											FIoHttpRequestCompleted&& OnCompleted);
+	[[nodiscard]] UE_API FIoHttpRequest	Get(const FName& HostGroup,
+											const FIoRelativeUrl& RelativeUrl,
+											FIoHttpHeaders&& Headers,
+											const FIoHttpOptions& Options,
+											FIoHttpRequestCompleted&& OnCompleted);
+	[[nodiscard]] UE_API FIoHttpRequest	Get(const FName& HostGroup,
+											const FIoRelativeUrl& RelativeUrl,
+											FIoHttpHeaders&& Headers,
+											FIoHttpRequestCompleted&& OnCompleted);
+	[[nodiscard]] UE_API FIoHttpRequest	Get(const FName& HostGroup,
+											const FIoRelativeUrl& RelativeUrl,
+											FIoHttpRequestCompleted&& OnCompleted);
+	UE_API void							Issue();
+
+	FIoHttpBatch&						operator=(const FIoHttpBatch&) = delete;
+	FIoHttpBatch&						operator=(FIoHttpBatch&& Other) = delete;
 
 private:
-							friend class FHttpIoDispatcher;
-							FIoHttpBatch() = default;
-	FIoHttpRequestHandle	First = 0;
-	FIoHttpRequestHandle 	Last = 0;
+										friend class FHttpIoDispatcher;
+										FIoHttpBatch() = default;
+
+	// The batch relies and the caller keeping a reference to each handle
+	IIoHttpRequestHandle* First = nullptr;
+	IIoHttpRequestHandle* Last = nullptr;
 };
 
 /** HTTP I/O dispatcher . */
 class FHttpIoDispatcher
 {
 public:
-	UE_API static bool				IsInitialized();
-	UE_API static FIoStatus			Initialize(TSharedPtr<class IHttpIoDispatcher> Dispatcher);
-	UE_API static FIoStatus			Shutdown();
-
-	UE_API static void				GetHttpStats(FOnDemandHttpStats& Out);
-	UE_API static FIoStatus			RegisterHostGroup(const FName& HostGroup, TConstArrayView<FAnsiString> HostNames, FAnsiStringView TestUrl);
-	UE_API static FIoStatus			RegisterHostGroup(const FName& HostGroup, TConstArrayView<FAnsiString> HostNames);
-	UE_API static FIoStatus			RegisterHostGroup(const FName& HostGroup, FAnsiStringView HostName, FAnsiStringView TestUrl);
-	UE_API static FIoStatus			RegisterHostGroup(const FName& HostGroup, FAnsiStringView HostName);
-	UE_API static bool				IsHostGroupRegistered(const FName& HostGroup);
-	UE_API static bool				IsHostGroupOk(const FName& HostGroup);
-
-	UE_API static FIoHttpBatch		NewBatch();
-	UE_API static FIoHttpRequest	Get(const FName& HostGroup,
-										const FIoRelativeUrl& RelativeUrl,
-										FIoHttpHeaders&& Headers,
-										const FIoHttpOptions& Options,
-										FIoHttpRequestCompleted&& OnCompleted);
-	UE_API static FIoHttpRequest	Get(const FName& HostGroup,
-										const FIoRelativeUrl& RelativeUrl,
-										FIoHttpHeaders&& Headers,
-										FIoHttpRequestCompleted&& OnCompleted);
-	UE_API static FIoHttpRequest	Get(const FName& HostGroup,
-										const FIoRelativeUrl& RelativeUrl,
-										FIoHttpRequestCompleted&& OnCompleted);
-
-	/** */
-	UE_API static FIoStatus			CacheResponse(const FIoHttpResponse& Response);
-	UE_API static FIoStatus			EvictFromCache(const FIoHttpResponse& Response);
+	[[nodiscard]] UE_API static bool			IsInitialized();
+	[[nodiscard]] UE_API static FIoStatus		Initialize(TSharedPtr<class IHttpIoDispatcher> Dispatcher);
+	[[nodiscard]] UE_API static FIoStatus		Shutdown();
+
+				  UE_API static void			GetHttpStats(FOnDemandHttpStats& Out);
+	[[nodiscard]] UE_API static FIoStatus		RegisterHostGroup(const FName& HostGroup, TConstArrayView<FAnsiString> HostNames, FAnsiStringView TestUrl);
+	[[nodiscard]] UE_API static FIoStatus		RegisterHostGroup(const FName& HostGroup, TConstArrayView<FAnsiString> HostNames);
+	[[nodiscard]] UE_API static FIoStatus		RegisterHostGroup(const FName& HostGroup, FAnsiStringView HostName, FAnsiStringView TestUrl);
+	[[nodiscard]] UE_API static FIoStatus		RegisterHostGroup(const FName& HostGroup, FAnsiStringView HostName);
+	[[nodiscard]] UE_API static bool			IsHostGroupRegistered(const FName& HostGroup);
+	[[nodiscard]] UE_API static bool			IsHostGroupOk(const FName& HostGroup);
+
+	[[nodiscard]] UE_API static FIoHttpBatch	NewBatch();
+	[[nodiscard]] UE_API static FIoHttpRequest	Get(const FName& HostGroup,
+													const FIoRelativeUrl& RelativeUrl,
+													FIoHttpHeaders&& Headers,
+													const FIoHttpOptions& Options,
+													FIoHttpRequestCompleted&& OnCompleted);
+	[[nodiscard]] UE_API static FIoHttpRequest	Get(const FName& HostGroup,
+													const FIoRelativeUrl& RelativeUrl,
+													FIoHttpHeaders&& Headers,
+													FIoHttpRequestCompleted&& OnCompleted);
+	[[nodiscard]] UE_API static FIoHttpRequest	Get(const FName& HostGroup,
+													const FIoRelativeUrl& RelativeUrl,
+													FIoHttpRequestCompleted&& OnCompleted);
+
+	[[nodiscard]] UE_API static FIoStatus		CacheResponse(const FIoHttpResponse& Response);
+	[[nodiscard]] UE_API static FIoStatus		EvictFromCache(const FIoHttpResponse& Response);
 
 	DECLARE_MULTICAST_DELEGATE_OneParam(FHostGroupRegistered, const FName&);
-	UE_API static FHostGroupRegistered& OnHostGroupRegistered();
+	UE_API static FHostGroupRegistered&			OnHostGroupRegistered();
+};
+
+/** A reference counted HTTP I/O dispatcher request handle. */
+class IIoHttpRequestHandle
+	: public IRefCountedObject
+{
+public: 
+	virtual EIoErrorCode GetStatus() const = 0;
 };
 
 /** HTTP I/O dispatcher interface. */
@@ -415,20 +423,18 @@ public:
 	virtual bool					IsHostGroupOk(const FName& HostGroup) = 0;
 	virtual FHostGroupRegistered&	OnHostGroupRegistered() = 0;
 
-	virtual FIoHttpRequestHandle	CreateRequest(
-										FIoHttpRequestHandle& First,
-										FIoHttpRequestHandle& Last,
+	virtual FIoHttpRequest			CreateRequest(
+										IIoHttpRequestHandle*& First,
+										IIoHttpRequestHandle*& Last,
 										const FName& HostGroup,
 										const FIoRelativeUrl& RelativeUrl,
 										const FIoHttpOptions& Options,
 										FIoHttpHeaders&& Headers,
 										FIoHttpRequestCompleted&& OnCompleted,
 										const FIoHash* ChunkHash = nullptr) = 0;
-	virtual void					IssueRequest(FIoHttpRequestHandle RequestHandle) = 0;
-	virtual void					CancelRequest(FIoHttpRequestHandle Handle) = 0;
-	virtual void					UpdateRequestPriority(FIoHttpRequestHandle Handle, int32 NewPriority) = 0;
-	virtual EIoErrorCode			GetRequestStatus(FIoHttpRequestHandle Handle) = 0;
-	virtual void					ReleaseRequest(FIoHttpRequestHandle Handle) = 0;
+	virtual void					IssueRequest(IIoHttpRequestHandle& Handle) = 0;
+	virtual void					CancelRequest(FIoHttpRequest Handle) = 0;
+	virtual void					UpdateRequestPriority(FIoHttpRequest Handle, int32 NewPriority) = 0;
 
 	virtual FIoStatus				CacheResponse(const FIoHttpResponse& Response) = 0;
 	virtual FIoStatus				EvictFromCache(const FIoHttpResponse& Response) = 0;
