Add FError detail filter to CAS/Journal error context.

#rnx

#Synced-CL 48874164

diff --git a/Engine/Source/Runtime/Experimental/IoStore/OnDemandCore/Public/IO/OnDemandError.h b/Engine/Source/Runtime/Experimental/IoStore/OnDemandCore/Public/IO/OnDemandError.h
--- a/Engine/Source/Runtime/Experimental/IoStore/OnDemandCore/Public/IO/OnDemandError.h
+++ b/Engine/Source/Runtime/Experimental/IoStore/OnDemandCore/Public/IO/OnDemandError.h
@@ -125,7 +125,7 @@ inline UE::IoStore::FResult MakeJournalError(
 		.IoErrorCode		= IoErrorCode,
 		.SystemErrorCode	= SystemErrorCode,
 		.ErrorMessage		= MoveTemp(ErrorMessage)
-	}));
+	}, UE::UnifiedError::EDetailFilter::All));
 }
 
 template <typename ResultType>
@@ -144,7 +144,7 @@ inline UE::IoStore::TResult<ResultType> MakeSnapshotError(
 		.IoErrorCode		= IoErrorCode,
 		.SystemErrorCode	= SystemErrorCode,
 		.ErrorMessage		= MoveTemp(ErrorMessage)
-	}));
+	}, UE::UnifiedError::EDetailFilter::All));
 }
 
 struct FChunkMissingErrorContext
