iax: Potential fix for "user callback returned error" seen in live.

To be able to distribute inbound data to the correct callbacks and IO buffer we need to look ahead for a corresponding ID to match to requests. Under very specific circumstances the look ahead could happen too early and cause miss-timed internal callbacks in the wrong place. This change ensures that all relevant data is consumed when response data is expected and not at any other time.

#rnx

diff --git a/Engine/Source/Runtime/Experimental/IoStore/HttpClient/Private/TransactionTwo.inl b/Engine/Source/Runtime/Experimental/IoStore/HttpClient/Private/TransactionTwo.inl
--- a/Engine/Source/Runtime/Experimental/IoStore/HttpClient/Private/TransactionTwo.inl
+++ b/Engine/Source/Runtime/Experimental/IoStore/HttpClient/Private/TransactionTwo.inl
@@ -504,43 +504,56 @@ FOutcome FDriverNg::RecvResponse(FTlsPeer& Peer, TArray<FHeaderNg>* Sink)
 			return FOutcome::Error("Not a headers frame", int32(RecvFrame.Type));
 		}
 
-		int32 EndOfHeaders = (RecvFrame.Flags & FFrame::Flag_EndHeaders);
-		int32 EndOfStream = (RecvFrame.Flags & FFrame::Flag_EndStream);
-		uint32 ThisId = RecvFrame.GetId();
+		int32 RecvSize = FMath::Min<int32>(sizeof(RecvBuffer), RecvRemaining);
 
-		if (EndOfHeaders)
+		FMutableMemoryView Dest(RecvBuffer, RecvSize);
+		FOutcome Outcome = Peer.Recv((char*)(Dest.GetData()), RecvSize);
+		if (!Outcome.IsOk())
 		{
-			RecvFrame = {};
+			return Outcome;
 		}
+		int32 Result = Outcome.GetResult();
 
-		FOutcome Outcome = Tick(Peer); 
-		if (Outcome.IsError())
+		Dest = Dest.Left(Outcome.GetResult());
+		if (!Dest.IsEmpty() && (Outcome = PumpRecv(Peer, Dest)).IsError())
 		{
 			return Outcome;
 		}
 
-		if (EndOfStream)
+		RecvRemaining -= Result;
+		if (RecvRemaining > 0)
+		{
+			continue;
+		}
+
+		check(RecvRemaining == 0);
+		RecvRemaining -= sizeof(RecvFrame);
+
+		int32 EndOfHeaders = (RecvFrame.Flags & FFrame::Flag_EndHeaders);
+
+		if (RecvFrame.Flags & FFrame::Flag_EndStream)
 		{
 			check(EndOfHeaders);
 			return FOutcome::Ok();
 		}
 
-		if (EndOfHeaders)
+		uint32 ThisId = RecvFrame.GetId();
+		if (Outcome = Tick(Peer); Outcome.IsError())
 		{
-			// We shall end headers once the first data frame arrives, which may
-			// have already happened in the Tick() call.
-			if (Outcome.GetResult() == ThisId)
-			{
-				return RecvResponse(Peer, ResponseHeaders);
-			}
-
-			return FOutcome::Waiting();
+			return Outcome;
 		}
 
 		if (Outcome.GetResult() != ThisId)
 		{
 			return FOutcome::WaitStream();
 		}
+
+		if (EndOfHeaders)
+		{
+			// We shall end headers once the first data frame arrives, which may
+			// have already happened in the Tick() call.
+			return RecvResponse(Peer, ResponseHeaders);
+		}
 	}
 }
 
