Enable streaming packages without pre-installing into the install cache.
- this change enables streaming all I/O store chunks just-in-time
- disabled by default and also currently in shipping builds
- enabled with pak.AllowMissingUToc and iax.PackageStreamingEnabled

#rnx

#Synced-CL 48874164

diff --git a/Engine/Source/Runtime/Experimental/IoStore/OnDemand/Private/OnDemandIoStore.cpp b/Engine/Source/Runtime/Experimental/IoStore/OnDemand/Private/OnDemandIoStore.cpp
--- a/Engine/Source/Runtime/Experimental/IoStore/OnDemand/Private/OnDemandIoStore.cpp
+++ b/Engine/Source/Runtime/Experimental/IoStore/OnDemand/Private/OnDemandIoStore.cpp
@@ -54,6 +54,24 @@
 namespace UE::IoStore
 {
 
+bool GIaxPackageStreamingEnabled = false;
+static FAutoConsoleVariableRef CVar_PackageStreamingEnabled(
+	TEXT("iax.PackageStreamingEnabled"),
+	GIaxPackageStreamingEnabled,
+	TEXT("Enables streaming packages without pre-installing into the install cache.")
+);
+
+///////////////////////////////////////////////////////////////////////////////
+bool IsPackageStreamingEnabled()
+{
+#if UE_BUILD_SHIPPING
+	return false;
+#else
+	static bool bPackageStreamingEnabled = FParse::Param(FCommandLine::Get(), TEXT("iax.PackageStreamingEnabled"));
+	return bPackageStreamingEnabled || GIaxPackageStreamingEnabled;
+#endif
+}
+
 ///////////////////////////////////////////////////////////////////////////////
 namespace Private
 {
@@ -882,7 +900,9 @@ FIoStatus FOnDemandIoStore::GetIsOnDemand(TConstArrayView<FOnDemandIsOnDemandArg
 
 FOnDemandChunkInfo FOnDemandIoStore::GetStreamingChunkInfo(const FIoChunkId& ChunkId)
 {
-	const EOnDemandContainerFlags ContainerFlags = EOnDemandContainerFlags::Mounted | EOnDemandContainerFlags::StreamOnDemand;
+	const EOnDemandContainerFlags ContainerFlags = IsPackageStreamingEnabled()
+		? EOnDemandContainerFlags::Mounted
+		: EOnDemandContainerFlags::Mounted | EOnDemandContainerFlags::StreamOnDemand;
 
 	TUniqueLock Lock(ContainerMutex);
 
@@ -1364,8 +1384,8 @@ FIoStatus FOnDemandIoStore::ProcessMountRequest(FMountRequest& MountRequest)
 			}
 		}
 
-		// Try fetch and deserialize the container header if the container is used for installing content.
-		if (EnumHasAnyFlags(Container->Flags, EOnDemandContainerFlags::InstallOnDemand))
+		// Try fetch and deserialize the container header if the container is used for installing content or if package streaming is enabled
+		if (EnumHasAnyFlags(Container->Flags, EOnDemandContainerFlags::InstallOnDemand) || IsPackageStreamingEnabled())
 		{
 			const FIoChunkId ChunkId			= CreateContainerHeaderChunkId(Container->ContainerId);
 			const FOnDemandChunkInfo ChunkInfo	= FOnDemandChunkInfo::Find(Container, ChunkId);
@@ -1938,7 +1958,7 @@ void FOnDemandIoStore::GetReferencedContent(TArray<FSharedOnDemandContainer>& Ou
 	UE::TUniqueLock Lock(ContainerMutex);
 	for (FSharedOnDemandContainer& Container : Containers)
 	{
-		if (EnumHasAnyFlags(Container->Flags, EOnDemandContainerFlags::StreamOnDemand))
+		if (EnumHasAnyFlags(Container->Flags, EOnDemandContainerFlags::StreamOnDemand) && !IsPackageStreamingEnabled())
 		{
 			continue;
 		}
diff --git a/Engine/Source/Runtime/Experimental/IoStore/OnDemand/Private/OnDemandPackageStoreBackend.cpp b/Engine/Source/Runtime/Experimental/IoStore/OnDemand/Private/OnDemandPackageStoreBackend.cpp
--- a/Engine/Source/Runtime/Experimental/IoStore/OnDemand/Private/OnDemandPackageStoreBackend.cpp
+++ b/Engine/Source/Runtime/Experimental/IoStore/OnDemand/Private/OnDemandPackageStoreBackend.cpp
@@ -15,6 +15,8 @@
 namespace UE::IoStore
 {
 
+bool IsPackageStreamingEnabled();
+
 ///////////////////////////////////////////////////////////////////////////////
 class FOnDemandPackageStoreBackend final
 	: public IOnDemandPackageStoreBackend
@@ -166,7 +168,7 @@ EPackageStoreEntryStatus FOnDemandPackageStoreBackend::GetPackageStoreEntry(
 			OutPackageStoreEntry.ShaderMapHashes =
 				MakeArrayView(PackageStoreEntry.ShaderMapHashes.Data(), PackageStoreEntry.ShaderMapHashes.Num());
 
-			if (PackageEntry->IsReferenced)
+			if (PackageEntry->IsReferenced || IsPackageStreamingEnabled())
 			{
 				return EPackageStoreEntryStatus::Ok;
 			}
diff --git a/Engine/Source/Runtime/PakFile/Private/IPlatformFilePak.cpp b/Engine/Source/Runtime/PakFile/Private/IPlatformFilePak.cpp
--- a/Engine/Source/Runtime/PakFile/Private/IPlatformFilePak.cpp
+++ b/Engine/Source/Runtime/PakFile/Private/IPlatformFilePak.cpp
@@ -228,6 +228,12 @@ static FAutoConsoleVariableRef CVar_UseDirectoryTreeForPakSearch(
 	TEXT("Enable use of directory tree when finding files"));
 #endif // ENABLE_PAKFILE_USE_DIRECTORY_TREE
 
+bool GPak_AllowMissingUToc = false;
+static FAutoConsoleVariableRef CVar_AllowMissingUToc(
+	TEXT("pak.AllowMissingUToc"),
+	GPak_AllowMissingUToc,
+	TEXT("Enables mounting a .pak file without the corresponding .utoc file even if the I/O dispatcher is enabled")
+);
 
 TPakChunkHash ComputePakChunkHash(const void* InData, int64 InDataSizeInBytes)
 {
@@ -247,6 +253,19 @@ TPakChunkHash ComputePakChunkHash(const void* InData, int64 InDataSizeInBytes)
 namespace UE::PakFile::Private
 {
 
+static bool MountFailsOnMissingUToc()
+{
+#if !UE_BUILD_SHIPPING
+	static bool bAllowMissingUToc = FParse::Param(FCommandLine::Get(), TEXT("pak.AllowMissingUToc"));
+	if (bAllowMissingUToc || GPak_AllowMissingUToc)
+	{
+		return false;
+	}
+#endif
+
+	return FPakPlatformFile::bMountFailOnMissingUtoc;
+}
+
 /** Helper class to filter out files which have already been visited in one of the pak files. */
 class FPreventDuplicatesVisitorBase
 {
@@ -6097,7 +6116,7 @@ bool FPakPlatformFile::Mount(const FPakMountArgs& MountArgs, FIoStatus* OutIoMou
 				}
 				else
 				{
-					if (FPakPlatformFile::bMountFailOnMissingUtoc)
+					if (UE::PakFile::Private::MountFailsOnMissingUToc())
 					{
 						IoStoreSuccess = FIoStatus(EIoErrorCode::NotFound, TEXT("utoc not found"));
 						UE_LOG(LogPakFile, Warning, TEXT("IoStore container \"%s\" not found"), *UtocPath);
