IAD Installer Supports Concurrent Requests
- IAD cache has In memory tracking of block sizes.
- PutChunk returns false when the cache is full.
- Canceled requests are no longer added to the cache if the http request succeeds.
- Purge/Defrag only happens when PutChunk fails. The amount purged is the total of pending but not cancelled requests.

#Synced-CL 48981444

diff --git a/Engine/Source/Runtime/Core/Internal/IO/IoContainers.h b/Engine/Source/Runtime/Core/Internal/IO/IoContainers.h
--- a/Engine/Source/Runtime/Core/Internal/IO/IoContainers.h
+++ b/Engine/Source/Runtime/Core/Internal/IO/IoContainers.h
@@ -175,6 +175,11 @@ public:
 		return Head;
 	}
 
+	const ElementType* PeekHead() const
+	{
+		return Head;
+	}
+
 	bool Remove(ElementType* Element)
 	{
 		if (Element == nullptr || IsEmpty())
diff --git a/Engine/Source/Runtime/Experimental/IoStore/OnDemand/Private/OnDemandContentInstaller.cpp b/Engine/Source/Runtime/Experimental/IoStore/OnDemand/Private/OnDemandContentInstaller.cpp
--- a/Engine/Source/Runtime/Experimental/IoStore/OnDemand/Private/OnDemandContentInstaller.cpp
+++ b/Engine/Source/Runtime/Experimental/IoStore/OnDemand/Private/OnDemandContentInstaller.cpp
@@ -84,7 +84,7 @@ namespace UE::IoStore
 		// Can't safely change this on the fly. Don't call directly
 		static TAutoConsoleVariable<bool> CVar_EnableConcurrentRequests(
 			TEXT("iostore.ContentInstallerEnableConcurrentInstallRequests"),
-			false,
+			true,
 			TEXT("Enable the content installer to install chunks for multiple requests concurrently."),
 			ECVF_SaveForNextBoot
 			);
@@ -1169,21 +1169,101 @@ void FOnDemandContentInstaller::ProcessDownloadedChunk(
 {
 	TRACE_CPUPROFILER_EVENT_SCOPE(FOnDemandContentInstaller::ProcessDownloadedChunk);
 
+	// Find the pending download for this chunk
+	bool bFoundRequest = false;
+	auto It = PendingChunkDownloads.CreateKeyIterator(ChunkHash);
+	for (; It; ++It)
+	{
+		FChunkDownloadRequest& ChunkRequest = It.Value();
+		if (ChunkRequest.RequestId == RequestId)
+		{
+			bFoundRequest = true;
+			break;
+		}
+	}
+	check(bFoundRequest);
+	FChunkDownloadRequest& ChunkRequest = It.Value();
+
 	FResult Result = MakeValue();
 
+	const bool bHttpCancelled = (InErrorCode == EIoErrorCode::Cancelled);
+	check(ChunkRequest.bChunkCanceled || !bHttpCancelled);
+
 	const bool bHttpOk = HttpStatusCode > 199 && HttpStatusCode < 300 && Chunk.GetSize() > 0;
-	if (bHttpOk)
+	if (bHttpOk == false)
+	{
+#if COUNTERSTRACE_ENABLED
+		if (InErrorCode == EIoErrorCode::Cancelled)
+		{
+			TRACE_COUNTER_INCREMENT(OnDemandContentInstaller_HttpRequestsActuallyCanceled);
+		}
+		else
+		{
+			TRACE_COUNTER_INCREMENT(OnDemandContentInstaller_HttpRequestsFail);
+		}
+#endif
+
+		Result = MakeError(UnifiedError::IoStoreOnDemand::HttpError::MakeError(HttpStatusCode));
+	}
+	else if (ChunkRequest.bChunkCanceled)
+	{
+		Result = MakeError(UnifiedError::Core::CancellationError::MakeError());
+	}
+	else
 	{
 		const FIoHash VerifyChunkHash = FIoHash::HashBuffer(Chunk.GetView());
 		if (ChunkHash == VerifyChunkHash)
 		{
 			TRACE_COUNTER_INCREMENT(OnDemandContentInstaller_HttpRequestsSuccess);
 
-			Result = IoStore.InstallCache->PutChunk(MoveTemp(Chunk), ChunkHash);
-			if (CVars::EnableConcurentInstallRequests())
+			TResult<bool> PutResult = IoStore.InstallCache->PutChunk(MoveTemp(Chunk), ChunkHash);
+			if (PutResult.HasError())
 			{
-				// TODO: Cache needs to track size of blocks in memory so we can quickly compute cache usage
-				// TODO: if cache is full, Defrag here based on size of the pending http chunks
+				Result = MakeError(PutResult.StealError());
+			}
+			else if (PutResult.GetValue() == false)
+			{
+				if (ensureMsgf(CVars::EnableConcurentInstallRequests(), TEXT("Install cache is full but iostore.ContentInstallerEnableConcurrentInstallRequests is not enabled!")))
+				{
+					uint64 PendingDownloadSize = 0;
+					for (const TPair<FIoHash, FChunkDownloadRequest>& Kv : PendingChunkDownloads)
+					{
+						const FChunkDownloadRequest& PendingChunkRequest = Kv.Value;
+						if (PendingChunkRequest.bChunkCanceled)
+						{
+							continue;
+						}
+
+						// Only need to check the first request waiting on the chunk.
+						// The size can come from any container with the chunk.
+						const FChunkHttpRequestHandle* ChunkHttpRequestHandle = PendingChunkRequest.ChunkRequestHttpHandles.PeekHead();
+						if (!ensure(ChunkHttpRequestHandle))
+						{
+							continue;
+						}
+
+						const FRequest& Request = *ChunkHttpRequestHandle->OwnerRequest;
+						const FRequest::FInstall& InstallRequest = Request.AsInstall();
+						const FSharedOnDemandContainer& Container = InstallRequest.ResolvedChunks[ChunkHttpRequestHandle->ContainerIndex].Container;
+						const FOnDemandChunkEntry& ChunkEntry = Container->ChunkEntries[ChunkHttpRequestHandle->EntryIndex];
+
+						PendingDownloadSize += ChunkEntry.GetDiskSize();
+					}
+
+					Result = IoStore.InstallCache->Purge(PendingDownloadSize);
+					if (Result.HasValue())
+					{
+						PutResult = IoStore.InstallCache->PutChunk(MoveTemp(Chunk), ChunkHash);
+						if (PutResult.HasError())
+						{
+							Result = MakeError(PutResult.StealError());
+						}
+						else if (PutResult.GetValue() == false)
+						{
+							Result = MakeError(UE::UnifiedError::IoStoreOnDemand::InstallerCacheFullError::MakeError());
+						}
+					}
+				}
 			}
 		}
 		else
@@ -1199,58 +1279,28 @@ void FOnDemandContentInstaller::ProcessDownloadedChunk(
 				}));
 		}
 	}
+
+	static constexpr const TCHAR LogFmt[] = TEXT("ProcessDownloadedChunk: ChunkHash=%s, RequestId=%u, DownloadResult=%s");
+	if (ChunkRequest.bChunkCanceled)
+	{
+		UE_LOG(LogIoStoreOnDemand, Verbose, LogFmt, *TStringBuilder<64>(InPlace, ChunkHash), ChunkRequest.RequestId, TEXT("Cancelled"));
+	}
+	else if (Result.HasError())
+	{
+		UE_LOG(LogIoStoreOnDemand, Warning, LogFmt, *TStringBuilder<64>(InPlace, ChunkHash), ChunkRequest.RequestId, *Result.GetError().GetErrorMessage(true).ToString());
+	}
 	else
 	{
-#if COUNTERSTRACE_ENABLED
-		if (InErrorCode == EIoErrorCode::Cancelled)
-		{
-			TRACE_COUNTER_INCREMENT(OnDemandContentInstaller_HttpRequestsActuallyCanceled);
-		}
-		else
-		{
-			TRACE_COUNTER_INCREMENT(OnDemandContentInstaller_HttpRequestsFail);
-		}
-#endif
-
-		Result = MakeError(UnifiedError::IoStoreOnDemand::HttpError::MakeError(HttpStatusCode));
+		UE_LOG(LogIoStoreOnDemand, VeryVerbose, LogFmt, *TStringBuilder<64>(InPlace, ChunkHash), ChunkRequest.RequestId, TEXT("OK"));
 	}
 
-	const bool bCancelled = (InErrorCode == EIoErrorCode::Cancelled);
-
-	bool bFoundRequest = false;
-	for (auto It = PendingChunkDownloads.CreateKeyIterator(ChunkHash); It; ++It)
+	while (FChunkHttpRequestHandle* ChunkHttpRequestHandle = ChunkRequest.ChunkRequestHttpHandles.PopHead())
 	{
-		FChunkDownloadRequest& ChunkRequest = It.Value();
-		if (ChunkRequest.RequestId != RequestId)
-		{
-			continue;
-		}
-
-		static constexpr const TCHAR LogFmt[] = TEXT("ProcessDownloadedChunk: ChunkHash=%s, RequestId=%u, DownloadResult=%s");
-		if (bCancelled)
-		{
-			UE_LOG(LogIoStoreOnDemand, Verbose, LogFmt, *TStringBuilder<64>(InPlace, ChunkHash), ChunkRequest.RequestId, TEXT("Cancelled"));
-		}
-		else if (Result.HasError())
-		{
-			UE_LOG(LogIoStoreOnDemand, Warning, LogFmt, *TStringBuilder<64>(InPlace, ChunkHash), ChunkRequest.RequestId, *Result.GetError().GetErrorMessage(true).ToString());
-		}
-		else
-		{
-			UE_LOG(LogIoStoreOnDemand, VeryVerbose, LogFmt, *TStringBuilder<64>(InPlace, ChunkHash), ChunkRequest.RequestId, TEXT("OK"));
-		}
-
-		bFoundRequest = true;
-		while (FChunkHttpRequestHandle* ChunkHttpRequestHandle = ChunkRequest.ChunkRequestHttpHandles.PopHead())
-		{
-			OnProcessDownloadedChunkNotifyRequest(*ChunkHttpRequestHandle, Result, bCancelled);
-		}
-
-		It.RemoveCurrent();
-		break;
+		OnProcessDownloadedChunkNotifyRequest(*ChunkHttpRequestHandle, Result, ChunkRequest.bChunkCanceled);
 	}
 
-	check(bFoundRequest);
+	// Delete from pending downloads
+	It.RemoveCurrent();
 }
 
 void FOnDemandContentInstaller::OnProcessDownloadedChunkNotifyRequest(
diff --git a/Engine/Source/Runtime/Experimental/IoStore/OnDemand/Private/OnDemandContentInstaller.h b/Engine/Source/Runtime/Experimental/IoStore/OnDemand/Private/OnDemandContentInstaller.h
--- a/Engine/Source/Runtime/Experimental/IoStore/OnDemand/Private/OnDemandContentInstaller.h
+++ b/Engine/Source/Runtime/Experimental/IoStore/OnDemand/Private/OnDemandContentInstaller.h
@@ -202,6 +202,12 @@ class FOnDemandContentInstaller
 		FVerify&			AsVerify()			{ return Variant.Get<FVerify>(); }
 		FFlushLastAccess&	AsFlushLastAccess()	{ return Variant.Get<FFlushLastAccess>(); }
 
+		const FInstall&		AsInstall() const { return Variant.Get<FInstall>(); }
+		const FPurge&		AsPurge() const { return Variant.Get<FPurge>(); }
+		const FDefrag&		AsDefrag() const { return Variant.Get<FDefrag>(); }
+		const FVerify&		AsVerify() const { return Variant.Get<FVerify>(); }
+		const FFlushLastAccess& AsFlushLastAccess() const { return Variant.Get<FFlushLastAccess>(); }
+
 		TOptional<UE::UnifiedError::FError> ConsumeError();
 
 		static uint32	NextSeqNo;
diff --git a/Engine/Source/Runtime/Experimental/IoStore/OnDemand/Private/OnDemandInstallCache.cpp b/Engine/Source/Runtime/Experimental/IoStore/OnDemand/Private/OnDemandInstallCache.cpp
--- a/Engine/Source/Runtime/Experimental/IoStore/OnDemand/Private/OnDemandInstallCache.cpp
+++ b/Engine/Source/Runtime/Experimental/IoStore/OnDemand/Private/OnDemandInstallCache.cpp
@@ -67,6 +67,20 @@ namespace CVars
 		GIoStoreOnDemandEnableDefrag,
 		TEXT("Whether to enable defrag when purging")
 	);
+
+	static bool GIoStoreOnDemandUseNewGetBlockInfo = true;
+	static FAutoConsoleVariableRef CVar_IoStoreOnDemandUseNewGetBlockInfo(
+		TEXT("iostore.UseNewGetBlockInfo"),
+		GIoStoreOnDemandUseNewGetBlockInfo,
+		TEXT("Use runtime blocksize tracking for GetBlockInfo")
+	);
+
+	static bool GIoStoreOnDemandVerifyNewGetBlockInfo = false;
+	static FAutoConsoleVariableRef CVar_IoStoreOnDemandVerifyNewGetBlockInfo(
+		TEXT("iostore.VerifyNewGetBlockInfo"),
+		GIoStoreOnDemandVerifyNewGetBlockInfo,
+		TEXT("Verify new GetBlockInfo against old GetBlockInfo")
+	);
 }
 
 ///////////////////////////////////////////////////////////////////////////////
@@ -163,6 +177,14 @@ struct FCasBlockInfo
 	uint64	FileSize = 0;
 	int64	LastAccess = 0;
 	uint64	RefSize = 0;
+
+	bool operator==(const FCasBlockInfo& Other) const
+	{
+		return
+			FileSize == Other.FileSize && 
+			LastAccess == Other.LastAccess && 
+			RefSize == Other.RefSize;
+	}
 };
 
 using FCasBlockInfoMap = TMap<FCasBlockId, FCasBlockInfo>;
@@ -643,6 +665,12 @@ void FCasJournal::FTransaction::CriticalError(FCasJournal::EErrorCode ErrorCode)
 }
 
 ///////////////////////////////////////////////////////////////////////////////
+struct FCasBlockState
+{
+	uint32 Size = 0;
+	uint32 ReadHandleCount = 0;
+};
+
 struct FCasSnapshot;
 struct FCas
 {
@@ -652,7 +680,8 @@ struct FCas
 	using FLookup				= TMap<FCasAddr, FCasLocation>;
 	using FReadHandles			= TMap<FCasBlockId, FWeakAsyncFileHandle>;
 	using FLastAccess			= TMap<FCasBlockId, int64>;
-	using FBlockIdHandleCounts	= TMap<FCasBlockId, int32>;
+	using FBlockStates			= TMap<FCasBlockId, FCasBlockState>;
+	using FBlockFileSizes		= TMap<FCasBlockId, uint64>;
 
 	FCas(const FOnDemandInstallCacheConfig& Config);
 
@@ -662,7 +691,7 @@ struct FCas
 	FResult						Initialize(FStringView Directory, bool bDeleteExisting = false);
 	FCasLocation				FindChunk(const FIoHash& Hash) const;
 	FCasBlockId					CreateBlock();
-	FResult						DeleteBlock(FCasBlockId BlockId, TArray<FCasAddr>& OutAddrs);
+	FResult						DeleteBlock(FCasBlockId BlockId, TFunctionRef<void(const FCasAddr&)> OutAddrs);
 	FString						GetBlockFilename(FCasBlockId BlockId) const;
 	TResult<FSharedFileHandle>	OpenRead(FCasBlockId BlockId);
 	FSharedFileOpenAsyncResult	OpenAsyncRead(FCasBlockId BlockId);
@@ -674,7 +703,16 @@ struct FCas
 		return TrackAccessIf(Type, BlockId, FDateTime::UtcNow().GetTicks(), bDirty);
 	};
 	bool						UnlockedTrackAccessIf(ECasTrackAccessType Type, uint32 BlockIdHash, FCasBlockId BlockId, int64 UtcTicks, bool bDirty);
-	uint64						GetBlockInfo(FCasBlockInfoMap& OutBlockInfo);
+	void						UpdateBlock(
+									FCasBlockId BlockId, 
+									uint32 BytesWritten, 
+									TConstArrayView<FIoHash> ChunkHashes, 
+									TConstArrayView<int64> Offsets,
+									TFunctionRef<void(const FCasAddr&, const FCasLocation&)> OutAddrs);
+	uint64						GetBlockInfoOld(FCasBlockInfoMap& OutBlockInfo) const;
+	uint64						GetBlockInfo(FCasBlockInfoMap& OutBlockInfo) const;
+	uint64						GetTotalBlockSize() const;
+	FBlockFileSizes				FindAllBlockFiles() const;
 	void						Compact();
 	FResult						Verify(TArray<FCasAddr>& OutAddrs);
 	void						LoadSnapshot(FCasSnapshot&& Snapshot);
@@ -690,7 +728,7 @@ struct FCas
 		return Temp;
 	}
 	// Returns the timestamps that are "dirty" and need to be flushed to disk and clears dirty flag
-	FCas::FLastAccess			GetAndClearDirtyLastAccess();
+	FLastAccess					GetAndClearDirtyLastAccess();
 
 	// Returns the place holder timestamp to use if a timestamp is not found in the LastAccess table
 	static int64				GetTimestampForMissingLastAccess();
@@ -713,7 +751,7 @@ private:
 public: // TODO: FIXME: encapsulate these members
 	FLookup					Lookup;
 private:
-	FBlockIdHandleCounts	BlockIds;
+	FBlockStates			BlockIds;
 	FLastAccess				LastAccess;
 	FReadHandles			ReadHandles;
 	FEventRef				BlockReadsDoneEvent;
@@ -805,14 +843,14 @@ FCasBlockId FCas::CreateBlock()
 			continue;
 		}
 
-		BlockIds.Add(BlockId, 0);
+		BlockIds.Add(BlockId);
 		Out = BlockId;
 	}
 
 	return Out;
 }
 
-FResult FCas::DeleteBlock(FCasBlockId BlockId, TArray<FCasAddr>& OutAddrs)
+FResult FCas::DeleteBlock(FCasBlockId BlockId, TFunctionRef<void(const FCasAddr&)> OutAddrs)
 {
 	using namespace UE::UnifiedError::IoStoreOnDemand;
 
@@ -829,7 +867,7 @@ FResult FCas::DeleteBlock(FCasBlockId BlockId, TArray<FCasAddr>& OutAddrs)
 	{
 		Lock.Lock();
 
-		const int32 RequestCount = BlockIds.FindRef(BlockId);
+		const int32 RequestCount = BlockIds.FindRef(BlockId).ReadHandleCount;
 		if (RequestCount)
 		{
 			Lock.Unlock();
@@ -863,7 +901,7 @@ FResult FCas::DeleteBlock(FCasBlockId BlockId, TArray<FCasAddr>& OutAddrs)
 	{
 		if (It->Value.BlockId == BlockId)
 		{
-			OutAddrs.Add(It->Key);
+			OutAddrs(It->Key);
 			It.RemoveCurrent();
 		}
 	}
@@ -896,7 +934,7 @@ TResult<FSharedFileHandle> FCas::OpenRead(FCasBlockId BlockId)
 	FFileOpenResult Result = Ipf.OpenRead(*Filename, IPlatformFile::EOpenReadFlags::AllowWrite);
 	if (Result.HasValue())
 	{
-		BlockIds.FindOrAdd(BlockId, 0)++;
+		BlockIds.FindOrAdd(BlockId).ReadHandleCount++;
 
 		FSharedFileHandle NewHandle(
 			Result.GetValue().Release(),
@@ -932,7 +970,7 @@ FSharedFileOpenAsyncResult FCas::OpenAsyncRead(FCasBlockId BlockId)
 
 	if (HandleResult.HasValue())
 	{
-		BlockIds.FindOrAdd(BlockId, 0)++;
+		BlockIds.FindOrAdd(BlockId).ReadHandleCount++;
 
 		FSharedAsyncFileHandle NewHandle(
 			HandleResult.GetValue().Release(),
@@ -953,7 +991,7 @@ FSharedFileOpenAsyncResult FCas::OpenAsyncRead(FCasBlockId BlockId)
 void FCas::OnFileHandleDeleted(FCasBlockId BlockId)
 {
 	UE::TUniqueLock Lock(Mutex);
-	const int32 Count = --BlockIds.FindChecked(BlockId);
+	const int32 Count = --BlockIds.FindChecked(BlockId).ReadHandleCount;
 	check(Count >= 0);
 	if (Count == 0)
 	{
@@ -1035,7 +1073,34 @@ bool FCas::UnlockedTrackAccessIf(const ECasTrackAccessType Type, const uint32 Bl
 	return false;
 }
 
-uint64 FCas::GetBlockInfo(FCasBlockInfoMap& OutBlockInfo)
+void FCas::UpdateBlock(
+	FCasBlockId BlockId, 
+	uint32 BytesWritten, 
+	TConstArrayView<FIoHash> ChunkHashes, 
+	TConstArrayView<int64> Offsets,
+	TFunctionRef<void(const FCasAddr&, const FCasLocation&)> OutAddrs)
+{
+	check(ChunkHashes.Num() == Offsets.Num());
+	check(BlockId.IsValid());
+
+	TUniqueLock Lock(Mutex);
+
+	BlockIds.FindChecked(BlockId).Size += BytesWritten;
+
+	for (int32 Idx = 0, Count = Offsets.Num(); Idx < Count; ++Idx)
+	{
+		const FCasAddr&	CasAddr = AsCasAddr(ChunkHashes[Idx]);
+		const uint32	ChunkOffset = IntCastChecked<uint32>(Offsets[Idx]);
+
+		FCasLocation& Loc = Lookup.FindOrAdd(CasAddr);
+		Loc.BlockId = BlockId;
+		Loc.BlockOffset = ChunkOffset;
+
+		OutAddrs(CasAddr, Loc);
+	}
+}
+
+uint64 FCas::GetBlockInfoOld(FCasBlockInfoMap& OutBlockInfo) const
 {
 	TStringBuilder<256> Path;
 	FPathViews::Append(Path, RootDirectory, TEXT("blocks"));
@@ -1110,6 +1175,122 @@ uint64 FCas::GetBlockInfo(FCasBlockInfoMap& OutBlockInfo)
 	return Visitor.TotalSize;
 }
 
+uint64 FCas::GetBlockInfo(FCasBlockInfoMap& OutBlockInfo) const
+{
+	if (CVars::GIoStoreOnDemandUseNewGetBlockInfo == false)
+	{
+		return GetBlockInfoOld(OutBlockInfo);
+	}
+
+	const int64	TimestampForMissingLastAccess = FCas::GetTimestampForMissingLastAccess();
+
+	uint64 TotalSize = 0;
+
+	{
+		TUniqueLock Lock(Mutex);
+
+		for (const TPair<FCasBlockId, FCasBlockState>& Kv : BlockIds)
+		{
+			const FCasBlockId BlockId = Kv.Key;
+			const uint32 FileSize = Kv.Value.Size;
+
+			const int64* UtcTicks = LastAccess.Find(BlockId);
+
+			OutBlockInfo.Add(BlockId, FCasBlockInfo
+			{
+				.FileSize = uint64(FileSize),
+				.LastAccess = UtcTicks ? (*UtcTicks & ~FCas::DirtyTimestampMask) : TimestampForMissingLastAccess
+			});
+
+			TotalSize += uint64(FileSize);
+		}
+	}
+
+	if (CVars::GIoStoreOnDemandVerifyNewGetBlockInfo)
+	{
+		FCasBlockInfoMap OldMap;
+		uint64 TotalSizeOld = GetBlockInfoOld(OldMap);
+		check(TotalSize == TotalSizeOld);
+		check(OutBlockInfo.OrderIndependentCompareEqual(OldMap))
+	}
+
+	return TotalSize;
+}
+
+uint64 FCas::GetTotalBlockSize() const
+{
+	uint64 TotalSize = 0;
+
+	TUniqueLock Lock(Mutex);
+
+	for (const TPair<FCasBlockId, FCasBlockState>& Kv : BlockIds)
+	{
+		const uint32 FileSize = Kv.Value.Size;
+		TotalSize += uint64(FileSize);
+	}
+
+	return TotalSize;
+}
+
+FCas::FBlockFileSizes FCas::FindAllBlockFiles() const
+{
+	struct FDirectoryVisitor final
+		: public IPlatformFile::FDirectoryStatVisitor
+	{
+		FDirectoryVisitor(FBlockFileSizes& InBlockFileSizes)
+			: BlockFileSizes(InBlockFileSizes)
+		{
+		}
+
+		virtual bool Visit(const TCHAR* FilenameOrDirectory, const FFileStatData& StatData) override
+		{
+			if (StatData.bIsDirectory)
+			{
+				return true;
+			}
+
+			const FStringView Filename(FilenameOrDirectory);
+			if (FPathViews::GetExtension(Filename) == TEXTVIEW("ucas") == false)
+			{
+				return true;
+			}
+
+			const FStringView	IndexHex = FPathViews::GetBaseFilename(Filename);
+			const FCasBlockId	BlockId(FParse::HexNumber(WriteToString<128>(IndexHex).ToString()));
+
+			if (BlockId.IsValid() == false || StatData.FileSize < 0)
+			{
+				UE_LOG(LogIoStoreOnDemand, Warning, TEXT("Found invalid CAS block '%s', FileSize=%" INT64_FMT),
+					FilenameOrDirectory, StatData.FileSize);
+				return true;
+			}
+
+			if (BlockFileSizes.Contains(BlockId))
+			{
+				UE_LOG(LogIoStoreOnDemand, Warning, TEXT("Found duplicate CAS block '%s'"), FilenameOrDirectory);
+				return true;
+			}
+
+			BlockFileSizes.Add(BlockId, uint64(StatData.FileSize));
+
+			return true;
+		}
+
+		FBlockFileSizes&	BlockFileSizes;
+	};
+
+	TStringBuilder<256> Path;
+	FPathViews::Append(Path, RootDirectory, TEXT("blocks"));
+
+	FBlockFileSizes OutBlockFileSizes;
+
+	IPlatformFile& Ipf = FPlatformFileManager::Get().GetPlatformFile();
+	FDirectoryVisitor Visitor(OutBlockFileSizes);
+	Ipf.IterateDirectoryStat(Path.ToString(), Visitor);
+
+	return OutBlockFileSizes;
+}
+
 void FCas::Compact()
 {
 	UE::TUniqueLock Lock(Mutex);
@@ -1123,10 +1304,9 @@ FResult FCas::Verify(TArray<FCasAddr>& OutAddrs)
 {
 	using namespace UE::UnifiedError::IoStoreOnDemand;
 
-	FCasBlockInfoMap	BlockInfo;
-	const uint64		TotalSize = GetBlockInfo(BlockInfo);
-	uint64				TotalVerifiedBytes = 0;
-	FResult				Result = MakeValue();
+	const FBlockFileSizes	BlockFiles = FindAllBlockFiles();
+	uint64					TotalVerifiedBytes = 0;
+	FResult					Result = MakeValue();
 
 	const int64 Now = FDateTime::UtcNow().GetTicks();
 	const int64 TimestampForMissingLastAccess = FCas::GetTimestampForMissingLastAccess();
@@ -1138,7 +1318,11 @@ FResult FCas::Verify(TArray<FCasAddr>& OutAddrs)
 	for (auto BlockIt = BlockIds.CreateIterator(); BlockIt; ++BlockIt)
 	{
 		const FCasBlockId BlockId = BlockIt->Key;
-		if (const FCasBlockInfo* Info = BlockInfo.Find(BlockId))
+		FCasBlockState& BlockState = BlockIt->Value;
+
+		FString ErrorMessage;
+
+		if (const uint64* FileSize = BlockFiles.Find(BlockId))
 		{
 			if (int64* TimeStamp = LastAccess.Find(BlockId))
 			{
@@ -1177,13 +1361,23 @@ FResult FCas::Verify(TArray<FCasAddr>& OutAddrs)
 				}
 			}
 
-			TotalVerifiedBytes += Info->FileSize;
-			continue;
-		}
+			if (*FileSize <= std::numeric_limits<uint32>::max())
+			{
+				BlockState.Size = uint32(*FileSize);
+				TotalVerifiedBytes += *FileSize;
+				continue;
+			}
 
-		const FString Filename = GetBlockFilename(BlockId);
-		FString ErrorMessage = FString::Printf(TEXT("Missing CAS block '%s'"), *Filename);
-		UE_LOG(LogIoStoreOnDemand, Warning, TEXT("%s"), *ErrorMessage);
+			const FString Filename = GetBlockFilename(BlockId);
+			ErrorMessage = FString::Printf(TEXT("Invalid file size %" UINT64_FMT " for CAS block '%s'"), *FileSize, *Filename);
+			UE_LOG(LogIoStoreOnDemand, Warning, TEXT("%s"), *ErrorMessage);
+		}
+		else
+		{
+			const FString Filename = GetBlockFilename(BlockId);
+			ErrorMessage = FString::Printf(TEXT("Missing CAS block '%s'"), *Filename);
+			UE_LOG(LogIoStoreOnDemand, Warning, TEXT("%s"), *ErrorMessage);
+		}
 
 		LastAccess.Remove(BlockId);
 		BlockIt.RemoveCurrent();
@@ -1193,7 +1387,7 @@ FResult FCas::Verify(TArray<FCasAddr>& OutAddrs)
 	UE_LOG(LogIoStoreOnDemand, Log, TEXT("Verified %d CAS blocks of total %.2lf MiB"),
 		BlockIds.Num(), ToMiB(TotalVerifiedBytes));
 
-	for (const TPair<FCasBlockId, FCasBlockInfo>& Kv : BlockInfo)
+	for (const TPair<FCasBlockId, uint64>& Kv : BlockFiles)
 	{
 		const FCasBlockId BlockId = Kv.Key;
 		if (BlockIds.Contains(BlockId))
@@ -1620,7 +1814,7 @@ void FCas::LoadSnapshot(FCasSnapshot&& Snapshot)
 	LastAccess.Reserve(Snapshot.Blocks.Num());
 	for (const FCasSnapshot::FBlock& Block : Snapshot.Blocks)
 	{
-		BlockIds.Add(Block.BlockId, 0);
+		BlockIds.Add(Block.BlockId);
 		LastAccess.Add(Block.BlockId, Block.LastAccess);
 	}
 }
@@ -1647,7 +1841,7 @@ void FCas::ApplyJournalEntry(const FCasJournal::FEntry& JournalEntry)
 	case FCasJournal::FEntry::EType::BlockCreated:
 	{
 		const FCasJournal::FEntry::FBlockOperation& Op = JournalEntry.BlockOperation;
-		BlockIds.Add(Op.BlockId, 0);
+		BlockIds.Add(Op.BlockId);
 		break;
 	}
 	case FCasJournal::FEntry::EType::BlockDeleted:
@@ -1795,7 +1989,7 @@ public:
 													TConstArrayView<int32> EntryIndices,
 													FOnDemandContentHandle ContentHandle,
 													TArray<int32>& OutMissing) override;
-	virtual FResult								PutChunk(FIoBuffer&& Chunk, const FIoHash& ChunkHash) override;
+	virtual TResult<bool>						PutChunk(FIoBuffer&& Chunk, const FIoHash& ChunkHash) override;
 	virtual FResult								Purge(uint64 BytesToInstall) override;
 	virtual FResult								PurgeAllUnreferenced(bool bDefrag, const uint64* BytesToPurge = nullptr) override;
 	virtual FResult								DefragAll(const uint64* BytesToFree = nullptr) override;
@@ -2371,13 +2565,21 @@ bool FOnDemandInstallCache::TryPinChunks(
 	return true;
 }
 
-FResult FOnDemandInstallCache::PutChunk(FIoBuffer&& Chunk, const FIoHash& ChunkHash)
+TResult<bool> FOnDemandInstallCache::PutChunk(FIoBuffer&& Chunk, const FIoHash& ChunkHash)
 {
 	if (PendingChunks.TotalSize > FPendingChunks::MaxPendingBytes)
 	{
+		const uint64 TotalBlockSize = Cas.GetTotalBlockSize();
+		if (TotalBlockSize + PendingChunks.TotalSize > MaxCacheSize)
+		{
+			// The calling code needs to decide if this is an error.
+			// Cache being full isn't necessarily an error, but we can't put the chunk.
+			return MakeValue(false);
+		}
+
 		if (FResult Result = FlushPendingChunks(PendingChunks); Result.HasError())
 		{
-			return Result;
+			return MakeError(Result.StealError());
 		}
 		check(PendingChunks.IsEmpty());
 	}
@@ -2385,7 +2587,7 @@ FResult FOnDemandInstallCache::PutChunk(FIoBuffer&& Chunk, const FIoHash& ChunkH
 	TUniqueLock Lock(PendingChunks.SharedMutex);
 
 	PendingChunks.Append(MoveTemp(Chunk), ChunkHash);
-	return MakeValue();
+	return MakeValue(true);
 }
 
 FResult FOnDemandInstallCache::Purge(const uint64 BytesToInstall)
@@ -2992,9 +3194,13 @@ FResult FOnDemandInstallCache::PurgeInternal(FCasBlockInfoMap& BlockInfo, const
 		}
 
 		FCasJournal::FTransaction	Transaction = FCasJournal::Begin(GetJournalFilename());
-		TArray<FCasAddr>			RemovedChunks;
 
-		if (FResult Result = Cas.DeleteBlock(BlockId, RemovedChunks); Result.HasError())
+		auto HandleChunkRemoved = [&Transaction](const FCasAddr& Addr)
+		{
+			Transaction.ChunkLocation(FCasLocation::Invalid, Addr);
+		};
+
+		if (FResult Result = Cas.DeleteBlock(BlockId, HandleChunkRemoved); Result.HasError())
 		{
 			return Result;
 		}
@@ -3008,11 +3214,6 @@ FResult FOnDemandInstallCache::PurgeInternal(FCasBlockInfoMap& BlockInfo, const
 		OutTotalPurgedBytes += Info.FileSize;
 
 		It.RemoveCurrent();
-
-		for (const FCasAddr& Addr : RemovedChunks)
-		{
-			Transaction.ChunkLocation(FCasLocation::Invalid, Addr);
-		}
 		Transaction.BlockDeleted(BlockId);
 
 		if (FResult Result = FCasJournal::Commit(MoveTemp(Transaction)); Result.HasError())
@@ -3368,10 +3569,14 @@ FResult FOnDemandInstallCache::Defrag(
 
 		FCasJournal::FTransaction Transaction = FCasJournal::Begin(GetJournalFilename());
 
+		auto HandleChunkRemoved = [&Transaction](const FCasAddr& Addr)
+		{
+			Transaction.ChunkLocation(FCasLocation::Invalid, Addr);
+		};
+
 		// Flushing should overwrite the lookup info for the cas addr to point at the new block.
 		// Can now remove the old block
-		TArray<FCasAddr> DeletedChunkAddresses;
-		if (FResult Result = Cas.DeleteBlock(DefragBlock.BlockId, DeletedChunkAddresses); Result.HasError())
+		if (FResult Result = Cas.DeleteBlock(DefragBlock.BlockId, HandleChunkRemoved); Result.HasError())
 		{
 			FOnDemandInstallCacheStats::OnDefrag(Result, TotalFragmentedBytes, TotalCachedBytes, BlocksToDefragFragmentedSize, BlocksToDefragTotalSize);
 			return Result;
@@ -3379,10 +3584,6 @@ FResult FOnDemandInstallCache::Defrag(
 
 		FOnDemandInstallCacheStats::OnBlockDeleted(DefragBlock.LastAccess, true /*FromDefrag*/);
 
-		for (const FCasAddr& Addr : DeletedChunkAddresses)
-		{
-			Transaction.ChunkLocation(FCasLocation::Invalid, Addr);
-		}
 		Transaction.BlockDeleted(DefragBlock.BlockId);
 
 		if (FResult Result = FCasJournal::Commit(MoveTemp(Transaction)); Result.HasError())
@@ -3660,21 +3861,12 @@ FResult FOnDemandInstallCache::FlushPendingChunksImpl(const FPendingChunks& Chun
 				Transaction.BlockAccess(CurrentBlockId, Now);
 			}
 
-			check(ChunkHashes.Num() == Offsets.Num());
-			check(CurrentBlockId.IsValid());
-			{
-				TUniqueLock Lock(Cas);
-				for (int32 Idx = 0, Count = Offsets.Num(); Idx < Count; ++Idx)
+			Cas.UpdateBlock(CurrentBlockId, IntCastChecked<uint32>(BytesToWrite), ChunkHashes, Offsets,
+				[&Transaction](const FCasAddr& CasAddr, const FCasLocation& Loc)
 				{
-					const FCasAddr	CasAddr = FCasAddr::From(ChunkHashes[Idx]);
-					const uint32	ChunkOffset = IntCastChecked<uint32>(Offsets[Idx]);
-
-					FCasLocation& Loc = Cas.Lookup.FindOrAdd(CasAddr);
-					Loc.BlockId	= CurrentBlockId;
-					Loc.BlockOffset	= ChunkOffset;
 					Transaction.ChunkLocation(Loc, CasAddr);
 				}
-			}
+			);
 
 			Ar.Seek(0);
 		}
diff --git a/Engine/Source/Runtime/Experimental/IoStore/OnDemand/Private/OnDemandInstallCache.h b/Engine/Source/Runtime/Experimental/IoStore/OnDemand/Private/OnDemandInstallCache.h
--- a/Engine/Source/Runtime/Experimental/IoStore/OnDemand/Private/OnDemandInstallCache.h
+++ b/Engine/Source/Runtime/Experimental/IoStore/OnDemand/Private/OnDemandInstallCache.h
@@ -35,7 +35,7 @@ public:
 													TConstArrayView<int32> EntryIndices,
 													FOnDemandContentHandle ContentHandle,
 													TArray<int32>& OutMissing) = 0;
-	virtual FResult								PutChunk(FIoBuffer&& Chunk, const FIoHash& ChunkHash) = 0;
+	virtual TResult<bool>						PutChunk(FIoBuffer&& Chunk, const FIoHash& ChunkHash) = 0;
 	virtual FResult								Purge(uint64 BytesToInstall) = 0;
 	virtual FResult								PurgeAllUnreferenced(bool bDefrag, const uint64* BytesToPurge = nullptr) = 0;
 	virtual FResult								DefragAll(const uint64* BytesToFree = nullptr) = 0;
diff --git a/Engine/Source/Runtime/Experimental/IoStore/OnDemand/Private/OnDemandIoStore.cpp b/Engine/Source/Runtime/Experimental/IoStore/OnDemand/Private/OnDemandIoStore.cpp
--- a/Engine/Source/Runtime/Experimental/IoStore/OnDemand/Private/OnDemandIoStore.cpp
+++ b/Engine/Source/Runtime/Experimental/IoStore/OnDemand/Private/OnDemandIoStore.cpp
@@ -1110,7 +1110,16 @@ FIoStatus FOnDemandIoStore::InitializeInstallCache()
 				}
 
 				const FIoHash ChunkHash = FIoHash::HashBuffer(Chunk.GetView());
-				Result = InstallCache->PutChunk(MoveTemp(Chunk), ChunkHash);
+				TResult<bool> PutResult = InstallCache->PutChunk(MoveTemp(Chunk), ChunkHash);
+				if (PutResult.HasError())
+				{
+					Result = MakeError(PutResult.StealError());
+				}
+				else if (PutResult.GetValue() == false)
+				{
+					Result = MakeError(UE::UnifiedError::IoStoreOnDemand::InstallerCacheFullError::MakeError());
+				}
+
 				Seed++;
 				FillSize -= ChunkSize;
 			}
diff --git a/Engine/Source/Runtime/Experimental/IoStore/OnDemandCore/Private/OnDemandError.cpp b/Engine/Source/Runtime/Experimental/IoStore/OnDemandCore/Private/OnDemandError.cpp
--- a/Engine/Source/Runtime/Experimental/IoStore/OnDemandCore/Private/OnDemandError.cpp
+++ b/Engine/Source/Runtime/Experimental/IoStore/OnDemandCore/Private/OnDemandError.cpp
@@ -15,6 +15,7 @@ UE_DEFINE_ERROR(InstallCacheVerificationError, IoStoreOnDemand);
 UE_DEFINE_ERROR(CasError, IoStoreOnDemand);
 UE_DEFINE_ERROR(CasJournalError, IoStoreOnDemand);
 UE_DEFINE_ERROR(CasSnapshotError, IoStoreOnDemand);
+UE_DEFINE_ERROR(InstallerCacheFullError, IoStoreOnDemand);
 
 namespace UE::UnifiedError::IoStore
 {
diff --git a/Engine/Source/Runtime/Experimental/IoStore/OnDemandCore/Public/IO/OnDemandError.h b/Engine/Source/Runtime/Experimental/IoStore/OnDemandCore/Public/IO/OnDemandError.h
--- a/Engine/Source/Runtime/Experimental/IoStore/OnDemandCore/Public/IO/OnDemandError.h
+++ b/Engine/Source/Runtime/Experimental/IoStore/OnDemandCore/Public/IO/OnDemandError.h
@@ -24,6 +24,7 @@ UE_DECLARE_ERROR(UE_API,			InstallCacheVerificationError,		8,	IoStoreOnDemand, N
 UE_DECLARE_ERROR(UE_API,			CasError,							9,	IoStoreOnDemand, NSLOCTEXT("IoStoreOnDemand", "CasError", "Cas error"));
 UE_DECLARE_ERROR(UE_API,			CasJournalError,					10,	IoStoreOnDemand, NSLOCTEXT("IoStoreOnDemand", "CasJournalError", "Cas journal error"));
 UE_DECLARE_ERROR(UE_API,			CasSnapshotError,					11,	IoStoreOnDemand, NSLOCTEXT("IoStoreOnDemand", "CasSnapshotError", "Cas snapshot error"));
+UE_DECLARE_ERROR(UE_API,			InstallerCacheFullError,			12, IoStoreOnDemand, NSLOCTEXT("IoStoreOnDemand", "InstallerCacheFullError", "Cache full error"));
 
 ////////////////////////////////////////////////////////////////////////////////
 namespace UE::IoStore
