iax: In-flight cancelling

#rnx

#Synced-CL 48897982

diff --git a/Engine/Source/Runtime/Experimental/IoStore/HttpClient/Private/Activity.inl b/Engine/Source/Runtime/Experimental/IoStore/HttpClient/Private/Activity.inl
--- a/Engine/Source/Runtime/Experimental/IoStore/HttpClient/Private/Activity.inl
+++ b/Engine/Source/Runtime/Experimental/IoStore/HttpClient/Private/Activity.inl
@@ -90,6 +90,7 @@ struct FActivity
 	FOutcome			Tick(FHttpPeer& Peer, int32* MaxRecvSize=nullptr);
 	void				SetSink(FTicketSink&& InSink, UPTRINT Param);
 	void				SetDestination(FIoBuffer* InDest);
+	bool				CancelTransaction(FHttpPeer& Peer);
 	void				Cancel();
 	void				Done();
 	void				Fail(const FOutcome& Outcome);
@@ -234,6 +235,18 @@ void FActivity::SetDestination(FIoBuffer* InDest)
 	Dest = InDest;
 }
 
+////////////////////////////////////////////////////////////////////////////////
+bool FActivity::CancelTransaction(FHttpPeer& Peer)
+{
+	if (TransactId == 0)
+	{
+		return true;
+	}
+
+	bool Result = Peer.Cancel(TransactId);
+	return Result;
+}
+
 ////////////////////////////////////////////////////////////////////////////////
 void FActivity::Cancel()
 {
diff --git a/Engine/Source/Runtime/Experimental/IoStore/HttpClient/Private/Loop.inl b/Engine/Source/Runtime/Experimental/IoStore/HttpClient/Private/Loop.inl
--- a/Engine/Source/Runtime/Experimental/IoStore/HttpClient/Private/Loop.inl
+++ b/Engine/Source/Runtime/Experimental/IoStore/HttpClient/Private/Loop.inl
@@ -5,6 +5,14 @@
 namespace UE::IoStore::HTTP
 {
 
+////////////////////////////////////////////////////////////////////////////////
+static bool GInflightCancels = true;
+static FAutoConsoleVariableRef CVar_IaxInflightCancels(
+	TEXT("iax.HttpInflightCancels"),
+	GInflightCancels,
+	TEXT("Allows cancelling of in-flight requests if supported by peer")
+);
+
 // {{{1 recv-stats .............................................................
 
 ////////////////////////////////////////////////////////////////////////////////
@@ -673,6 +681,7 @@ public:
 
 private:
 	bool				MaybeStartNewWork(FTickState& State);
+	void				TickCancels(FTickState& State);
 	void				Negotiate(FTickState& State);
 	void				RecvInternal(FTickState& State);
 	void				SendInternal(FTickState& State);
@@ -837,6 +846,17 @@ void FPeerGroup::RecvInternal(FTickState& State)
 
 		if (Activity == nullptr)
 		{
+			// When an activity is cancelled, its sink is called and the activity retired. If
+			// supported then the peer may also be informed about the cancel and it will not
+			// send any further related frames. Some frames however may already be on the wire
+			// and when received, have no corresponding activity. So here we have to skip over
+			// transactions we are unable to match to an activity.
+			if (GInflightCancels)
+			{
+				Peer.SkipPendingTransactId();
+				return RecvInternal(State);
+			}
+
 			Outcome = FOutcome::Error("could not find work for transact id", TransactId);
 			break;
 		}
@@ -950,6 +970,11 @@ bool FPeerGroup::Tick(FTickState& State)
 		RecvInternal(State);
 	}
 
+	if (!bNegotiating)
+	{
+		TickCancels(State);
+	}
+
 	return !!IsKeepAlive | !(Send.IsEmpty() & Recv.IsEmpty());
 }
 
@@ -1043,6 +1068,50 @@ void FPeerGroup::TickSend(FTickState& State, FHost& Host, FPoller& Poller)
 	bWaiting = true;
 }
 
+////////////////////////////////////////////////////////////////////////////////
+void FPeerGroup::TickCancels(FTickState& State)
+{
+	if (State.Cancels == 0 || !GInflightCancels)
+	{
+		return;
+	}
+
+	auto Impl = [this, &State] (FActivityList& List)
+	{
+		if (List.IsEmpty())
+		{
+			return;
+		}
+
+		for (FActivityNode* Next, *Activity = List.Detach(); Activity != nullptr; Activity = Next)
+		{
+			Next = Activity->Next;
+			Activity->Next = nullptr;
+
+			bool bCancelled = false;
+			if (uint64 Slot = (1ull << Activity->Slot); (State.Cancels & Slot) != 0)
+			{
+				bCancelled = Activity->CancelTransaction(Peer);
+			}
+
+			if (!bCancelled)
+			{
+				List.Prepend(Activity);
+				continue;
+			}
+
+			Activity->Cancel();
+			State.DoneList.Prepend(Activity);
+			--InflightNum;
+		}
+		
+		List.Reverse();
+	};
+
+	Impl(Send);
+	Impl(Recv);
+}
+
 
 
 // {{{1 host-group .............................................................
diff --git a/Engine/Source/Runtime/Experimental/IoStore/HttpClient/Private/Peer.inl b/Engine/Source/Runtime/Experimental/IoStore/HttpClient/Private/Peer.inl
--- a/Engine/Source/Runtime/Experimental/IoStore/HttpClient/Private/Peer.inl
+++ b/Engine/Source/Runtime/Experimental/IoStore/HttpClient/Private/Peer.inl
@@ -545,7 +545,9 @@ public:
 	uint32			GetVersion() const;
 	FOutcome		Handshake();
 	FTransactRef	Transact();
+	bool			Cancel(FTransactId Id);
 	FOutcome		GetPendingTransactId();
+	void			SkipPendingTransactId();
 
 private:
 	void*			PeerData = nullptr;
@@ -638,6 +640,28 @@ FTransactRef FHttpPeer::Transact()
 	}
 }
 
+////////////////////////////////////////////////////////////////////////////////
+bool FHttpPeer::Cancel(FTransactId Id)
+{
+	switch (Proto)
+	{
+	case 1:		return false;
+	case 2:		return CancelTwo(*this, PeerData, Id);
+	default:	return false;
+	}
+}
+
+////////////////////////////////////////////////////////////////////////////////
+void FHttpPeer::SkipPendingTransactId()
+{
+	switch (Proto)
+	{
+	case 1:		return;
+	case 2:		return SkipPendingTransactTwo(*this, PeerData);
+	default:	return;
+	}
+}
+
 ////////////////////////////////////////////////////////////////////////////////
 FOutcome FHttpPeer::GetPendingTransactId()
 {
diff --git a/Engine/Source/Runtime/Experimental/IoStore/HttpClient/Private/TransactionTwo.inl b/Engine/Source/Runtime/Experimental/IoStore/HttpClient/Private/TransactionTwo.inl
--- a/Engine/Source/Runtime/Experimental/IoStore/HttpClient/Private/TransactionTwo.inl
+++ b/Engine/Source/Runtime/Experimental/IoStore/HttpClient/Private/TransactionTwo.inl
@@ -184,6 +184,8 @@ public:
 	FOutcome			Send(FTlsPeer& Peer);
 	FOutcome			RecvResponse(FTlsPeer& Peer, TArray<FHeaderNg>* Sink);
 	FOutcome			RecvBody(FTlsPeer& Peer, FMutableMemoryView Dest);
+	bool				CancelTransaction(FTlsPeer& Peer, FTransactId Id);
+	void				SkipTransaction();
 
 private:
 	FOutcome			Tick(FTlsPeer& Peer);
@@ -629,6 +631,30 @@ FOutcome FDriverNg::RecvBody(FTlsPeer& Peer, FMutableMemoryView Dest)
 	}
 }
 
+////////////////////////////////////////////////////////////////////////////////
+bool FDriverNg::CancelTransaction(FTlsPeer& Peer, FTransactId Id)
+{
+	int32 Result = nghttp2_submit_rst_stream(Session, NGHTTP2_FLAG_NONE, Id, NGHTTP2_CANCEL);
+	if (Result)
+	{
+		return false;
+	}
+
+	PumpSend(Peer);
+	return true;
+}
+
+////////////////////////////////////////////////////////////////////////////////
+void FDriverNg::SkipTransaction()
+{
+	if (!RecvFrame.IsPertinent())
+	{
+		return;
+	}
+
+	RecvFrame = {};
+}
+
 ////////////////////////////////////////////////////////////////////////////////
 int32 FDriverNg::OnHeader(nghttp2_rcbuf* Name, nghttp2_rcbuf* Value)
 {
@@ -940,6 +966,32 @@ FOutcome TickTwo(FTlsPeer& Peer, void* PeerData)
 	return Driver->GetTransactId(Peer);
 }
 
+////////////////////////////////////////////////////////////////////////////////
+bool CancelTwo(FTlsPeer& Peer, void* PeerData, FTransactId Id)
+{
+	using namespace DetailTwo;
+
+	auto* Driver = (FDriverNg*)PeerData;
+	if (Driver != nullptr)
+	{
+		return Driver->CancelTransaction(Peer, Id);
+	}
+
+	return false;
+}
+
+////////////////////////////////////////////////////////////////////////////////
+void SkipPendingTransactTwo(FTlsPeer&, void* PeerData)
+{
+	using namespace DetailTwo;
+
+	auto* Driver = (FDriverNg*)PeerData;
+	if (Driver != nullptr)
+	{
+		return Driver->SkipTransaction();
+	}
+}
+
 ////////////////////////////////////////////////////////////////////////////////
 void GoAwayTwo(FTlsPeer& Peer, void* PeerData)
 {
