iax: null session out when tearing down

#rnx

#Synced-CL 48897982

diff --git a/Engine/Source/Runtime/Experimental/IoStore/HttpClient/Private/TransactionTwo.inl b/Engine/Source/Runtime/Experimental/IoStore/HttpClient/Private/TransactionTwo.inl
--- a/Engine/Source/Runtime/Experimental/IoStore/HttpClient/Private/TransactionTwo.inl
+++ b/Engine/Source/Runtime/Experimental/IoStore/HttpClient/Private/TransactionTwo.inl
@@ -256,12 +256,15 @@ FDriverNg::FDriverNg(FTlsPeer& Peer)
 ////////////////////////////////////////////////////////////////////////////////
 void FDriverNg::GoAway(FTlsPeer& Peer)
 {
-	if (Session != nullptr)
+	if (Session == nullptr)
 	{
-		nghttp2_session_terminate_session(Session, NGHTTP2_NO_ERROR);
-		PumpSend(Peer);
-		nghttp2_session_del(Session);
+		return;
 	}
+
+	nghttp2_session_terminate_session(Session, NGHTTP2_NO_ERROR);
+	PumpSend(Peer);
+	nghttp2_session_del(Session);
+	Session = nullptr;
 }
 
 ////////////////////////////////////////////////////////////////////////////////
