iax: tvos fix

#rnx

#Synced-CL 48897982

diff --git a/Engine/Source/Runtime/Experimental/IoStore/HttpClient/Private/TransactionTwo.inl b/Engine/Source/Runtime/Experimental/IoStore/HttpClient/Private/TransactionTwo.inl
--- a/Engine/Source/Runtime/Experimental/IoStore/HttpClient/Private/TransactionTwo.inl
+++ b/Engine/Source/Runtime/Experimental/IoStore/HttpClient/Private/TransactionTwo.inl
@@ -1035,9 +1035,11 @@ public:
 	FOutcome		TryRecv(FMutableMemoryView, FTlsPeer&)		{ return FOutcome::None(); }
 };
 
-FOutcome	HandshakeTwo(FTlsPeer&, void*&)	{ return FOutcome::None(); }
-FOutcome	TickTwo(FTlsPeer&, void*)		{ return FOutcome::None(); }
-void		GoAwayTwo(FTlsPeer&, void*)		{}
+FOutcome	HandshakeTwo(FTlsPeer&, void*&)			{ return FOutcome::None(); }
+FOutcome	TickTwo(FTlsPeer&, void*)				{ return FOutcome::None(); }
+bool		CancelTwo(FTlsPeer&, void*, FTransactId){ return false; }
+void		SkipPendingTransactTwo(FTlsPeer&, void*){}
+void		GoAwayTwo(FTlsPeer&, void*)				{}
 
 } // namespace UE::IoStore::HTTP
 
