Fix crash when using the cmdline 'Iax.DefaultHostGroup' with a dev cdn that is later filtered.

#rnx

- Currently FOnDemandHostGroup has 4 ways to create the host group, 2 with single urls and 2 with arrays of urls.
- The 2 ways to create via arrays of urls both have their own implementation and duplicate a lot of functionality.
- One of the code paths was fixed to correctly handle all urls being filtered out but the other was not.
- Best fix is to remove the duplication and have one method (array of FString) call the other (array of FAnsiString)

#Synced-CL 48862713

diff --git a/Engine/Source/Runtime/Experimental/IoStore/OnDemandCore/Private/OnDemandHostGroup.cpp b/Engine/Source/Runtime/Experimental/IoStore/OnDemandCore/Private/OnDemandHostGroup.cpp
--- a/Engine/Source/Runtime/Experimental/IoStore/OnDemandCore/Private/OnDemandHostGroup.cpp
+++ b/Engine/Source/Runtime/Experimental/IoStore/OnDemandCore/Private/OnDemandHostGroup.cpp
@@ -222,48 +222,6 @@ TIoStatusOr<FOnDemandHostGroup> FOnDemandHostGroup::Create(TConstArrayView<FAnsi
 		Impl->HostUrls.Emplace(Url);
 	}
 
-	Impl->PrimaryIndex = 0; 
-
-	return FOnDemandHostGroup(MoveTemp(Impl)); 
-}
-
-TIoStatusOr<FOnDemandHostGroup> FOnDemandHostGroup::Create(TConstArrayView<FString> Urls)
-{
-	if (Urls.IsEmpty())
-	{
-		return FIoStatus(EIoErrorCode::InvalidParameter, TEXT("No URLs specified"));
-	}
-
-	FSharedImpl Impl = MakeShared<FImpl>();
-	FString		Reason;
-
-	for (const FString& Url : Urls)
-	{
-		FStringView UrlView = Url;
-
-		if (UrlView.EndsWith(TEXT('/')))
-		{
-			UrlView.RemoveSuffix(1);
-		}
-
-		FAnsiString AnsiUrl(StringCast<ANSICHAR>(UrlView.GetData(), UrlView.Len()));
-
-		if (Private::ValidateUrl(AnsiUrl, Reason) == false)
-		{
-			return FIoStatus(EIoErrorCode::InvalidParameter, Reason);
-		}
-
-#if !UE_BUILD_SHIPPING
-		if (Private::ShouldSkipUrl(UrlView))
-		{
-			UE_LOGFMT(LogIas, Log, "Skipping development CDN '{Url}' when creating HostGroup", UrlView);
-			continue;
-		}
-#endif // !UE_BUILD_SHIPPING
-
-		Impl->HostUrls.Emplace(MoveTemp(AnsiUrl));
-	}
-
 	if (!Impl->HostUrls.IsEmpty())
 	{
 		Impl->PrimaryIndex = 0;
@@ -276,6 +234,17 @@ TIoStatusOr<FOnDemandHostGroup> FOnDemandHostGroup::Create(TConstArrayView<FStri
 	}
 }
 
+TIoStatusOr<FOnDemandHostGroup> FOnDemandHostGroup::Create(TConstArrayView<FString> Urls)
+{
+	TArray<FAnsiString> ConvertedUrls;
+	for (const FString& Url : Urls)
+	{
+		ConvertedUrls.Emplace(FAnsiString(Url));
+	}
+
+	return Create(ConvertedUrls);
+}
+
 FName FOnDemandHostGroup::DefaultName = FName("Default");
 
 } //namespace UE::IoStore
