[Chaos] Updated memory dump command to accept a geometry type filter.

#Synced-CL 48990935

diff --git a/Engine/Source/Runtime/Engine/Private/PhysicsEngine/PhysUtils.cpp b/Engine/Source/Runtime/Engine/Private/PhysicsEngine/PhysUtils.cpp
--- a/Engine/Source/Runtime/Engine/Private/PhysicsEngine/PhysUtils.cpp
+++ b/Engine/Source/Runtime/Engine/Private/PhysicsEngine/PhysUtils.cpp
@@ -342,7 +342,7 @@ bool FPhysicsInterface::ExecPhysCommands(const TCHAR* Cmd, FOutputDevice* Output
 {
 	if (FParse::Command(&Cmd, TEXT("ChaosGeometryMemory")))
 	{
-		Chaos::FTrackedGeometryManager::Get().DumpMemoryUsage(OutputDevice);
+		Chaos::FTrackedGeometryManager::Get().DumpMemoryUsage(Cmd, OutputDevice);
 		return true;
 	}
 
diff --git a/Engine/Source/Runtime/Experimental/Chaos/Private/Chaos/TrackedGeometryManager.cpp b/Engine/Source/Runtime/Experimental/Chaos/Private/Chaos/TrackedGeometryManager.cpp
--- a/Engine/Source/Runtime/Experimental/Chaos/Private/Chaos/TrackedGeometryManager.cpp
+++ b/Engine/Source/Runtime/Experimental/Chaos/Private/Chaos/TrackedGeometryManager.cpp
@@ -12,6 +12,12 @@ namespace Chaos
 	}
 
 	void FTrackedGeometryManager::DumpMemoryUsage(FOutputDevice* Ar) const
+	{
+		const TCHAR* Empty = TEXT("");
+		DumpMemoryUsage(Empty, Ar);
+	}
+
+	void FTrackedGeometryManager::DumpMemoryUsage(const TCHAR* Cmd, FOutputDevice* Ar) const
 	{
 		struct FMemInfo
 		{
@@ -22,17 +28,46 @@ namespace Chaos
 			bool operator<(const FMemInfo& Other) const { return NumBytes < Other.NumBytes; }
 		};
 
+		EImplicitObjectType FilterType = ImplicitObjectType::ConcreteObjectCount;
+
+		if(FParse::Param(Cmd, TEXT("-ListTypes")))
+		{
+			Ar->Logf(TEXT("Available Geometry Types:"));
+			for(uint8 TypeIndex = 0; TypeIndex < ImplicitObjectType::ConcreteObjectCount; ++TypeIndex)
+			{
+				Ar->Logf(TEXT("\t%s"), ImplicitObjectTypeToString(EImplicitObjectType(TypeIndex)));
+			}
+			return;
+		}
+
+		FString FilterString;
+		if(FParse::Value(Cmd, TEXT("TypeFilter="), FilterString))
+		{
+			FilterType = ImplicitObjectTypeFromString(FilterString);
+		}
+
 		TArray<FMemInfo> MemEntries;
 		uint32 TotalBytes = 0;
-		for (const auto& Itr : SharedGeometry)
+		for(const auto& Itr : SharedGeometry)
 		{
+			FImplicitObject* NonConst = const_cast<FImplicitObject*>(Itr.Key.Get());	//only doing this to write out, serialize is non const for read in
+
+			if(FilterType != ImplicitObjectType::ConcreteObjectCount)
+			{
+				if(!(NonConst->GetType() == FilterType || NonConst->GetCollisionType() == FilterType))
+				{
+					continue;
+				}
+			}
+
 			FMemInfo Info;
 			Info.DebugInfo = Itr.Value;
 			TArray<uint8> Data;
 			FMemoryWriter MemAr(Data);
 			FChaosArchive ChaosAr(MemAr);
-			FImplicitObject* NonConst = const_cast<FImplicitObject*>(Itr.Key.Get());	//only doing this to write out, serialize is non const for read in
+
 			NonConst->Serialize(ChaosAr);
+
 			Info.ToString = NonConst->ToString();
 			Info.NumBytes = Data.Num();
 			MemEntries.Add(Info);
@@ -45,7 +80,7 @@ namespace Chaos
 		Ar->Logf(TEXT("Chaos Tracked Geometry:"));
 		Ar->Logf(TEXT(""));
 
-		for (const FMemInfo& Info : MemEntries)
+		for(const FMemInfo& Info : MemEntries)
 		{
 			Ar->Logf(TEXT("%-10d %s ToString:%s"), Info.NumBytes, *Info.DebugInfo, *Info.ToString);
 		}
diff --git a/Engine/Source/Runtime/Experimental/Chaos/Public/Chaos/ImplicitObjectType.h b/Engine/Source/Runtime/Experimental/Chaos/Public/Chaos/ImplicitObjectType.h
--- a/Engine/Source/Runtime/Experimental/Chaos/Public/Chaos/ImplicitObjectType.h
+++ b/Engine/Source/Runtime/Experimental/Chaos/Public/Chaos/ImplicitObjectType.h
@@ -98,11 +98,75 @@ namespace Chaos
 				TEXT("EImplicitObjectType:MLLevelSet"),
 				TEXT("EImplicitObjectType:SkinnedTriangleMesh"),
 				TEXT("EImplicitObjectType::ExtrudedTaperedCapsule")
-			};
+		};
 
 		static_assert(Chaos::ImplicitObjectType::ConcreteObjectCount == UE_ARRAY_COUNT(Names), "");
 
 		return Names[Val];
 	}
 
+	inline const TCHAR* ImplicitObjectTypeToString(EImplicitObjectType Type)
+	{
+		switch(Type)
+		{
+		case ImplicitObjectType::Sphere:
+			return TEXT("Sphere");
+		case ImplicitObjectType::Box:
+			return TEXT("Box");
+		case ImplicitObjectType::Plane:
+			return TEXT("Plane");
+		case ImplicitObjectType::Capsule:
+			return TEXT("Capsule");
+		case ImplicitObjectType::Transformed:
+			return TEXT("Transformed");
+		case ImplicitObjectType::Union:
+			return TEXT("Union");
+		case ImplicitObjectType::LevelSet:
+			return TEXT("LevelSet");
+		case ImplicitObjectType::Unknown:
+			return TEXT("Unknown");
+		case ImplicitObjectType::Convex:
+			return TEXT("Convex");
+		case ImplicitObjectType::TaperedCylinder:
+			return TEXT("TaperedCylinder");
+		case ImplicitObjectType::Cylinder:
+			return TEXT("Cylinder");
+		case ImplicitObjectType::TriangleMesh:
+			return TEXT("TriangleMesh");
+		case ImplicitObjectType::HeightField:
+			return TEXT("HeightField");
+		case ImplicitObjectType::DEPRECATED_Scaled:
+			return TEXT("DEPRECATED_Scaled");
+		case ImplicitObjectType::Triangle:
+			return TEXT("Triangle");
+		case ImplicitObjectType::UnionClustered:
+			return TEXT("UnionClustered");
+		case ImplicitObjectType::TaperedCapsule:
+			return TEXT("TaperedCapsule");
+		case ImplicitObjectType::MLLevelSet:
+			return TEXT("MLLevelSet");
+		case ImplicitObjectType::SkinnedTriangleMesh:
+			return TEXT("SkinnedTriangleMesh");
+		case ImplicitObjectType::ExtrudedTaperedCapsule:
+			return TEXT("ExtrudedTaperedCapsule");
+		default:
+			break;
+		}
+
+		return TEXT("Unknown");
+	}
+
+	inline EImplicitObjectType ImplicitObjectTypeFromString(const FString& InStr)
+	{
+		for(uint8 Index = 0; Index < ImplicitObjectType::ConcreteObjectCount; ++Index)
+		{
+			if(FCString::Strcmp(*InStr, ImplicitObjectTypeToString(Index)) == 0)
+			{
+				return EImplicitObjectType(Index);
+			}
+		}
+
+		return ImplicitObjectType::Sphere;
+	}
+
 }
diff --git a/Engine/Source/Runtime/Experimental/Chaos/Public/Chaos/TrackedGeometryManager.h b/Engine/Source/Runtime/Experimental/Chaos/Public/Chaos/TrackedGeometryManager.h
--- a/Engine/Source/Runtime/Experimental/Chaos/Public/Chaos/TrackedGeometryManager.h
+++ b/Engine/Source/Runtime/Experimental/Chaos/Public/Chaos/TrackedGeometryManager.h
@@ -15,7 +15,9 @@ namespace Chaos
 	public:
 		static CHAOS_API FTrackedGeometryManager& Get();
 		
+		UE_DEPRECATED(5.8, "Use the version taking the remainder of the command")
 		CHAOS_API void DumpMemoryUsage(FOutputDevice* Ar) const;
+		CHAOS_API void DumpMemoryUsage(const TCHAR* Cmd, FOutputDevice* Ar) const;
 		
 	private:
 		TMap<TSerializablePtr<FImplicitObject>, FString> SharedGeometry;
