[Chaos] Fixed HashMappedArray Reset implelementation to not shrink internal allocations

#Synced-CL 48990935

diff --git a/Engine/Source/Runtime/Experimental/Chaos/Public/Chaos/Framework/HashMappedArray.h b/Engine/Source/Runtime/Experimental/Chaos/Public/Chaos/Framework/HashMappedArray.h
--- a/Engine/Source/Runtime/Experimental/Chaos/Public/Chaos/Framework/HashMappedArray.h
+++ b/Engine/Source/Runtime/Experimental/Chaos/Public/Chaos/Framework/HashMappedArray.h
@@ -64,8 +64,16 @@ namespace Chaos::Private
 		// Clear the hash map and reserve space for the specified number of elements (will not shrink)
 		void Reset(const int32 InReserveElements)
 		{
+			check(InReserveElements >= 0);
+
 			HashTable.Clear();
-			HashTable.Resize(InReserveElements);
+
+			// Resize will tightly fit the internal index allocation - only do this if the container grows
+			if((uint32)InReserveElements > HashTable.GetIndexSize())
+			{
+				HashTable.Resize(InReserveElements);
+			}
+
 			Elements.Reset(InReserveElements);
 		}
 
