Fix 'iax.ResetOnScreenStatistics'

#rnx

- The cvar was not resetting the more recently added display items for bytes written to the cache and bytes evicted from the cache.

#Synced-CL 48862713

diff --git a/Engine/Source/Runtime/Experimental/IoStore/OnDemand/Private/Statistics.cpp b/Engine/Source/Runtime/Experimental/IoStore/OnDemand/Private/Statistics.cpp
--- a/Engine/Source/Runtime/Experimental/IoStore/OnDemand/Private/Statistics.cpp
+++ b/Engine/Source/Runtime/Experimental/IoStore/OnDemand/Private/Statistics.cpp
@@ -775,7 +775,9 @@ void FOnDemandIoBackendStats::ResetOnScreenStatistics()
 
 	GCacheRejectBytes.Set(0);
 	GCacheReadBytes.Set(0);
+	GCacheWrittenBytes.Set(0);
 	GCacheGetCount.Set(0);
+	GCacheEvictedBytes.Set(0);
 
 	GCacheDecodeErrorCount.Set(0);
 
