[Chaos] [Dataflow] Added a preview scene controller for physics asset dataflow preview scenes so we can override solver behaviour (iterations, collision behaviour, async/fixed tick)

#Synced-CL 48990935

diff --git a/Engine/Plugins/Experimental/ChaosRigidAsset/Source/ChaosRigidAssetEngine/ChaosRigidAssetEngine.Build.cs b/Engine/Plugins/Experimental/ChaosRigidAsset/Source/ChaosRigidAssetEngine/ChaosRigidAssetEngine.Build.cs
--- a/Engine/Plugins/Experimental/ChaosRigidAsset/Source/ChaosRigidAssetEngine/ChaosRigidAssetEngine.Build.cs
+++ b/Engine/Plugins/Experimental/ChaosRigidAsset/Source/ChaosRigidAssetEngine/ChaosRigidAssetEngine.Build.cs
@@ -27,5 +27,7 @@ public class ChaosRigidAssetEngine : ModuleRules
 		);
 
 		bAllowUETypesInNamespaces = true;
+
+		SetupModulePhysicsSupport(Target);
 	}
 }
diff --git a/Engine/Plugins/Experimental/ChaosRigidAsset/Source/ChaosRigidAssetEngine/Private/PreviewSceneControllerComponent.cpp b/Engine/Plugins/Experimental/ChaosRigidAsset/Source/ChaosRigidAssetEngine/Private/PreviewSceneControllerComponent.cpp
new file mode 100644
--- /dev/null
+++ b/Engine/Plugins/Experimental/ChaosRigidAsset/Source/ChaosRigidAssetEngine/Private/PreviewSceneControllerComponent.cpp
@@ -0,0 +1,44 @@
+// Copyright Epic Games, Inc. All Rights Reserved.
+
+#pragma once
+
+#include "PreviewSceneControllerComponent.h"
+
+#include "Chaos/ChaosScene.h"
+#include "ChaosSolverConfiguration.h"
+#include "Engine/World.h"
+#include "PBDRigidsSolver.h"
+#include "Physics/Experimental/PhysScene_Chaos.h"
+
+void UPreviewSceneControllerComponent::ApplySolverConfig(const FChaosSolverConfiguration& Config)
+{
+	if(UWorld* World = GetWorld())
+	{
+		if(FPhysScene* PhysScene = World->GetPhysicsScene())
+		{
+			PhysScene->GetSolver()->ApplyConfig(Config);
+		}
+	}
+}
+
+void UPreviewSceneControllerComponent::EnableAsyncTick(float InAsyncDt)
+{
+	if(UWorld* World = GetWorld())
+	{
+		if(FPhysScene* PhysScene = World->GetPhysicsScene())
+		{
+			PhysScene->GetSolver()->EnableAsyncMode(InAsyncDt);
+		}
+	}
+}
+
+void UPreviewSceneControllerComponent::DisableAsyncTick()
+{
+	if(UWorld* World = GetWorld())
+	{
+		if(FPhysScene* PhysScene = World->GetPhysicsScene())
+		{
+			PhysScene->GetSolver()->DisableAsyncMode();
+		}
+	}
+}
diff --git a/Engine/Plugins/Experimental/ChaosRigidAsset/Source/ChaosRigidAssetEngine/Public/PreviewSceneControllerComponent.h b/Engine/Plugins/Experimental/ChaosRigidAsset/Source/ChaosRigidAssetEngine/Public/PreviewSceneControllerComponent.h
new file mode 100644
--- /dev/null
+++ b/Engine/Plugins/Experimental/ChaosRigidAsset/Source/ChaosRigidAssetEngine/Public/PreviewSceneControllerComponent.h
@@ -0,0 +1,35 @@
+// Copyright Epic Games, Inc. All Rights Reserved.
+
+#pragma once
+
+#include "Components/ActorComponent.h"
+
+#include "PreviewSceneControllerComponent.generated.h"
+
+#define UE_API CHAOSRIGIDASSETENGINE_API
+
+struct FChaosSolverConfiguration;
+
+/**
+ * Controller for simulation solver interaction in a dataflow preview scene
+ */
+UCLASS(MinimalAPI, meta = (BlueprintSpawnableComponent))
+class UPreviewSceneControllerComponent : public UActorComponent
+{
+	GENERATED_BODY()
+
+public:
+
+	UFUNCTION(BlueprintCallable, Category= PreviewSceneControl)
+	UE_API void ApplySolverConfig(const FChaosSolverConfiguration& Config);
+
+	UFUNCTION(BlueprintCallable, Category= PreviewSceneControl)
+	UE_API void EnableAsyncTick(float InAsyncDt);
+
+	UFUNCTION(BlueprintCallable, Category= PreviewSceneControl)
+	UE_API void DisableAsyncTick();
+
+private:
+};
+
+#undef UE_API
\ No newline at end of file
diff --git a/Engine/Source/Runtime/Experimental/Chaos/Public/ChaosSolverConfiguration.h b/Engine/Source/Runtime/Experimental/Chaos/Public/ChaosSolverConfiguration.h
--- a/Engine/Source/Runtime/Experimental/Chaos/Public/ChaosSolverConfiguration.h
+++ b/Engine/Source/Runtime/Experimental/Chaos/Public/ChaosSolverConfiguration.h
@@ -47,7 +47,7 @@ struct FChaosSolverDestructionSettings
 	bool bOptimizeForRuntimeMemory = false;
 };
 
-USTRUCT()
+USTRUCT(BlueprintType)
 struct FChaosSolverConfiguration
 {
 	GENERATED_BODY();
diff --git a/Engine/Source/Runtime/Experimental/FieldSystem/Source/FieldSystemEngine/Private/Field/FieldSystemComponent.cpp b/Engine/Source/Runtime/Experimental/FieldSystem/Source/FieldSystemEngine/Private/Field/FieldSystemComponent.cpp
--- a/Engine/Source/Runtime/Experimental/FieldSystem/Source/FieldSystemEngine/Private/Field/FieldSystemComponent.cpp
+++ b/Engine/Source/Runtime/Experimental/FieldSystem/Source/FieldSystemEngine/Private/Field/FieldSystemComponent.cpp
@@ -25,6 +25,8 @@ UFieldSystemComponent::UFieldSystemComponent(const FObjectInitializer& ObjectIni
 	, bIsWorldField(false)
 	, bIsChaosField(true)
 	, SupportedSolvers()
+	, bEnableInEditorWorlds(false)
+	, bEnableInPreviewWorlds(false)
 	, ChaosModule(nullptr)
 	, bHasPhysicsState(false)
 	, SetupConstructionFields()
@@ -116,14 +118,21 @@ void UFieldSystemComponent::OnCreatePhysicsState()
 {
 	UActorComponent::OnCreatePhysicsState();
 	
-	const bool bValidWorld = GetWorld() && GetWorld()->IsGameWorld();
-	if(bValidWorld)
+	if(UWorld* World = GetWorld())
 	{
-		// Check we can get a suitable dispatcher
-		ChaosModule = FChaosSolversModule::GetModule();
-		check(ChaosModule);
+		const bool bValidWorld =
+			World->IsGameWorld() ||
+			(bEnableInEditorWorlds && World->IsEditorWorld()) ||
+			(bEnableInPreviewWorlds && World->IsPreviewWorld());
 
-		bHasPhysicsState = true;
+		if(bValidWorld)
+		{
+			// Check we can get a suitable dispatcher
+			ChaosModule = FChaosSolversModule::GetModule();
+			check(ChaosModule);
+
+			bHasPhysicsState = true;
+		}
 	}
 }
 
diff --git a/Engine/Source/Runtime/Experimental/FieldSystem/Source/FieldSystemEngine/Public/Field/FieldSystemComponent.h b/Engine/Source/Runtime/Experimental/FieldSystem/Source/FieldSystemEngine/Public/Field/FieldSystemComponent.h
--- a/Engine/Source/Runtime/Experimental/FieldSystem/Source/FieldSystemEngine/Public/Field/FieldSystemComponent.h
+++ b/Engine/Source/Runtime/Experimental/FieldSystem/Source/FieldSystemEngine/Public/Field/FieldSystemComponent.h
@@ -297,6 +297,14 @@ protected:
 	FIELDSYSTEMENGINE_API virtual bool HasValidPhysicsState() const override;
 	//~ End UActorComponent Interface.
 
+	/** Whether to intialise this component in Editor worlds */
+	UPROPERTY(EditDefaultsOnly, Category=Field)
+	bool bEnableInEditorWorlds;
+
+	/** Whether to intialise this component in Preview worlds (both game and editor) */
+	UPROPERTY(EditDefaultsOnly, Category=Field)
+	bool bEnableInPreviewWorlds;
+
 	/** Chaos module linked to that component */
 	FChaosSolversModule* ChaosModule;
 
