[NetPhysics] Fix crash from using uninitialized pointer.

#Synced-CL 49070933

diff --git a/Engine/Source/Runtime/Experimental/Chaos/Public/Chaos/Framework/PhysicsSolverBase.h b/Engine/Source/Runtime/Experimental/Chaos/Public/Chaos/Framework/PhysicsSolverBase.h
--- a/Engine/Source/Runtime/Experimental/Chaos/Public/Chaos/Framework/PhysicsSolverBase.h
+++ b/Engine/Source/Runtime/Experimental/Chaos/Public/Chaos/Framework/PhysicsSolverBase.h
@@ -520,18 +520,6 @@ namespace Chaos
 			return bUseCollisionResimCache;
 		}
 
-		UE_DEPRECATED(5.6, "Use SetPhysicsReplication_Internal")
-			void SetPhysicsReplication(FPhysicsReplicationAsync* InPhysicsReplication)
-		{
-			PhysicsReplication_DEPRECATED = InPhysicsReplication;
-		}
-
-		UE_DEPRECATED(5.6, "Use GetPhysicsReplication_Internal which return the interface instead of the implementation")
-			FPhysicsReplicationAsync* GetPhysicsReplication()
-		{
-			return PhysicsReplication_DEPRECATED;
-		}
-
 		void SetPhysicsReplication_Internal(IPhysicsReplicationAsync* InPhysicsReplication)
 		{
 			PhysicsReplicationAsync = InPhysicsReplication;
@@ -824,7 +812,6 @@ namespace Chaos
 			return PhysicsHistoryTimeLength;
 		}
 
-
 		/** Get the number of physics history frames to cache */
 		static int32 GetPhysicsHistoryCount()
 		{
@@ -832,13 +819,6 @@ namespace Chaos
 			return PhysicsHistoryCount;
 		}
 
-		UE_DEPRECATED(5.5, "ResimulationErrorThreshold has been renamed, please use GetResimulationErrorPositionThreshold.")
-		static float ResimulationErrorThreshold()
-		{
-			const float ResimulationErrorThreshold = FChaosSolversModule::GetModule()->GetSettingsProvider().GetResimulationErrorPositionThreshold();
-			return ResimulationErrorThreshold;
-		}
-
 		static bool GetResimulationErrorPositionThresholdEnabled()
 		{
 			const bool ResimulationErrorPositionThresholdEnabled = FChaosSolversModule::GetModule()->GetSettingsProvider().GetResimulationErrorPositionThresholdEnabled();
@@ -955,9 +935,8 @@ namespace Chaos
 
 		TUniquePtr<FRewindData> MRewindData;
 		TUniquePtr<IRewindCallback> MRewindCallback;
-		FPhysicsReplicationAsync* PhysicsReplication_DEPRECATED;
-		IPhysicsReplicationAsync* PhysicsReplicationAsync;
-		IPhysicsReplicationLODAsync* PhysicsReplicationLODAsync;
+		IPhysicsReplicationAsync* PhysicsReplicationAsync = nullptr;
+		IPhysicsReplicationLODAsync* PhysicsReplicationLODAsync = nullptr;
 
 		bool bUseCollisionResimCache;
 
