Do not filter an IAX url if it would result in the host group not having any valid urls to use.

#rnx

- Before we filter a url, check to see if it is the last and/or only url we have. If so then log a warning that we were not able to filter it and keep it instead.
-- Technically now we should never end up with an empty hostgroup but I opted to keep the check for this and continue to return EIoErrorCode::InvalidParameter if that condition is hit.
-- This check was removed from the expected code flow and added as an early out, to make it clear that it is not expected.
-Remove the FAnsiStringView overload of FOnDemandHostGroup::Create as it wasn't used and we have too many overloads for this already.

#Synced-CL 48862713

diff --git a/Engine/Source/Runtime/Experimental/IoStore/OnDemandCore/Private/OnDemandHostGroup.cpp b/Engine/Source/Runtime/Experimental/IoStore/OnDemandCore/Private/OnDemandHostGroup.cpp
--- a/Engine/Source/Runtime/Experimental/IoStore/OnDemandCore/Private/OnDemandHostGroup.cpp
+++ b/Engine/Source/Runtime/Experimental/IoStore/OnDemandCore/Private/OnDemandHostGroup.cpp
@@ -162,29 +162,10 @@ bool FOnDemandHostGroup::IsEmpty() const
 	return Impl->HostUrls.IsEmpty();
 }
 
-TIoStatusOr<FOnDemandHostGroup> FOnDemandHostGroup::Create(FAnsiStringView Url)
-{
-	if (Url.EndsWith('/'))
-	{
-		Url.RemoveSuffix(1);
-	}
-
-	FString Reason;
-	if (Private::ValidateUrl(Url, Reason) == false)
-	{
-		return FIoStatus(EIoErrorCode::InvalidParameter, Reason);
-	}
-
-	FSharedImpl Impl = MakeShared<FImpl>();
-	Impl->HostUrls.Emplace(Url);
-	Impl->PrimaryIndex = 0;
-
-	return FOnDemandHostGroup(MoveTemp(Impl)); 
-}
-
 TIoStatusOr<FOnDemandHostGroup> FOnDemandHostGroup::Create(FStringView Url)
 {
-	return Create(FAnsiString(StringCast<ANSICHAR>(Url.GetData(), Url.Len()))); 
+	FAnsiString AnsiUrl(StringCast<ANSICHAR>(Url.GetData(), Url.Len()));
+	return Create(MakeConstArrayView(&AnsiUrl, 1));
 }
 
 TIoStatusOr<FOnDemandHostGroup> FOnDemandHostGroup::Create(TConstArrayView<FAnsiString> Urls)
@@ -197,41 +178,47 @@ TIoStatusOr<FOnDemandHostGroup> FOnDemandHostGroup::Create(TConstArrayView<FAnsi
 	FSharedImpl Impl = MakeShared<FImpl>();
 	FString		Reason;
 
-	for (const FAnsiString& Url : Urls)
+	for( int32 Index = 0; Index < Urls.Num(); ++Index)
 	{
-		FAnsiStringView UrlView = Url;
+		FAnsiStringView Url = Urls[Index];
 
-		if (UrlView.EndsWith('/'))
+		if (Url.EndsWith('/'))
 		{
-			UrlView.RemoveSuffix(1);
+			Url.RemoveSuffix(1);
 		}
 
-		if (Private::ValidateUrl(UrlView, Reason) == false)
+		if (Private::ValidateUrl(Url, Reason) == false)
 		{
 			return FIoStatus(EIoErrorCode::InvalidParameter, Reason);
 		}
 
 #if !UE_BUILD_SHIPPING
-		if (Private::ShouldSkipUrl(UrlView))
+		if (Private::ShouldSkipUrl(Url))
 		{
-			UE_LOGFMT(LogIas, Log, "Skipping development CDN '{Url}' when creating HostGroup", UrlView);
-			continue;
+			if (!Impl->HostUrls.IsEmpty() || Urls.Num() - Index > 1)
+			{
+				UE_LOGFMT(LogIas, Log, "Skipping development CDN '{Url}' when creating HostGroup", Url);
+				continue;
+			}
+			else
+			{
+				UE_LOGFMT(LogIas, Warning, "Cannot skip development CDN '{Url}' as it is the last possible url to create the HostGroup", Url);
+			}
 		}
 #endif // !UE_BUILD_SHIPPING
 
 		Impl->HostUrls.Emplace(Url);
 	}
 
-	if (!Impl->HostUrls.IsEmpty())
-	{
-		Impl->PrimaryIndex = 0;
-		return FOnDemandHostGroup(MoveTemp(Impl));
-	}
-	else
+	if (Impl->HostUrls.IsEmpty())
 	{
+		// This branch should not be possible but leaving it just in case
 		Impl->PrimaryIndex = -1;
 		return FIoStatus(EIoErrorCode::InvalidParameter, TEXT("Alls URLs were filtered out"));
 	}
+	
+	Impl->PrimaryIndex = 0;
+	return FOnDemandHostGroup(MoveTemp(Impl));
 }
 
 TIoStatusOr<FOnDemandHostGroup> FOnDemandHostGroup::Create(TConstArrayView<FString> Urls)
diff --git a/Engine/Source/Runtime/Experimental/IoStore/OnDemandCore/Public/IO/OnDemandHostGroup.h b/Engine/Source/Runtime/Experimental/IoStore/OnDemandCore/Public/IO/OnDemandHostGroup.h
--- a/Engine/Source/Runtime/Experimental/IoStore/OnDemandCore/Public/IO/OnDemandHostGroup.h
+++ b/Engine/Source/Runtime/Experimental/IoStore/OnDemandCore/Public/IO/OnDemandHostGroup.h
@@ -39,8 +39,6 @@ public:
 	/** Returns whether the group is empty or not. */
 	UE_API bool IsEmpty() const;
 	/** Create a new host group with the specified URL. */
-	UE_API static TIoStatusOr<FOnDemandHostGroup> Create(FAnsiStringView Url);
-	/** Create a new host group with the specified URL. */
 	UE_API static TIoStatusOr<FOnDemandHostGroup> Create(FStringView Url);
 	/** Create a new host group with the specified URLs. */
 	UE_API static TIoStatusOr<FOnDemandHostGroup> Create(TConstArrayView<FAnsiString> Urls);
