Fix incorrect format specifier.

#rnx

- Calling ::HandleId returns a UPTRINT and we should not assume it's size as that can vary between platforms.
- Replaced use of %llX with UPTRINT_X_FMT which should always be correct no matter what the platform.

#Synced-CL 48862713

diff --git a/Engine/Source/Runtime/Experimental/IoStore/OnDemandCore/Private/IoStoreOnDemandInternals.cpp b/Engine/Source/Runtime/Experimental/IoStore/OnDemandCore/Private/IoStoreOnDemandInternals.cpp
--- a/Engine/Source/Runtime/Experimental/IoStore/OnDemandCore/Private/IoStoreOnDemandInternals.cpp
+++ b/Engine/Source/Runtime/Experimental/IoStore/OnDemandCore/Private/IoStoreOnDemandInternals.cpp
@@ -130,5 +130,5 @@ void FOnDemandInstallRequest::UpdatePriority(int32 NewPriority)
 ////////////////////////////////////////////////////////////////////////////////
 FString LexToString(const UE::IoStore::FOnDemandInternalContentHandle& Handle)
 {
-	return FString::Printf(TEXT("Id=0x%llX DebugName='%s'"), Handle.HandleId(), *Handle.DebugName);
+	return FString::Printf(TEXT("Id=0x%" UPTRINT_X_FMT " DebugName='%s'"), Handle.HandleId(), *Handle.DebugName);
 }
