Fix out of date code documentation.

#rnx

- The documentation was referencing an old extension, replaced with the new one (iochunktoc)

#Synced-CL 48862713

diff --git a/Engine/Source/Runtime/Experimental/IoStore/OnDemand/Private/Tool/Upload.cpp b/Engine/Source/Runtime/Experimental/IoStore/OnDemand/Private/Tool/Upload.cpp
--- a/Engine/Source/Runtime/Experimental/IoStore/OnDemand/Private/Tool/Upload.cpp
+++ b/Engine/Source/Runtime/Experimental/IoStore/OnDemand/Private/Tool/Upload.cpp
@@ -270,7 +270,7 @@ static void GetTocPath(
 	// The configuration file should specify a service URL without any trailing
 	// host path, i.e. http://{host:port}/{host-path}. Add the trailing path
 	// to the TOC path to form the complete path the TOC from the host, i.e
-	// TocPath={host-path}/{bucket}/{bucket-prefix}/{toc-hash}.uchunktoc
+	// TocPath={host-path}/{bucket}/{bucket-prefix}/{toc-hash}.iochunktoc
 	
 	if (BucketPrefix.StartsWith(TEXT("/")))
 	{
