Added two vector in8 register simd types to chaos (16 and 32).
#Synced-CL 49038256

diff --git a/Engine/Source/Runtime/Experimental/ChaosCore/Public/Chaos/Vector16Int8.h b/Engine/Source/Runtime/Experimental/ChaosCore/Public/Chaos/Vector16Int8.h
new file mode 100644
--- /dev/null
+++ b/Engine/Source/Runtime/Experimental/ChaosCore/Public/Chaos/Vector16Int8.h
@@ -0,0 +1,15 @@
+// Copyright Epic Games, Inc. All Rights Reserved.
+#pragma once
+
+#include "Math/VectorRegister.h"
+
+// Platform specific vector intrinsics include.
+#if PLATFORM_ENABLE_VECTORINTRINSICS_NEON
+#include "Chaos/Vector16Int8Neon.h"
+#elif defined(__cplusplus_cli)
+#include "Chaos/Vector16Int8FPU.h"
+#elif PLATFORM_ENABLE_VECTORINTRINSICS
+#include "Chaos/Vector16Int8SSE.h"
+#else
+#include "Chaos/Vector16Int8FPU.h"
+#endif
diff --git a/Engine/Source/Runtime/Experimental/ChaosCore/Public/Chaos/Vector16Int8FPU.h b/Engine/Source/Runtime/Experimental/ChaosCore/Public/Chaos/Vector16Int8FPU.h
new file mode 100644
--- /dev/null
+++ b/Engine/Source/Runtime/Experimental/ChaosCore/Public/Chaos/Vector16Int8FPU.h
@@ -0,0 +1,68 @@
+// Copyright Epic Games, Inc. All Rights Reserved.
+#pragma once
+
+// HEADER_UNIT_SKIP - Not included directly
+
+#include "Math/VectorRegister.h"
+
+UE_INTERNAL
+namespace Chaos
+{
+	struct VectorRegister16Int8
+	{
+		int8 Value[16];
+	};
+
+	FORCEINLINE VectorRegister16Int8 Vector16Int8Set1(int8 Value)
+	{
+		VectorRegister16Int8 Result;
+		for (uint32 I = 0; I < 16; ++I)
+		{
+			Result.Value[I] = Value;
+		}
+		return Result;
+	}
+
+	FORCEINLINE VectorRegister16Int8 VectorLoad16Int8(const int8* Ptr)
+	{
+		VectorRegister16Int8 Result;
+		for (uint32 I = 0; I < 16; ++I)
+		{
+			Result.Value[I] = Ptr[I];
+		}
+		return Result;
+	}
+
+	FORCEINLINE void VectorStore(VectorRegister16Int8 Vec, int8* Ptr)
+	{
+		for (uint32 I = 0; I < 16; ++I)
+		{
+			Ptr[I] = Vec.Value[I];
+		}
+	}
+
+	FORCEINLINE VectorRegister16Int8 VectorCompareEQ(VectorRegister16Int8 Vec1, VectorRegister16Int8 Vec2)
+	{
+		VectorRegister16Int8 Result;
+		for (uint32 I = 0; I < 16; ++I)
+		{
+			Result.Value[I] = (Vec1.Value[I] == Vec2.Value[I]) ? 0xFF : 0;
+		}
+		return Result;
+	}
+
+	FORCEINLINE uint16 VectorMaskBits(VectorRegister16Int8 VecMask)
+	{
+		constexpr uint8 SignMask = (uint16)(1 << 7);
+		uint16 Result = 0;
+		for (uint16 I = 0; I < 16; ++I)
+		{
+			const uint8 Sign = VecMask.Value[I] & SignMask;
+			if (Sign != 0)
+			{
+				Result |= 1 << I;
+			}
+		}
+		return Result;
+	}
+} // namespace Chaos
diff --git a/Engine/Source/Runtime/Experimental/ChaosCore/Public/Chaos/Vector16Int8Neon.h b/Engine/Source/Runtime/Experimental/ChaosCore/Public/Chaos/Vector16Int8Neon.h
new file mode 100644
--- /dev/null
+++ b/Engine/Source/Runtime/Experimental/ChaosCore/Public/Chaos/Vector16Int8Neon.h
@@ -0,0 +1,59 @@
+// Copyright Epic Games, Inc. All Rights Reserved.
+#pragma once
+
+// HEADER_UNIT_SKIP - Not included directly
+
+#if PLATFORM_ENABLE_VECTORINTRINSICS_NEON
+
+#include "Math/VectorRegister.h"
+
+UE_INTERNAL
+namespace Chaos
+{
+	struct VectorRegister16Int8
+	{
+		int8x16_t Value;
+	};
+
+	FORCEINLINE VectorRegister16Int8 Vector16Int8Set1(int8 Value)
+	{
+		VectorRegister16Int8 Result;
+		Result.Value = vdupq_n_s8(Value);
+		return Result;
+	}
+
+	FORCEINLINE VectorRegister16Int8 VectorLoad16Int8(const int8* Ptr)
+	{
+		VectorRegister16Int8 Result;
+		Result.Value = vld1q_s8(Ptr);
+		return Result;
+	}
+
+	FORCEINLINE void VectorStore(VectorRegister16Int8 Vec, int8* Ptr)
+	{
+		vst1q_s8(Ptr, Vec.Value);
+	}
+
+	FORCEINLINE VectorRegister16Int8 VectorCompareEQ(VectorRegister16Int8 Vec1, VectorRegister16Int8 Vec2)
+	{
+		VectorRegister16Int8 Result;
+		Result.Value = vceqq_s8(Vec1.Value, Vec2.Value);
+		return Result;
+	}
+
+	FORCEINLINE uint16 VectorMaskBits(VectorRegister16Int8 VecMask)
+	{
+		const uint8 Weights[16]{ 1, 2, 4, 8, 16, 32, 64, 128, 1, 2, 4, 8, 16, 32, 64, 128 };
+
+		int32x4_t Signs = vshrq_n_s8(VecMask.Value, 7); // sign bit of each lane replicated 32x
+		uint8x16_t Weighted = vandq_u8(Signs, vld1q_u8(Weights)); // pick bit for lane position
+
+		// Split into low and high and sum across the vectors
+		const uint16 Low = vaddv_u8(vget_low_u8(Weighted));
+		const uint16 High = vaddv_u8(vget_high_u8(Weighted));
+		const uint16 Result = static_cast<uint16>(High << 8u) | Low;
+		return Result;
+	}
+} // namespace Chaos
+
+#endif // PLATFORM_ENABLE_VECTORINTRINSICS_NEON
diff --git a/Engine/Source/Runtime/Experimental/ChaosCore/Public/Chaos/Vector16Int8SSE.h b/Engine/Source/Runtime/Experimental/ChaosCore/Public/Chaos/Vector16Int8SSE.h
new file mode 100644
--- /dev/null
+++ b/Engine/Source/Runtime/Experimental/ChaosCore/Public/Chaos/Vector16Int8SSE.h
@@ -0,0 +1,51 @@
+// Copyright Epic Games, Inc. All Rights Reserved.
+#pragma once
+
+// HEADER_UNIT_SKIP - Not included directly
+
+#if !PLATFORM_ENABLE_VECTORINTRINSICS_NEON && !(defined(__cplusplus_cli)) && PLATFORM_ENABLE_VECTORINTRINSICS
+
+#include "Math/VectorRegister.h"
+
+UE_INTERNAL
+namespace Chaos
+{
+	struct VectorRegister16Int8
+	{
+		VectorRegister4Int Value;
+	};
+
+	FORCEINLINE VectorRegister16Int8 Vector16Int8Set1(int8 Value)
+	{
+		VectorRegister16Int8 Result;
+		Result.Value = _mm_set1_epi8(Value);
+		return Result;
+	}
+
+	FORCEINLINE VectorRegister16Int8 VectorLoad16Int8(const int8* Ptr)
+	{
+		VectorRegister16Int8 Result;
+		Result.Value = _mm_loadu_si128((const __m128i*)(Ptr));
+		return Result;
+	}
+
+	FORCEINLINE void VectorStore(VectorRegister16Int8 Vec, int8* Ptr)
+	{
+		_mm_storeu_si128((__m128i*)(Ptr), Vec.Value);
+	}
+
+	FORCEINLINE VectorRegister16Int8 VectorCompareEQ(VectorRegister16Int8 Vec1, VectorRegister16Int8 Vec2)
+	{
+		VectorRegister16Int8 Result;
+		Result.Value = _mm_cmpeq_epi8(Vec1.Value, Vec2.Value);
+		return Result;
+	}
+
+	FORCEINLINE uint16 VectorMaskBits(VectorRegister16Int8 VecMask)
+	{
+		uint16 Result = (uint16)_mm_movemask_epi8(VecMask.Value);
+		return Result;
+	}
+} // namespace Chaos
+
+#endif // PLATFORM_ENABLE_VECTORINTRINSICS
diff --git a/Engine/Source/Runtime/Experimental/ChaosCore/Public/Chaos/Vector32Int8.h b/Engine/Source/Runtime/Experimental/ChaosCore/Public/Chaos/Vector32Int8.h
new file mode 100644
--- /dev/null
+++ b/Engine/Source/Runtime/Experimental/ChaosCore/Public/Chaos/Vector32Int8.h
@@ -0,0 +1,11 @@
+// Copyright Epic Games, Inc. All Rights Reserved.
+#pragma once
+
+#include "Math/VectorRegister.h"
+
+// Platform specific vector intrinsics include.
+#if !PLATFORM_ENABLE_VECTORINTRINSICS_NEON && !(defined(__cplusplus_cli)) && PLATFORM_ENABLE_VECTORINTRINSICS && UE_PLATFORM_MATH_USE_AVX
+#include "Chaos/Vector32Int8AVX.h"
+#else
+#include "Chaos/Vector32Int8Common.h"
+#endif
diff --git a/Engine/Source/Runtime/Experimental/ChaosCore/Public/Chaos/Vector32Int8AVX.h b/Engine/Source/Runtime/Experimental/ChaosCore/Public/Chaos/Vector32Int8AVX.h
new file mode 100644
--- /dev/null
+++ b/Engine/Source/Runtime/Experimental/ChaosCore/Public/Chaos/Vector32Int8AVX.h
@@ -0,0 +1,73 @@
+// Copyright Epic Games, Inc. All Rights Reserved.
+#pragma once
+
+// HEADER_UNIT_SKIP - Not included directly
+
+#if !PLATFORM_ENABLE_VECTORINTRINSICS_NEON && !(defined(__cplusplus_cli)) && PLATFORM_ENABLE_VECTORINTRINSICS && UE_PLATFORM_MATH_USE_AVX
+
+#include "Chaos/Vector16Int8.h"
+
+UE_INTERNAL
+namespace Chaos
+{
+	struct VectorRegister32Int8
+	{
+		__m256i V;
+	};
+
+	FORCEINLINE VectorRegister32Int8 Vector32Int8Set1(int8 Value)
+	{
+		VectorRegister32Int8 Result;
+		Result.V = _mm256_set1_epi8(Value);
+		return Result;
+	}
+
+	FORCEINLINE VectorRegister32Int8 VectorLoad32Int8(const int8* Ptr)
+	{
+		VectorRegister32Int8 Result;
+		Result.V = _mm256_loadu_si256((const __m256i*)Ptr);
+		return Result;
+	}
+
+	FORCEINLINE void VectorStore(VectorRegister32Int8 Vec, int8* Ptr)
+	{
+		_mm256_storeu_si256((__m256i*)Ptr, Vec.V);
+	}
+
+	FORCEINLINE VectorRegister32Int8 VectorCompareEQ(VectorRegister32Int8 Vec1, VectorRegister32Int8 Vec2)
+	{
+		VectorRegister32Int8 Result;
+#if UE_PLATFORM_MATH_USE_AVX_2
+		Result.V = _mm256_cmpeq_epi8(Vec1.V, Vec2.V);
+#else
+		VectorRegister16Int8 V1Low = { _mm256_extractf128_si256(Vec1.V, 0) };
+		VectorRegister16Int8 V1High = { _mm256_extractf128_si256(Vec1.V, 1) };
+		VectorRegister16Int8 V2Low = { _mm256_extractf128_si256(Vec2.V, 0) };
+		VectorRegister16Int8 V2High = { _mm256_extractf128_si256(Vec2.V, 1) };
+
+		VectorRegister16Int8 CmpLow = VectorCompareEQ(V1Low, V2Low);
+		VectorRegister16Int8 CmpHigh = VectorCompareEQ(V1High, V2High);
+
+		Result.V = _mm256_castsi128_si256(CmpLow.Value);
+		Result.V = _mm256_insertf128_si256(Result.V, CmpHigh.Value, 1);
+#endif
+		return Result;
+	}
+
+	FORCEINLINE uint32 VectorMaskBits(VectorRegister32Int8 VecMask)
+	{
+		uint32 Result = 0;
+#if UE_PLATFORM_MATH_USE_AVX_2
+		Result = _mm256_movemask_epi8(VecMask.V);
+#else
+		VectorRegister16Int8 VecMaskLow = { _mm256_extractf128_si256(VecMask.V, 0) };
+		VectorRegister16Int8 VecMaskHigh = { _mm256_extractf128_si256(VecMask.V, 1) };
+		uint32 Low = (uint32)VectorMaskBits(VecMaskLow);
+		uint32 High = (uint32)VectorMaskBits(VecMaskHigh);
+		Result = (High << 16) | Low;
+#endif
+		return Result;
+	}
+} // namespace Chaos
+
+#endif
diff --git a/Engine/Source/Runtime/Experimental/ChaosCore/Public/Chaos/Vector32Int8Common.h b/Engine/Source/Runtime/Experimental/ChaosCore/Public/Chaos/Vector32Int8Common.h
new file mode 100644
--- /dev/null
+++ b/Engine/Source/Runtime/Experimental/ChaosCore/Public/Chaos/Vector32Int8Common.h
@@ -0,0 +1,53 @@
+// Copyright Epic Games, Inc. All Rights Reserved.
+#pragma once
+
+// HEADER_UNIT_SKIP - Not included directly
+
+#include "Chaos/Vector16Int8.h"
+
+UE_INTERNAL
+namespace Chaos
+{
+	struct VectorRegister32Int8
+	{
+		VectorRegister16Int8 V[2];
+	};
+
+	FORCEINLINE VectorRegister32Int8 Vector32Int8Set1(int8 Value)
+	{
+		VectorRegister32Int8 Result;
+		Result.V[1] = Result.V[0] = Vector16Int8Set1(Value);
+		return Result;
+	}
+
+	FORCEINLINE VectorRegister32Int8 VectorLoad32Int8(const int8* Ptr)
+	{
+		VectorRegister32Int8 Result;
+		Result.V[0] = VectorLoad16Int8(Ptr);
+		Result.V[1] = VectorLoad16Int8(Ptr + 16);
+		return Result;
+	}
+
+	FORCEINLINE void VectorStore(VectorRegister32Int8 Vec, int8* Ptr)
+	{
+		VectorStore(Vec.V[0], Ptr);
+		VectorStore(Vec.V[1], Ptr + 16);
+	}
+
+	FORCEINLINE VectorRegister32Int8 VectorCompareEQ(VectorRegister32Int8 Vec1, VectorRegister32Int8 Vec2)
+	{
+		VectorRegister32Int8 Result;
+		Result.V[0] = VectorCompareEQ(Vec1.V[0], Vec2.V[0]);
+		Result.V[1] = VectorCompareEQ(Vec1.V[1], Vec2.V[1]);
+		return Result;
+	}
+
+	FORCEINLINE uint32 VectorMaskBits(VectorRegister32Int8 VecMask)
+	{
+		uint32 Result = 0;
+		uint32 Low = (uint32)VectorMaskBits(VecMask.V[0]);
+		uint32 High = (uint32)VectorMaskBits(VecMask.V[1]);
+		Result = (High << 16) | Low;
+		return Result;
+	}
+} // namespace Chaos
diff --git a/Engine/Source/Runtime/Experimental/ChaosCore/Tests/Chaos/Vector16Int8Tests.cpp b/Engine/Source/Runtime/Experimental/ChaosCore/Tests/Chaos/Vector16Int8Tests.cpp
new file mode 100644
--- /dev/null
+++ b/Engine/Source/Runtime/Experimental/ChaosCore/Tests/Chaos/Vector16Int8Tests.cpp
@@ -0,0 +1,133 @@
+// Copyright Epic Games, Inc. All Rights Reserved.
+#if WITH_LOW_LEVEL_TESTS
+
+#include "CoreMinimal.h"
+#include "TestHarness.h"
+
+#include "Chaos/Vector16Int8.h"
+
+#include <catch2/generators/catch_generators.hpp>
+
+namespace Chaos
+{
+	TEST_CASE("VectorRegister16Int8::LoadStore", "[Chaos][VectorRegister16Int8]")
+	{
+		const int8 Input[16]{ -128, -7, -6, -5, -4, -3, -2, -1, 0, 1, 2, 3, 4, 5, 6, 127 };
+
+		const VectorRegister16Int8 Register = VectorLoad16Int8(Input);
+		int8 Output[16]{ -1 };
+		VectorStore(Register, Output);
+
+		for (uint8 I = 0; I < 16; ++I)
+		{
+			CAPTURE(I);
+			CHECK(Output[I] == Input[I]);
+		}
+	}
+
+	TEST_CASE("VectorRegister16Int8::Set1", "[Chaos][VectorRegister16Int8]")
+	{
+		const int8 Input = (int8)GENERATE(-1, 0, 1);
+
+		const VectorRegister16Int8 Register = Vector16Int8Set1(Input);
+		int8 Output[16]{ -2 };
+		VectorStore(Register, Output);
+
+		for (uint8 I = 0; I < 16; ++I)
+		{
+			CAPTURE(I);
+			CHECK(Output[I] == Input);
+		}
+	}
+
+	TEST_CASE("VectorRegister16Int8::CompareEQ", "[Chaos][VectorRegister16Int8]")
+	{
+		int8 Output[16]{ -1 };
+
+		SECTION("All Equal")
+		{
+			const int8 Input[16]{ -128, -7, -6, -5, -4, -3, -2, -1, 0, 1, 2, 3, 4, 5, 6, 127 };
+
+			const VectorRegister16Int8 Register = VectorLoad16Int8(Input);
+			const VectorRegister16Int8 Result = VectorCompareEQ(Register, Register);
+			VectorStore(Result, Output);
+
+			for (uint8 I = 0; I < 16; ++I)
+			{
+				CAPTURE(I);
+				CHECK(Output[I] != 0);
+			}
+		}
+		SECTION("Some Equal")
+		{
+			const int8 InputA[16]{ 1, 2, 3, 4, 5, 6, 7, 8, 9, 10, 11, 12, 13, 14, 15, 16 };
+			const int8 InputB[16]{ 1, 0, 3, 0, 5, 0, 7, 0, 9,  0, 11,  0, 13,  0, 15,  0 };
+
+			const VectorRegister16Int8 RegisterA = VectorLoad16Int8(InputA);
+			const VectorRegister16Int8 RegisterB = VectorLoad16Int8(InputB);
+			const VectorRegister16Int8 Result = VectorCompareEQ(RegisterA, RegisterB);
+			VectorStore(Result, Output);
+
+			for (uint8 I = 0; I < 16; ++I)
+			{
+				CAPTURE(I);
+				const bool bExpected = (I % 2) == 0 ? true : false;
+				const bool bActual = Output[I] != 0;
+				CHECK(bExpected == bActual);
+			}
+		}
+		SECTION("None Equal")
+		{
+			const int8 InputA[16]{ 1, 2, 3, 4, 5, 6, 7, 8, 9, 10, 11, 12, 13, 14, 15, 16 };
+			const int8 InputB[16]{ -1, -2, -3, -4, -5, -6, -7, -8, -9, -10, -11, -12, -13, -14, -15, -16 };
+			const VectorRegister16Int8 RegisterA = VectorLoad16Int8(InputA);
+			const VectorRegister16Int8 RegisterB = VectorLoad16Int8(InputB);
+			const VectorRegister16Int8 Result = VectorCompareEQ(RegisterA, RegisterB);
+			VectorStore(Result, Output);
+
+			for (uint8 I = 0; I < 16; ++I)
+			{
+				CAPTURE(I);
+				CHECK(Output[I] == 0);
+			}
+		}
+	}
+
+	TEST_CASE("VectorRegister16Int8::VectorMaskBits", "[Chaos][VectorRegister16Int8]")
+	{
+		const int8 InputA[16]{ 1, 2, 3, 4, 5, 6, 7, 8, 9, 10, 11, 12, 13, 14, 15, 16 };
+		const VectorRegister16Int8 RegisterA = VectorLoad16Int8(InputA);
+		SECTION("Test Each Bit")
+		{
+			for (uint8 I = 0; I < 16; ++I)
+			{
+				CAPTURE(I);
+				const VectorRegister16Int8 RegisterB = Vector16Int8Set1(InputA[I]);
+
+				const VectorRegister16Int8 Result = VectorCompareEQ(RegisterA, RegisterB);
+				const uint16 Actual = VectorMaskBits(Result);
+
+				const uint16 Expected = (uint16)(1 << I);
+				CHECK(Expected == Actual);
+			}
+		}
+		SECTION("All Bits")
+		{
+			const VectorRegister16Int8 Result = VectorCompareEQ(RegisterA, RegisterA);
+			const uint16 Actual = VectorMaskBits(Result);
+
+			CHECK(Actual == 0xFFFF);
+		}
+		SECTION("No Bits")
+		{
+			const VectorRegister16Int8 RegisterB = Vector16Int8Set1(-1);
+
+			const VectorRegister16Int8 Result = VectorCompareEQ(RegisterA, RegisterB);
+			const uint16 Actual = VectorMaskBits(Result);
+
+			CHECK(Actual == 0);
+		}
+	}
+} // namespace Chaos
+
+#endif // WITH_LOW_LEVEL_TESTS
diff --git a/Engine/Source/Runtime/Experimental/ChaosCore/Tests/Chaos/Vector32Int8Tests.cpp b/Engine/Source/Runtime/Experimental/ChaosCore/Tests/Chaos/Vector32Int8Tests.cpp
new file mode 100644
--- /dev/null
+++ b/Engine/Source/Runtime/Experimental/ChaosCore/Tests/Chaos/Vector32Int8Tests.cpp
@@ -0,0 +1,133 @@
+// Copyright Epic Games, Inc. All Rights Reserved.
+#if WITH_LOW_LEVEL_TESTS
+
+#include "CoreMinimal.h"
+#include "TestHarness.h"
+
+#include "Chaos/Vector32Int8.h"
+
+#include <catch2/generators/catch_generators.hpp>
+
+namespace Chaos
+{
+	TEST_CASE("VectorRegister32Int8::LoadStore", "[Chaos][VectorRegister32Int8]")
+	{
+		const int8 Input[32]{ -128, -15, -14, -13, -12, -11, -10, -9, -8, -7, -6, -5, -4, -3, -2, -1, 0, 1, 2, 3, 4, 5, 6, 7, 8, 9, 10, 11, 12, 13, 14, 127 };
+
+		const VectorRegister32Int8 Register = VectorLoad32Int8(Input);
+		int8 Output[32]{ -1 };
+		VectorStore(Register, Output);
+
+		for (uint8 I = 0; I < 32; ++I)
+		{
+			CAPTURE(I);
+			CHECK(Output[I] == Input[I]);
+		}
+	}
+
+	TEST_CASE("VectorRegister32Int8::Set1", "[Chaos][VectorRegister32Int8]")
+	{
+		const int8 Input = (int8)GENERATE(-1, 0, 1);
+
+		const VectorRegister32Int8 Register = Vector32Int8Set1(Input);
+		int8 Output[32]{ -2 };
+		VectorStore(Register, Output);
+
+		for (uint8 I = 0; I < 32; ++I)
+		{
+			CAPTURE(I);
+			CHECK(Output[I] == Input);
+		}
+	}
+
+	TEST_CASE("VectorRegister32Int8::CompareEQ", "[Chaos][VectorRegister32Int8]")
+	{
+		int8 Output[32]{ -1 };
+
+		SECTION("All Equal")
+		{
+			const int8 Input[32]{ -128, -15, -14, -13, -12, -11, -10, -9, -8, -7, -6, -5, -4, -3, -2, -1, 0, 1, 2, 3, 4, 5, 6, 7, 8, 9, 10, 11, 12, 13, 14, 127 };
+
+			const VectorRegister32Int8 Register = VectorLoad32Int8(Input);
+			const VectorRegister32Int8 Result = VectorCompareEQ(Register, Register);
+			VectorStore(Result, Output);
+
+			for (uint8 I = 0; I < 32; ++I)
+			{
+				CAPTURE(I);
+				CHECK(Output[I] != 0);
+			}
+		}
+		SECTION("Some Equal")
+		{
+			const int8 InputA[32]{ 1, 2, 3, 4, 5, 6, 7, 8, 9, 10, 11, 12, 13, 14, 15, 16, 17, 18, 19, 20, 21, 22, 23, 24, 25, 26, 27, 28, 29, 30, 31, 32 };
+			const int8 InputB[32]{ 1, 0, 3, 0, 5, 0, 7, 0, 9, 00, 11, 00, 13, 00, 15, 00, 17, 00, 19, 00, 21, 00, 23, 00, 25, 00, 27, 00, 29, 00, 31, 00 };
+
+			const VectorRegister32Int8 RegisterA = VectorLoad32Int8(InputA);
+			const VectorRegister32Int8 RegisterB = VectorLoad32Int8(InputB);
+			const VectorRegister32Int8 Result = VectorCompareEQ(RegisterA, RegisterB);
+			VectorStore(Result, Output);
+
+			for (uint8 I = 0; I < 32; ++I)
+			{
+				CAPTURE(I);
+				const bool bExpected = (I % 2) == 0 ? true : false;
+				const bool bActual = Output[I] != 0;
+				CHECK(bExpected == bActual);
+			}
+		}
+		SECTION("None Equal")
+		{
+			const int8 InputA[32]{ 1, 2, 3, 4, 5, 6, 7, 8, 9, 10, 11, 12, 13, 14, 15, 16, 17, 18, 19, 20, 21, 22, 23, 24, 25, 26, 27, 28, 29, 30, 31, 32 };
+			const int8 InputB[32]{ -1, -2, -3, -4, -5, -6, -7, -8, -9, -10, -11, -12, -13, -14, -15, -16, -17, -18, -19, -20, -21, -22, -23, -24, -25, -26, -27, -28, -29, -30, -31, -32 };
+			const VectorRegister32Int8 RegisterA = VectorLoad32Int8(InputA);
+			const VectorRegister32Int8 RegisterB = VectorLoad32Int8(InputB);
+			const VectorRegister32Int8 Result = VectorCompareEQ(RegisterA, RegisterB);
+			VectorStore(Result, Output);
+
+			for (uint8 I = 0; I < 32; ++I)
+			{
+				CAPTURE(I);
+				CHECK(Output[I] == 0);
+			}
+		}
+	}
+
+	TEST_CASE("VectorRegister32Int8::VectorMaskBits", "[Chaos][VectorRegister32Int8]")
+	{
+		const int8 InputA[32]{ 1, 2, 3, 4, 5, 6, 7, 8, 9, 10, 11, 12, 13, 14, 15, 16, 17, 18, 19, 20, 21, 22, 23, 24, 25, 26, 27, 28, 29, 30, 31, 32 };
+		const VectorRegister32Int8 RegisterA = VectorLoad32Int8(InputA);
+		SECTION("Test Each Bit")
+		{
+			for (uint8 I = 0; I < 32; ++I)
+			{
+				CAPTURE(I);
+				const VectorRegister32Int8 RegisterB = Vector32Int8Set1(InputA[I]);
+
+				const VectorRegister32Int8 Result = VectorCompareEQ(RegisterA, RegisterB);
+				const uint32 Actual = VectorMaskBits(Result);
+
+				const uint32 Expected = 1u << I;
+				CHECK(Expected == Actual);
+			}
+		}
+		SECTION("All Bits")
+		{
+			const VectorRegister32Int8 Result = VectorCompareEQ(RegisterA, RegisterA);
+			const uint32 Actual = VectorMaskBits(Result);
+
+			CHECK(Actual == 0xFFFFFFFF);
+		}
+		SECTION("No Bits")
+		{
+			const VectorRegister32Int8 RegisterB = Vector32Int8Set1(-1);
+
+			const VectorRegister32Int8 Result = VectorCompareEQ(RegisterA, RegisterB);
+			const uint32 Actual = VectorMaskBits(Result);
+
+			CHECK(Actual == 0);
+		}
+	}
+} // namespace Chaos
+
+#endif // WITH_LOW_LEVEL_TESTS
