Pass through IAS' cache path when querying journalled file system info as there may be more than one mounter, each with different size parameters.

#Synced-CL 48872424

diff --git a/Engine/Source/Runtime/Experimental/IoStore/OnDemand/Private/JournaledCache.cpp b/Engine/Source/Runtime/Experimental/IoStore/OnDemand/Private/JournaledCache.cpp
--- a/Engine/Source/Runtime/Experimental/IoStore/OnDemand/Private/JournaledCache.cpp
+++ b/Engine/Source/Runtime/Experimental/IoStore/OnDemand/Private/JournaledCache.cpp
@@ -80,9 +80,9 @@ enum class EAilments
 #endif
 
 #if IAS_HAS_WRITE_COMMIT_THRESHOLD
-	int32 GetWriteCommitThreshold();
+	int32 GetWriteCommitThreshold(FStringView);
 #else
-	static int32 GetWriteCommitThreshold()
+	static int32 GetWriteCommitThreshold(FStringView)
 	{
 		return 0;
 	}
@@ -2236,7 +2236,7 @@ bool FJournaledCache::Initialize(const TCHAR* RootDir, const FIasCacheConfig& Co
 	Cfg.Demand.Threshold	= Demand.Threshold;
 	Cfg.Demand.Boost		= Demand.Boost;
 	Cfg.Demand.SuperBoost	= Demand.SuperBoost;
-	Cfg.CommitBufferSize	= JournaledCache::GetWriteCommitThreshold();
+	Cfg.CommitBufferSize	= JournaledCache::GetWriteCommitThreshold(CachePath);
 	DiskCacheGovernor.ConfigureStreamingCache(Cfg);
 
 	TUniquePtr<FCache> NewCache = MakeUnique<FCache>(MoveTemp(EventualConfig));
