[TraceBasedDebuggers] refactored and moved TraceSessionsManager from ChaosVD
Main changes:
- settings are now Attributes to assign/bind when creating the manager
- sessions manager creation now requires explicit references since this is no longer setup as a group of related singletons:
- debugger trace runtime module
- engine-editor bridge
- trace module (TraceServices::IModule modular feature)
- FChaosVDTraceSessionDescriptor replaced by UE::TraceBasedDebuggers::FTraceSessionDescriptor
- moved DataStream related types to (FSessionStreamStatsUpdater, IRemoteSessionDataStream, FDirectSocketStream, FRelayDataStream, etc.) under namespace UE::TraceBasedDebuggers

#Synced-CL 49027723

diff --git a/Engine/Plugins/ChaosVD/Source/ChaosVD/ChaosVD.Build.cs b/Engine/Plugins/ChaosVD/Source/ChaosVD/ChaosVD.Build.cs
--- a/Engine/Plugins/ChaosVD/Source/ChaosVD/ChaosVD.Build.cs
+++ b/Engine/Plugins/ChaosVD/Source/ChaosVD/ChaosVD.Build.cs
@@ -34,6 +34,7 @@ public class ChaosVD : ModuleRules
 				"WorkspaceMenuStructure",
 				"TraceAnalysis", 
 				"TraceBasedDebuggers",
+				"TraceBasedDebuggersAnalysis",
 				"TraceInsights",
 				"TraceLog", 
 				"GeometryFramework",
diff --git a/Engine/Plugins/ChaosVD/Source/ChaosVD/Private/ChaosVDEngine.cpp b/Engine/Plugins/ChaosVD/Source/ChaosVD/Private/ChaosVDEngine.cpp
--- a/Engine/Plugins/ChaosVD/Source/ChaosVD/Private/ChaosVDEngine.cpp
+++ b/Engine/Plugins/ChaosVD/Source/ChaosVD/Private/ChaosVDEngine.cpp
@@ -11,10 +11,20 @@
 #include "RemoteSessionsManager.h"
 #include "Settings/ChaosVDMiscSettings.h"
 #include "Trace/ChaosVDCombinedTraceFile.h"
+#include "TraceSessionsManager.h"
 #include "Trace/ChaosVDTraceManager.h"
 
 #include UE_INLINE_GENERATED_CPP_BY_NAME(ChaosVDEngine)
 
+FChaosVDEngine::FChaosVDEngine()
+{
+	InstanceGUID = FGuid::NewGuid();
+}
+
+FChaosVDEngine::~FChaosVDEngine()
+{
+}
+
 void FChaosVDEngine::Initialize()
 {
 	if (bIsInitialized)
@@ -29,9 +39,10 @@ void FChaosVDEngine::Initialize()
 	PlaybackController = MakeShared<FChaosVDPlaybackController>(CurrentScene);
 	PlaybackController->Initialize();
 
-	if (TSharedPtr<UE::TraceBasedDebuggers::FRemoteSessionsManager> RemoteSessionManager = FChaosVDEngineEditorBridge::Get().GetSessionsManager())
+	using namespace UE::TraceBasedDebuggers;
+	if (const TSharedPtr<FRemoteSessionsManager> RemoteSessionManager = FChaosVDEngineEditorBridge::Get().GetSessionsManager())
 	{
-		LiveSessionStoppedDelegateHandle = RemoteSessionManager->OnSessionRecordingStopped().AddLambda([WeakThis = AsWeak()](TWeakPtr<UE::TraceBasedDebuggers::FSessionInfo> Session)
+		LiveSessionStoppedDelegateHandle = RemoteSessionManager->OnSessionRecordingStopped().AddLambda([WeakThis = AsWeak()](TWeakPtr<FSessionInfo> Session)
 		{
 			if (const TSharedPtr<FChaosVDEngine> CVDEngine = WeakThis.Pin())
 			{
@@ -55,11 +66,12 @@ void FChaosVDEngine::Initialize()
 
 void FChaosVDEngine::CloseActiveTraceSessions()
 {
-	if (const TSharedPtr<FChaosVDTraceManager> CVDTraceManager = FChaosVDModule::Get().GetTraceManager())
+	using namespace UE::TraceBasedDebuggers;
+	if (const TSharedPtr<FTraceSessionsManager> CVDTraceManager = FChaosVDModule::Get().GetTraceSessionsManager())
 	{
-		for (FChaosVDTraceSessionDescriptor& SessionDescriptor : CurrentSessionDescriptors)
+		for (FTraceSessionDescriptor& SessionDescriptor : CurrentSessionDescriptors)
 		{
-			OnSessionClosed().Broadcast(SessionDescriptor);
+			OnTraceSessionClosed().Broadcast(SessionDescriptor);
 
 			CVDTraceManager->CloseSession(SessionDescriptor.SessionName);
 			SessionDescriptor.bIsLiveSession = false;
@@ -76,19 +88,20 @@ void FChaosVDEngine::CloseActiveTraceSessions()
 
 void FChaosVDEngine::CloseSessionByRemoteSessionID(FGuid RemoteSessionID)
 {
-	FChaosVDTraceSessionDescriptor DescriptorFoundCopy = GetSessionDescriptor(RemoteSessionID);
+	using namespace UE::TraceBasedDebuggers;
+	FTraceSessionDescriptor DescriptorFoundCopy = GetTraceSessionDescriptor(RemoteSessionID);
 
 	if (!DescriptorFoundCopy.IsValid())
 	{
 		return;
 	}
 
-	if (const TSharedPtr<FChaosVDTraceManager> CVDTraceManager = FChaosVDModule::Get().GetTraceManager())
+	if (const TSharedPtr<FTraceSessionsManager> CVDTraceManager = FChaosVDModule::Get().GetTraceSessionsManager())
 	{
 		CVDTraceManager->CloseSession(DescriptorFoundCopy.SessionName);
 	}
 
-	OnSessionClosed().Broadcast(DescriptorFoundCopy);
+	OnTraceSessionClosed().Broadcast(DescriptorFoundCopy);
 	CurrentSessionDescriptors.Remove(DescriptorFoundCopy);
 
 	if (PlaybackController && CurrentSessionDescriptors.IsEmpty())
@@ -99,9 +112,10 @@ void FChaosVDEngine::CloseSessionByRemoteSessionID(FGuid RemoteSessionID)
 
 void FChaosVDEngine::StopActiveTraceSessions()
 {
-	if (const TSharedPtr<FChaosVDTraceManager> CVDTraceManager = FChaosVDModule::Get().GetTraceManager())
+	using namespace UE::TraceBasedDebuggers;
+	if (const TSharedPtr<FTraceSessionsManager> CVDTraceManager = FChaosVDModule::Get().GetTraceSessionsManager())
 	{
-		for (FChaosVDTraceSessionDescriptor& SessionDescriptor : CurrentSessionDescriptors)
+		for (FTraceSessionDescriptor& SessionDescriptor : CurrentSessionDescriptors)
 		{
 			CVDTraceManager->StopSession(SessionDescriptor.SessionName);
 		}
@@ -122,7 +136,7 @@ void FChaosVDEngine::DeInitialize()
 
 	CloseActiveTraceSessions();
 
-	if (TSharedPtr<UE::TraceBasedDebuggers::FRemoteSessionsManager> RemoteSessionManager = FChaosVDEngineEditorBridge::Get().GetSessionsManager())
+	if (const TSharedPtr<UE::TraceBasedDebuggers::FRemoteSessionsManager> RemoteSessionManager = FChaosVDEngineEditorBridge::Get().GetSessionsManager())
 	{
 		RemoteSessionManager->OnSessionRecordingStopped().Remove(LiveSessionStoppedDelegateHandle);
 	}
@@ -225,10 +239,19 @@ bool FChaosVDEngine::HasAnyLiveSessionActive() const
 	return PlaybackController && PlaybackController->IsPlayingLiveSession();
 }
 
-const FChaosVDTraceSessionDescriptor& FChaosVDEngine::GetSessionDescriptor(FGuid RemoteSessionID) const
+PRAGMA_DISABLE_DEPRECATION_WARNINGS
+const FChaosVDTraceSessionDescriptor& FChaosVDEngine::GetSessionDescriptor(FGuid) const
 {
 	static FChaosVDTraceSessionDescriptor InvalidSessionDescriptor;
-	const FChaosVDTraceSessionDescriptor* DescriptorFound = CurrentSessionDescriptors.FindByPredicate([RemoteSessionID](const FChaosVDTraceSessionDescriptor& InDescriptor)
+	return InvalidSessionDescriptor;
+}
+PRAGMA_ENABLE_DEPRECATION_WARNINGS
+
+const UE::TraceBasedDebuggers::FTraceSessionDescriptor& FChaosVDEngine::GetTraceSessionDescriptor(FGuid RemoteSessionID) const
+{
+	using namespace UE::TraceBasedDebuggers;
+	static FTraceSessionDescriptor InvalidSessionDescriptor;
+	const FTraceSessionDescriptor* DescriptorFound = CurrentSessionDescriptors.FindByPredicate([RemoteSessionID](const FTraceSessionDescriptor& InDescriptor)
 	{
 		return RemoteSessionID == InDescriptor.RemoteSessionID;
 	});
@@ -236,14 +259,17 @@ const FChaosVDTraceSessionDescriptor& FChaosVDEngine::GetSessionDescriptor(FGuid
 	return DescriptorFound ? *DescriptorFound : InvalidSessionDescriptor;
 }
 
-void FChaosVDEngine::LoadRecording_Internal(const TFunction<FString(const TSharedPtr<FChaosVDRecording>&)>& LoadCallback, EChaosVDLoadRecordedDataMode LoadingMode)
+TArrayView<UE::TraceBasedDebuggers::FTraceSessionDescriptor> FChaosVDEngine::GetTraceSessionDescriptors()
 {
-	FChaosVDTraceSessionDescriptor NewSessionFromFileDescriptor;
+	return CurrentSessionDescriptors;
+}
 
+void FChaosVDEngine::LoadRecording_Internal(const TFunction<FString(const TSharedPtr<FChaosVDRecording>&)>& LoadCallback, EChaosVDLoadRecordedDataMode LoadingMode)
+{
 	TSharedPtr<FChaosVDRecording> ExistingRecordingInstance = nullptr;
 	if (LoadingMode == EChaosVDLoadRecordedDataMode::MultiSource && !CurrentSessionDescriptors.IsEmpty())
 	{
-		if (const TSharedPtr<const TraceServices::IAnalysisSession> TraceSession = FChaosVDModule::Get().GetTraceManager()->GetSession(CurrentSessionDescriptors[0].SessionName))
+		if (const TSharedPtr<const TraceServices::IAnalysisSession> TraceSession = FChaosVDModule::Get().GetTraceSessionsManager()->GetSession(CurrentSessionDescriptors[0].SessionName))
 		{
 			if (const FChaosVDTraceProvider* ChaosVDProvider = TraceSession->ReadProvider<FChaosVDTraceProvider>(FChaosVDTraceProvider::ProviderName))
 			{
@@ -251,9 +277,10 @@ void FChaosVDEngine::LoadRecording_Internal(const TFunction<FString(const TShare
 			}
 		}
 	}
-	
+
+	using namespace UE::TraceBasedDebuggers;
+	FTraceSessionDescriptor NewSessionFromFileDescriptor;
 	NewSessionFromFileDescriptor.SessionName = LoadCallback(ExistingRecordingInstance);
-	
 	NewSessionFromFileDescriptor.bIsLiveSession = false;
 
 	if (NewSessionFromFileDescriptor.IsValid())
@@ -275,12 +302,16 @@ void FChaosVDEngine::LoadRecording(const FString& FilePath, EChaosVDLoadRecorded
 		return;
 	}
 
-	auto LoadSingleFileFromTraceManager = [&FilePath](const TSharedPtr<FChaosVDRecording>& ExistingRecordingInstance)
+	auto LoadSingleFilePathFromTraceManager = [&FilePath](const TSharedPtr<FChaosVDRecording>& ExistingRecordingInstance)
 	{
-		return FChaosVDModule::Get().GetTraceManager()->LoadTraceFile(FilePath, ExistingRecordingInstance);
+		return FChaosVDModule::Get().GetTraceSessionsManager()->LoadTraceFile(FilePath, /*PreStartAnalysisCallback*/
+			[ExistingRecordingInstance]
+			{
+				FChaosVDTraceManagerThreadContext::Get().PendingExternalRecordingWeakPtr = ExistingRecordingInstance;
+			});
 	};
-	
-	LoadRecording_Internal(LoadSingleFileFromTraceManager, LoadingMode);
+
+	LoadRecording_Internal(LoadSingleFilePathFromTraceManager, LoadingMode);
 	UpdateRecentFilesList(FilePath);
 }
 
@@ -309,7 +340,11 @@ void FChaosVDEngine::LoadCombinedMultiRecording(const FString& FilePath)
 		FString SessionName = FString::Format(TEXT("{0}-{1}"), Args);
 		auto LoadSingleFileHandleFromTraceManager = [&Handle, &SessionName](const TSharedPtr<FChaosVDRecording>& ExistingRecordingInstance)
 		{
-			return FChaosVDModule::Get().GetTraceManager()->LoadTraceFile(MoveTemp(Handle), SessionName, ExistingRecordingInstance);
+			return FChaosVDModule::Get().GetTraceSessionsManager()->LoadTraceFile(MoveTemp(Handle), SessionName,  /*PreStartAnalysisCallback*/
+				[ExistingRecordingInstance]
+				{
+					FChaosVDTraceManagerThreadContext::Get().PendingExternalRecordingWeakPtr = ExistingRecordingInstance;
+				});
 		};
 
 		LoadRecording_Internal(LoadSingleFileHandleFromTraceManager, EChaosVDLoadRecordedDataMode::MultiSource);
@@ -326,7 +361,8 @@ bool FChaosVDEngine::ConnectToLiveSession(uint32 SessionID, const FString& InSes
 		CloseActiveTraceSessions();
 	}
 	
-	FChaosVDTraceSessionDescriptor NewSessionFromFileDescriptor;
+	using namespace UE::TraceBasedDebuggers;
+	FTraceSessionDescriptor NewSessionFromFileDescriptor;
 
 	NewSessionFromFileDescriptor.bIsLiveSession = true;
 	
@@ -336,7 +372,13 @@ bool FChaosVDEngine::ConnectToLiveSession(uint32 SessionID, const FString& InSes
 		ExistingRecordingInstance = PlaybackController->GetCurrentRecording().Pin();
 	}
 
-	NewSessionFromFileDescriptor.SessionName = FChaosVDModule::Get().GetTraceManager()->ConnectToLiveSession(InSessionAddress, SessionID, ExistingRecordingInstance);
+	auto PreStartAnalysisCallback = [ExistingRecordingInstance]
+		{
+			FChaosVDTraceManagerThreadContext::Get().PendingExternalRecordingWeakPtr = ExistingRecordingInstance;
+		};
+
+	NewSessionFromFileDescriptor.SessionName =
+		FChaosVDModule::Get().GetTraceSessionsManager()->ConnectToLiveSession(InSessionAddress, SessionID, PreStartAnalysisCallback);
 
 	bool bSuccess = false;
 
@@ -352,7 +394,7 @@ bool FChaosVDEngine::ConnectToLiveSession(uint32 SessionID, const FString& InSes
 
 bool FChaosVDEngine::ConnectToLiveSession_Direct(EChaosVDLoadRecordedDataMode LoadingMode)
 {
-	UE_LOG(LogChaosVDEditor, Error, TEXT("[%hs] This menthod is no longer supported. Please use the overload that takes a Remote Session ID."), __func__);
+	UE_LOG(LogChaosVDEditor, Error, TEXT("[%hs] This method is no longer supported. Please use the overload that takes a Remote Session ID."), __func__);
 	return false;
 }
 
@@ -363,18 +405,25 @@ bool FChaosVDEngine::ConnectToLiveSession_Direct(FGuid RemoteSessionID, EChaosVD
 		CloseActiveTraceSessions();
 	}
 
-	FChaosVDTraceSessionDescriptor NewSessionFromFileDescriptor;
+	using namespace UE::TraceBasedDebuggers;
+	FTraceSessionDescriptor NewSessionFromFileDescriptor;
 
 	NewSessionFromFileDescriptor.RemoteSessionID = RemoteSessionID;
 	NewSessionFromFileDescriptor.bIsLiveSession = true;
-	
+
 	TSharedPtr<FChaosVDRecording> ExistingRecordingInstance = nullptr;
 	if (LoadingMode == EChaosVDLoadRecordedDataMode::MultiSource && PlaybackController)
 	{
 		ExistingRecordingInstance = PlaybackController->GetCurrentRecording().Pin();
 	}
 
-	NewSessionFromFileDescriptor.SessionName = FChaosVDModule::Get().GetTraceManager()->ConnectToLiveSession_Direct(RemoteSessionID, NewSessionFromFileDescriptor.SessionPort, ExistingRecordingInstance);
+	auto PreStartAnalysisCallback = [ExistingRecordingInstance]
+		{
+			FChaosVDTraceManagerThreadContext::Get().PendingExternalRecordingWeakPtr = ExistingRecordingInstance;
+		};
+
+	NewSessionFromFileDescriptor.SessionName =
+		FChaosVDModule::Get().GetTraceSessionsManager()->ConnectToLiveSession_Direct(RemoteSessionID, NewSessionFromFileDescriptor.SessionPort, PreStartAnalysisCallback);
 
 	bool bSuccess = false;
 
@@ -394,18 +443,25 @@ bool FChaosVDEngine::ConnectToLiveSession_Relay(FGuid RemoteSessionID, EChaosVDL
 		CloseActiveTraceSessions();
 	}
 
-	FChaosVDTraceSessionDescriptor NewSessionFromFileDescriptor;
+	using namespace UE::TraceBasedDebuggers;
+	FTraceSessionDescriptor NewSessionFromFileDescriptor;
 
 	NewSessionFromFileDescriptor.bIsLiveSession = true;
-	
+
 	TSharedPtr<FChaosVDRecording> ExistingRecordingInstance = nullptr;
 	if (LoadingMode == EChaosVDLoadRecordedDataMode::MultiSource && PlaybackController)
 	{
 		ExistingRecordingInstance = PlaybackController->GetCurrentRecording().Pin();
 	}
 
+	auto PreStartAnalysisCallback = [ExistingRecordingInstance]
+		{
+			FChaosVDTraceManagerThreadContext::Get().PendingExternalRecordingWeakPtr = ExistingRecordingInstance;
+		};
+
 	NewSessionFromFileDescriptor.RemoteSessionID = RemoteSessionID;
-	NewSessionFromFileDescriptor.SessionName = FChaosVDModule::Get().GetTraceManager()->ConnectToLiveSession_Relay(RemoteSessionID, ExistingRecordingInstance);
+	NewSessionFromFileDescriptor.SessionName =
+		FChaosVDModule::Get().GetTraceSessionsManager()->ConnectToLiveSession_Relay(RemoteSessionID, PreStartAnalysisCallback);
 
 	bool bSuccess = false;
 
@@ -418,7 +474,7 @@ bool FChaosVDEngine::ConnectToLiveSession_Relay(FGuid RemoteSessionID, EChaosVDL
 	return bSuccess;
 }
 
-void FChaosVDEngine::OpenSession(const FChaosVDTraceSessionDescriptor& SessionDescriptor, EChaosVDLoadRecordedDataMode LoadingMode)
+void FChaosVDEngine::OpenSession(const UE::TraceBasedDebuggers::FTraceSessionDescriptor& SessionDescriptor, EChaosVDLoadRecordedDataMode LoadingMode)
 {
 	CurrentSessionDescriptors.Emplace(SessionDescriptor);
 
@@ -427,7 +483,7 @@ void FChaosVDEngine::OpenSession(const FChaosVDTraceSessionDescriptor& SessionDe
 		PlaybackController->LoadChaosVDRecordingFromTraceSession(SessionDescriptor);
 	}
 
-	OnSessionOpened().Broadcast(SessionDescriptor);
+	OnTraceSessionOpened().Broadcast(SessionDescriptor);
 }
 
 bool FChaosVDEngine::SaveOpenSessionToCombinedFile(const FString& InTargetFilePath)
@@ -442,7 +498,8 @@ bool FChaosVDEngine::SaveOpenSessionToCombinedFile(const FString& InTargetFilePa
 	TArray<TUniquePtr<IFileHandle>> FilePathsToCombine;
 	FilePathsToCombine.Reserve(CurrentSessionDescriptors.Num());
 
-	for (FChaosVDTraceSessionDescriptor& SessionDescriptor : CurrentSessionDescriptors)
+	using namespace UE::TraceBasedDebuggers;
+	for (FTraceSessionDescriptor& SessionDescriptor : CurrentSessionDescriptors)
 	{
 		if (SessionDescriptor.SessionName.EndsWith(TEXT("cvdmulti")))
 		{
@@ -482,7 +539,8 @@ bool FChaosVDEngine::CanCombineOpenSessions() const
 {
 	if (CurrentSessionDescriptors.Num() > 1)
 	{
-		for (const FChaosVDTraceSessionDescriptor& Session : CurrentSessionDescriptors)
+		using namespace UE::TraceBasedDebuggers;
+		for (const FTraceSessionDescriptor& Session : CurrentSessionDescriptors)
 		{
 			if (Session.SessionName.EndsWith(TEXT("cvdmulti")))
 			{
diff --git a/Engine/Plugins/ChaosVD/Source/ChaosVD/Private/ChaosVDModule.cpp b/Engine/Plugins/ChaosVD/Source/ChaosVD/Private/ChaosVDModule.cpp
--- a/Engine/Plugins/ChaosVD/Source/ChaosVD/Private/ChaosVDModule.cpp
+++ b/Engine/Plugins/ChaosVD/Source/ChaosVD/Private/ChaosVDModule.cpp
@@ -1,28 +1,31 @@
-// Copyright Epic Games, Inc. All Rights Reserved.
+﻿// Copyright Epic Games, Inc. All Rights Reserved.
 
 #include "ChaosVDModule.h"
 
+#include "Chaos/ChaosVDEngineEditorBridge.h"
 #include "ChaosVDCommands.h"
 #include "ChaosVDEngine.h"
 #include "ChaosVDPlaybackController.h"
+#include "ChaosVDRuntimeModule.h"
 #include "ChaosVDSceneCompositionReport.h"
 #include "ChaosVDSceneParticleCustomization.h"
 #include "ChaosVDSettingsManager.h"
 #include "ChaosVDStyle.h"
 #include "ChaosVDTabsIDs.h"
 #include "DesktopPlatformModule.h"
+#include "ExtensionsSystem/ChaosVDExtensionsManager.h"
+#include "Framework/Application/SlateApplication.h"
+#include "HAL/FileManager.h"
 #include "IDesktopPlatform.h"
 #include "Misc/App.h"
 #include "Misc/Guid.h"
-#include "Trace/ChaosVDTraceManager.h"
+#include "Settings/ChaosVDGeneralSettings.h"
+#include "Trace/ChaosVDTraceModule.h"
+#include "TraceSessionsManager.h"
 #include "Widgets/Docking/SDockTab.h"
 #include "Widgets/SChaosVDMainTab.h"
 #include "WorkspaceMenuStructure.h"
 #include "WorkspaceMenuStructureModule.h"
-#include "DetailsCustomizations/ChaosVDShapeDataCustomization.h"
-#include "ExtensionsSystem/ChaosVDExtensionsManager.h"
-#include "Framework/Application/SlateApplication.h"
-#include "HAL/FileManager.h"
 
 #define LOCTEXT_NAMESPACE "ChaosVisualDebugger"
 
@@ -45,9 +48,9 @@ FChaosVDModule& FChaosVDModule::Get()
 }
 
 void FChaosVDModule::StartupModule()
-{	
+{
 	FChaosVDStyle::Initialize();
-	
+
 	FChaosVDCommands::Register();
 
 	FGlobalTabmanager::Get()->RegisterNomadTabSpawner(FChaosVDTabID::ChaosVisualDebuggerTab, FOnSpawnTab::CreateRaw(this, &FChaosVDModule::SpawnMainTab))
@@ -56,7 +59,33 @@ void FChaosVDModule::StartupModule()
 								.SetIcon(FSlateIcon(FChaosVDStyle::GetStyleSetName(), "ChaosVisualDebugger"))
 								.SetGroup(WorkspaceMenu::GetMenuStructure().GetDeveloperToolsDebugCategory());
 
-	ChaosVDTraceManager = MakeShared<FChaosVDTraceManager>();
+	ChaosVDTraceManager = MakeShared<UE::TraceBasedDebuggers::FTraceSessionsManager>(
+		FChaosVDRuntimeModule::Get()
+		, FChaosVDEngineEditorBridge::Get()
+		, MakeShared<FChaosVDTraceModule>()
+		, /*DefaultSavePath*/TEXT("ChaosVisualDebugger"));
+
+	// Set options using lambdas so the most recent values from the settings are used
+	// (i.e., users can change settings during the session while the trace manager is created only once)
+	ChaosVDTraceManager->SetSaveRecordingsToDisk(MakeAttributeLambda([]
+		{
+			if (const UChaosVDGeneralSettings* Settings = FChaosVDSettingsManager::Get().GetSettingsObject<UChaosVDGeneralSettings>())
+			{
+				return Settings->bSaveRecordingsToDisk;
+			}
+			return false;
+		}));
+
+	ChaosVDTraceManager->SetSavePathOverride(MakeAttributeLambda([]
+		{
+			if (const UChaosVDGeneralSettings* Settings = FChaosVDSettingsManager::Get().GetSettingsObject<UChaosVDGeneralSettings>())
+			{
+				return Settings->DefaultSavePathOverride;
+			}
+			return FFilePath{};
+		}));
+
+	UE::Trace::ToggleChannel(TEXT("ChaosVD"), /*bEnabled*/true);
 
 	if (IsStandaloneChaosVisualDebugger())
 	{
diff --git a/Engine/Plugins/ChaosVD/Source/ChaosVD/Private/ChaosVDPlaybackController.cpp b/Engine/Plugins/ChaosVD/Source/ChaosVD/Private/ChaosVDPlaybackController.cpp
--- a/Engine/Plugins/ChaosVD/Source/ChaosVD/Private/ChaosVDPlaybackController.cpp
+++ b/Engine/Plugins/ChaosVD/Source/ChaosVD/Private/ChaosVDPlaybackController.cpp
@@ -12,8 +12,8 @@
 #include "ExtensionsSystem/ChaosVDExtensionsManager.h"
 #include "Misc/MessageDialog.h"
 #include "Settings/ChaosVDGeneralSettings.h"
-#include "Trace/ChaosVDTraceManager.h"
 #include "Trace/ChaosVDTraceProvider.h"
+#include "TraceSessionsManager.h"
 #include "TraceServices/Model/AnalysisSession.h"
 #include "Widgets/SChaosVDTimelineWidget.h"
 
@@ -86,7 +86,7 @@ void FChaosVDPlaybackController::DeInitialize()
 	bInitialized = false;
 }
 
-bool FChaosVDPlaybackController::LoadChaosVDRecordingFromTraceSession(const FChaosVDTraceSessionDescriptor& InSessionDescriptor)
+bool FChaosVDPlaybackController::LoadChaosVDRecordingFromTraceSession(const UE::TraceBasedDebuggers::FTraceSessionDescriptor& InSessionDescriptor)
 {
 	if (!ensure(InSessionDescriptor.IsValid()))
 	{
@@ -98,7 +98,7 @@ bool FChaosVDPlaybackController::LoadChaosVDRecordingFromTraceSession(const FCha
 		UnloadCurrentRecording();
 	}
 
-	if (const TSharedPtr<const TraceServices::IAnalysisSession> TraceSession = FChaosVDModule::Get().GetTraceManager()->GetSession(InSessionDescriptor.SessionName))
+	if (const TSharedPtr<const TraceServices::IAnalysisSession> TraceSession = FChaosVDModule::Get().GetTraceSessionsManager()->GetSession(InSessionDescriptor.SessionName))
 	{
 		if (const FChaosVDTraceProvider* ChaosVDProvider = TraceSession->ReadProvider<FChaosVDTraceProvider>(FChaosVDTraceProvider::ProviderName))
 		{
diff --git a/Engine/Plugins/ChaosVD/Source/ChaosVD/Private/Trace/ChaosVDTraceManager.cpp b/Engine/Plugins/ChaosVD/Source/ChaosVD/Private/Trace/ChaosVDTraceManager.cpp
deleted file mode 100644
diff --git a/Engine/Plugins/ChaosVD/Source/ChaosVD/Private/Widgets/ChaosVDRecordingControls.cpp b/Engine/Plugins/ChaosVD/Source/ChaosVD/Private/Widgets/ChaosVDRecordingControls.cpp
--- a/Engine/Plugins/ChaosVD/Source/ChaosVD/Private/Widgets/ChaosVDRecordingControls.cpp
+++ b/Engine/Plugins/ChaosVD/Source/ChaosVD/Private/Widgets/ChaosVDRecordingControls.cpp
@@ -1,4 +1,4 @@
-// Copyright Epic Games, Inc. All Rights Reserved.
+﻿// Copyright Epic Games, Inc. All Rights Reserved.
 
 #include "Widgets/ChaosVDRecordingControls.h"
 
@@ -26,6 +26,8 @@
 #include "ToolMenu.h"
 #include "ToolMenus.h"
 #include "TraceDataRelayTransport.h"
+#include "TraceSessionsManager.h"
+#include "Trace/StoreClient.h"
 #include "Widgets/Images/SImage.h"
 #include "Widgets/Input/SButton.h"
 #include "Widgets/Input/SComboButton.h"
@@ -655,7 +657,7 @@ bool FChaosVDRecordingControls::CanChangeLoadingMode() const
 		else if (const TSharedPtr<SChaosVDMainTab> MainTabSharedPtr = MainTabWeakPtr.Pin())
 		{
 			// If nothing is loaded yet, it does not make sense change the loading mode
-			return !MainTabSharedPtr->GetChaosVDEngineInstance()->GetCurrentSessionDescriptors().IsEmpty();
+			return !MainTabSharedPtr->GetChaosVDEngineInstance()->GetTraceSessionDescriptors().IsEmpty();
 		}
 	}
 
@@ -933,7 +935,7 @@ FChaosVDRecordingControls::FTraceStoreConnectionSettings FChaosVDRecordingContro
 
 	const FTraceConnectionDetails& RecordingSessionDetails = InSessionInfoRef->GetTraceConnectionDetails();
 
-	if (const UE::Trace::FStoreClient::FSessionInfo* TraceSessionInfo = FChaosVDTraceManager::GetTraceSessionInfo(RecordingSessionDetails.TraceTarget, RecordingSessionDetails.TraceGuid))
+	if (const UE::Trace::FStoreClient::FSessionInfo* TraceSessionInfo = FTraceSessionsManager::GetTraceSessionInfo(RecordingSessionDetails.TraceTarget, RecordingSessionDetails.TraceGuid))
 	{
 		ConnectionSettings.SessionAddress = RecordingSessionDetails.TraceTarget;
 		ConnectionSettings.TraceID = TraceSessionInfo->GetTraceId();
@@ -1069,7 +1071,7 @@ void FChaosVDRecordingControls::SetSessionRecordingState(bool bIsRecording, cons
 			}
 
 			// If we get here it means we have at least one valid session
-			const FChaosVDTraceSessionDescriptor& SessionDesc = MainTabSharedPtr->GetChaosVDEngineInstance()->GetCurrentSessionDescriptors().Last();
+			const FTraceSessionDescriptor& SessionDesc = MainTabSharedPtr->GetChaosVDEngineInstance()->GetTraceSessionDescriptors().Last();
 
 			RecordingParams.Target.Appendf(TEXT(":%u"), SessionDesc.SessionPort);
 
diff --git a/Engine/Plugins/ChaosVD/Source/ChaosVD/Private/Widgets/SChaosVDMainTab.cpp b/Engine/Plugins/ChaosVD/Source/ChaosVD/Private/Widgets/SChaosVDMainTab.cpp
--- a/Engine/Plugins/ChaosVD/Source/ChaosVD/Private/Widgets/SChaosVDMainTab.cpp
+++ b/Engine/Plugins/ChaosVD/Source/ChaosVD/Private/Widgets/SChaosVDMainTab.cpp
@@ -61,7 +61,7 @@
 #include "Styling/StyleColors.h"
 #include "Styling/ToolBarStyle.h"
 #include "TabSpawers/ChaosVDSceneQueryBrowserTab.h"
-#include "Trace/ChaosVDTraceManager.h"
+#include "TraceSessionsManager.h"
 #include "Visualizers/ChaosVDCharacterGroundConstraintsDataComponentVisualizer.h"
 #include "Visualizers/ChaosVDJointConstraintsDataComponentVisualizer.h"
 #include "Visualizers/ChaosVDParticleDataComponentVisualizer.h"
@@ -394,9 +394,9 @@ void SChaosVDMainTab::UpdateTraceProcessingNotificationUpdate()
 	TSharedRef<FChaosVDEngine> CVDEngineInstance = GetChaosVDEngineInstance();
 
 	// Live sessions should not show the notification pop up. When CVD is connected in live debugging mode the UI already has feedback to inform the user.
-	if (CVDEngineInstance->GetCurrentSessionDescriptors().Num() > 0 && !CVDEngineInstance->HasAnyLiveSessionActive())
+	if (CVDEngineInstance->GetTraceSessionDescriptors().Num() > 0 && !CVDEngineInstance->HasAnyLiveSessionActive())
 	{
-		FChaosVDModule::Get().GetTraceManager()->EnumerateActiveSessions([&bIsTraceAnalysisInProgress](const TSharedRef<const TraceServices::IAnalysisSession>& InSession)
+		FChaosVDModule::Get().GetTraceSessionsManager()->EnumerateActiveSessions([&bIsTraceAnalysisInProgress](const TSharedRef<const TraceServices::IAnalysisSession>& InSession)
 		{
 			// TODO: We should implement a way to gather the actual progress so we can report that in the UI
 			// but the Trace API don't have that data yet.
@@ -727,15 +727,14 @@ FReply SChaosVDMainTab::BrowseAndOpenChaosVDRecording()
 			}
 		case EChaosVDBrowseFileModalResponse::CurrentSaveFolder:
 			{
-				UChaosVDGeneralSettings* Settings = FChaosVDSettingsManager::Get().GetSettingsObject<UChaosVDGeneralSettings>();
-				const FString FolderPath = Settings && !Settings->DefaultSavePathOverride.FilePath.IsEmpty() ? Settings->DefaultSavePathOverride.FilePath : FChaosVDTraceManager::GetDefaultSaveDirectoryPath();
-
+				// GetSaveDirPath will return the path to use base on the current ChaosVDGeneralSettings (i.e., using DefaultSavePathOverride if set)
+				const FString FolderPath = FChaosVDModule::Get().GetTraceSessionsManager()->GetSaveDirPath();
 				BrowseChaosVDRecordingFromFolder(*FolderPath, LoadingMode);
 				break;
 			}
 		case EChaosVDBrowseFileModalResponse::LocalTraceStore:
 			{
-				const FString TraceStorePath = FChaosVDModule::Get().GetTraceManager()->GetLocalTraceStoreDirPath();
+				const FString TraceStorePath = FChaosVDModule::Get().GetTraceSessionsManager()->GetLocalTraceStoreDirPath();
 				if (TraceStorePath.IsEmpty())
 				{
 					UE_LOG(LogChaosVDEditor, Error, TEXT("[%hs] Failed to access Trace Store..."), __func__);
@@ -983,7 +982,7 @@ bool SChaosVDMainTab::ConnectToLiveSession(int32 SessionID, const FString& InSes
 
 bool SChaosVDMainTab::ConnectToLiveSession_Direct(EChaosVDLoadRecordedDataMode LoadingMode) const
 {
-	UE_LOG(LogChaosVDEditor, Error, TEXT("[%hs] This menthod is no longer supported. Please use the overload that takes a Remote Session ID."), __func__);
+	UE_LOG(LogChaosVDEditor, Error, TEXT("[%hs] This method is no longer supported. Please use the overload that takes a Remote Session ID."), __func__);
 	return false;
 }
 
@@ -1243,10 +1242,12 @@ void SChaosVDMainTab::HandleDisconnectSessionClicked()
 	
 		if (bIsAlreadyInLiveSession)
 		{
-			TArrayView<FChaosVDTraceSessionDescriptor> ActiveSessions = GetChaosVDEngineInstance()->GetCurrentSessionDescriptors();
-			for (FChaosVDTraceSessionDescriptor& ActiveSession : ActiveSessions)
+			using namespace UE::TraceBasedDebuggers;
+			const TSharedPtr<FTraceSessionsManager> SessionsManager = FChaosVDModule::Get().GetTraceSessionsManager();
+			const TArrayView<FTraceSessionDescriptor> ActiveSessions = GetChaosVDEngineInstance()->GetTraceSessionDescriptors();
+			for (FTraceSessionDescriptor& ActiveSession : ActiveSessions)
 			{
-				FChaosVDModule::Get().GetTraceManager()->CloseSession(ActiveSession.SessionName);
+				SessionsManager->CloseSession(ActiveSession.SessionName);
 				ActiveSession.bIsLiveSession = false;
 			}
 
@@ -1259,7 +1260,7 @@ FText SChaosVDMainTab::GetDisconnectButtonText() const
 {
 	if (GetChaosVDEngineInstance()->HasAnyLiveSessionActive())
 	{
-		if (GetChaosVDEngineInstance()->GetCurrentSessionDescriptors().Num() > 1)
+		if (GetChaosVDEngineInstance()->GetTraceSessionDescriptors().Num() > 1)
 		{
 			return LOCTEXT("DisconnectFromMultipleSessions", "Disconnect from all Sessions");
 		}
diff --git a/Engine/Plugins/ChaosVD/Source/ChaosVD/Private/Widgets/SChaosVDRecordedLogBrowser.cpp b/Engine/Plugins/ChaosVD/Source/ChaosVD/Private/Widgets/SChaosVDRecordedLogBrowser.cpp
--- a/Engine/Plugins/ChaosVD/Source/ChaosVD/Private/Widgets/SChaosVDRecordedLogBrowser.cpp
+++ b/Engine/Plugins/ChaosVD/Source/ChaosVD/Private/Widgets/SChaosVDRecordedLogBrowser.cpp
@@ -3,14 +3,15 @@
 #include "SChaosVDRecordedLogBrowser.h"
 
 #include "ChaosVDEngine.h"
+#include "ChaosVDModule.h"
 #include "ChaosVDPlaybackController.h"
 #include "ChaosVDScene.h"
+#include "Components/VerticalBox.h"
+#include "Misc/TextFilterExpressionEvaluator.h"
 #include "SChaosVDLogBrowserToolbar.h"
 #include "SChaosVDRecordedLogView.h"
-#include "SlateOptMacros.h"
 #include "ToolMenus.h"
-#include "Components/VerticalBox.h"
-#include "Misc/TextFilterExpressionEvaluator.h"
+#include "TraceSessionsManager.h"
 #include "TraceServices/Model/AnalysisSession.h"
 #include "TraceServices/Model/Log.h"
 
@@ -156,7 +157,7 @@ void SChaosVDRecordedLogBrowser::UpdateBrowserContents()
 
 	if (!CurrentSessionName.IsEmpty())
 	{
-		Session = FChaosVDModule::Get().GetTraceManager()->GetSession(CurrentSessionName);
+		Session = FChaosVDModule::Get().GetTraceSessionsManager()->GetSession(CurrentSessionName);
 	}
 	
 	if (!Session)
@@ -459,7 +460,7 @@ void SChaosVDRecordedLogBrowser::HandleItemFocused(const TSharedPtr<FChaosVDLogV
 
 	TSharedPtr<FChaosVDEngine> CVDEngineInstance = EngineInstanceWeakPtr.Pin();
 	TSharedPtr<const TraceServices::IAnalysisSession> Session = CVDEngineInstance && !CurrentSessionName.IsEmpty() ?
-																FChaosVDModule::Get().GetTraceManager()->GetSession(CurrentSessionName) :
+																FChaosVDModule::Get().GetTraceSessionsManager()->GetSession(CurrentSessionName) :
 																nullptr;
 
 	if (!Session)
@@ -610,7 +611,7 @@ void SChaosVDRecordedLogBrowser::ProcessPendingTrackUpdates()
 	}
 
 	TSharedPtr<FChaosVDEngine> CVDEngineInstance = EngineInstanceWeakPtr.Pin();
-	TSharedPtr<FChaosVDTraceManager> TraceManager = FChaosVDModule::Get().GetTraceManager();
+	TSharedPtr<UE::TraceBasedDebuggers::FTraceSessionsManager> TraceManager = FChaosVDModule::Get().GetTraceSessionsManager();
 	TSharedPtr<const TraceServices::IAnalysisSession> Session = CVDEngineInstance && TraceManager && !CurrentSessionName.IsEmpty() ?
 																TraceManager->GetSession(CurrentSessionName) :
 																nullptr;
diff --git a/Engine/Plugins/ChaosVD/Source/ChaosVD/Private/Widgets/SChaosVDRecordedLogBrowser.h b/Engine/Plugins/ChaosVD/Source/ChaosVD/Private/Widgets/SChaosVDRecordedLogBrowser.h
--- a/Engine/Plugins/ChaosVD/Source/ChaosVD/Private/Widgets/SChaosVDRecordedLogBrowser.h
+++ b/Engine/Plugins/ChaosVD/Source/ChaosVD/Private/Widgets/SChaosVDRecordedLogBrowser.h
@@ -2,11 +2,10 @@
 
 #pragma once
 
-#include "CoreMinimal.h"
+#include "Async/AsyncWork.h"
 #include "ChaosVDPlaybackControllerInstigator.h"
 #include "ChaosVDPlaybackControllerObserver.h"
 #include "ContentBrowserDataMenuContexts.h"
-#include "Async/AsyncWork.h"
 #include "Framework/Docking/TabManager.h"
 #include "Misc/TextFilterExpressionEvaluator.h"
 #include "TraceServices/Model/Log.h"
@@ -14,7 +13,6 @@
 
 #include "SChaosVDRecordedLogBrowser.generated.h"
 
-struct FChaosVDTraceSessionDescriptor;
 class FTabManager;
 class FChaosVDEngine;
 class SChaosVDLogBrowserToolbar;
diff --git a/Engine/Plugins/ChaosVD/Source/ChaosVD/Private/Widgets/SChaosVDRecordedLogInstances.cpp b/Engine/Plugins/ChaosVD/Source/ChaosVD/Private/Widgets/SChaosVDRecordedLogInstances.cpp
--- a/Engine/Plugins/ChaosVD/Source/ChaosVD/Private/Widgets/SChaosVDRecordedLogInstances.cpp
+++ b/Engine/Plugins/ChaosVD/Source/ChaosVD/Private/Widgets/SChaosVDRecordedLogInstances.cpp
@@ -2,14 +2,14 @@
 
 #include "SChaosVDRecordedLogInstances.h"
 
+#include "Application/SlateApplicationBase.h"
 #include "ChaosVDEngine.h"
 #include "ChaosVDModule.h"
 #include "ChaosVDStyle.h"
 #include "SChaosVDRecordedLogBrowser.h"
-#include "Application/SlateApplicationBase.h"
-#include "Trace/ChaosVDTraceManager.h"
-#include "Widgets/SToolTip.h"
+#include "TraceSessionsManager.h"
 #include "Widgets/Docking/SDockTab.h"
+#include "Widgets/SToolTip.h"
 #include "Widgets/Text/STextBlock.h"
 
 #define LOCTEXT_NAMESPACE "ChaosVisualDebugger"
@@ -20,8 +20,8 @@ SChaosVDRecordedLogInstances::~SChaosVDRecordedLogInstances()
 
 	if (TSharedPtr<FChaosVDEngine> EngineInstance = EngineInstanceWeakPtr.Pin())
 	{
-		EngineInstance->OnSessionOpened().RemoveAll(this);
-		EngineInstance->OnSessionClosed().RemoveAll(this);
+		EngineInstance->OnTraceSessionOpened().RemoveAll(this);
+		EngineInstance->OnTraceSessionClosed().RemoveAll(this);
 	}
 }
 
@@ -43,20 +43,21 @@ void SChaosVDRecordedLogInstances::Construct(const FArguments& InArgs, const TSh
 		]
 	];
 
-	InEngineInstance->OnSessionClosed().AddSP(this, &SChaosVDRecordedLogInstances::HandleSessionClosed);
-	InEngineInstance->OnSessionOpened().AddSP(this, &SChaosVDRecordedLogInstances::HandleSessionOpened);
+	InEngineInstance->OnTraceSessionClosed().AddSP(this, &SChaosVDRecordedLogInstances::HandleSessionClosed);
+	InEngineInstance->OnTraceSessionOpened().AddSP(this, &SChaosVDRecordedLogInstances::HandleSessionOpened);
 
-	TArrayView<FChaosVDTraceSessionDescriptor> SessionDescriptors = InEngineInstance->GetCurrentSessionDescriptors();
+	using namespace UE::TraceBasedDebuggers;
+	TArrayView<FTraceSessionDescriptor> SessionDescriptors = InEngineInstance->GetTraceSessionDescriptors();
 	if (!SessionDescriptors.IsEmpty())
 	{
-		for (const FChaosVDTraceSessionDescriptor& SessionDescriptor : SessionDescriptors)
+		for (const FTraceSessionDescriptor& SessionDescriptor : SessionDescriptors)
 		{
 			HandleSessionOpened(SessionDescriptor);
 		}
 	}
 }
 
-void SChaosVDRecordedLogInstances::HandleSessionOpened(const FChaosVDTraceSessionDescriptor& InSessionDescriptor)
+void SChaosVDRecordedLogInstances::HandleSessionOpened(const UE::TraceBasedDebuggers::FTraceSessionDescriptor& InSessionDescriptor)
 {
 	FName TabID = GetAvailableTabID();
 
@@ -65,7 +66,7 @@ void SChaosVDRecordedLogInstances::HandleSessionOpened(const FChaosVDTraceSessio
 	InstancesTabManager->TryInvokeTab(TabID);
 }
 
-void SChaosVDRecordedLogInstances::HandleSessionClosed(const FChaosVDTraceSessionDescriptor& InSessionDescriptor)
+void SChaosVDRecordedLogInstances::HandleSessionClosed(const UE::TraceBasedDebuggers::FTraceSessionDescriptor& InSessionDescriptor)
 {
 	if (TWeakPtr<SDockTab>* ActiveTabForSessionPtrPtr = ActiveTabsBySessionName.Find(InSessionDescriptor.SessionName))
 	{
diff --git a/Engine/Plugins/ChaosVD/Source/ChaosVD/Private/Widgets/SChaosVDRecordedLogInstances.h b/Engine/Plugins/ChaosVD/Source/ChaosVD/Private/Widgets/SChaosVDRecordedLogInstances.h
--- a/Engine/Plugins/ChaosVD/Source/ChaosVD/Private/Widgets/SChaosVDRecordedLogInstances.h
+++ b/Engine/Plugins/ChaosVD/Source/ChaosVD/Private/Widgets/SChaosVDRecordedLogInstances.h
@@ -1,12 +1,16 @@
-// Copyright Epic Games, Inc. All Rights Reserved.
+﻿// Copyright Epic Games, Inc. All Rights Reserved.
 
 #pragma once
 
 #include "Framework/Docking/TabManager.h"
 #include "Widgets/SCompoundWidget.h"
 
+namespace UE::TraceBasedDebuggers
+{
+struct FTraceSessionDescriptor;
+}
+
 class FSpawnTabArgs;
-struct FChaosVDTraceSessionDescriptor;
 class FChaosVDEngine;
 class SDockTab;
 
@@ -27,8 +31,8 @@ public:
 
 private:
 
-	void HandleSessionOpened(const FChaosVDTraceSessionDescriptor& InSessionDescriptor);
-	void HandleSessionClosed(const FChaosVDTraceSessionDescriptor& InSessionDescriptor);
+	void HandleSessionOpened(const UE::TraceBasedDebuggers::FTraceSessionDescriptor& InSessionDescriptor);
+	void HandleSessionClosed(const UE::TraceBasedDebuggers::FTraceSessionDescriptor& InSessionDescriptor);
 
 	virtual TSharedRef<SDockTab> HandleTabSpawnRequest(const FSpawnTabArgs& Args, FName TabSessionID);
 
diff --git a/Engine/Plugins/ChaosVD/Source/ChaosVD/Public/ChaosVDEngine.h b/Engine/Plugins/ChaosVD/Source/ChaosVD/Public/ChaosVDEngine.h
--- a/Engine/Plugins/ChaosVD/Source/ChaosVD/Public/ChaosVDEngine.h
+++ b/Engine/Plugins/ChaosVD/Source/ChaosVD/Public/ChaosVDEngine.h
@@ -2,22 +2,25 @@
 
 #pragma once
 
-#include "Templates/SharedPointer.h"
 #include "Containers/Ticker.h"
 #include "Misc/Guid.h"
-#include "Trace/ChaosVDTraceManager.h"
-
+#include "Templates/SharedPointer.h"
 #include "ChaosVDEngine.generated.h"
 
 namespace Chaos::VisualDebugger
 {
-	struct FChaosVDOptionalDataChannel;
+struct FChaosVDOptionalDataChannel;
+}
+
+namespace UE::TraceBasedDebuggers
+{
+struct FTraceSessionDescriptor;
 }
 
-class FChaosVDTraceManager;
 class FChaosVDPlaybackController;
 class FChaosVDScene;
 class FChaosVisualDebuggerMainUI;
+struct FChaosVDRecording;
 
 /** Enumeration of the available modes controlling how data is loaded into CVD */
 UENUM()
@@ -29,16 +32,19 @@ enum class EChaosVDLoadRecordedDataMode : uint8
 	MultiSource
 };
 
-DECLARE_MULTICAST_DELEGATE_OneParam(FSessionStateChangedDelegate, const FChaosVDTraceSessionDescriptor& InSessionDescriptor)
+DECLARE_MULTICAST_DELEGATE_OneParam(FTraceSessionStateChangedDelegate, const UE::TraceBasedDebuggers::FTraceSessionDescriptor& InSessionDescriptor)
+
+PRAGMA_DISABLE_DEPRECATION_WARNINGS
+UE_DEPRECATED(5.8, "Use FTraceSessionStateChangedDelegate taking UE::TraceBasedDebuggers::FTraceSessionDescriptor in parameter")
+DECLARE_MULTICAST_DELEGATE_OneParam(FSessionStateChangedDelegate, const struct FChaosVDTraceSessionDescriptor& InSessionDescriptor)
+PRAGMA_ENABLE_DEPRECATION_WARNINGS
 
 /** Core Implementation of the visual debugger - Owns the systems that are not UI */
 class FChaosVDEngine : public FTSTickerObjectBase, public TSharedFromThis<FChaosVDEngine>
 {
 public:
-	FChaosVDEngine()
-	{
-		InstanceGUID = FGuid::NewGuid();
-	}
+	CHAOSVD_API FChaosVDEngine();
+	CHAOSVD_API virtual ~FChaosVDEngine();
 
 	CHAOSVD_API void Initialize();
 	CHAOSVD_API void CloseActiveTraceSessions();
@@ -65,42 +71,70 @@ public:
 		return PlaybackController;
 	}
 
-	TArrayView<FChaosVDTraceSessionDescriptor> GetCurrentSessionDescriptors()
-	{
-		return CurrentSessionDescriptors;
-	}
-
 	bool HasAnyLiveSessionActive() const;
 
 	CHAOSVD_API void LoadRecording(const FString& FilePath, EChaosVDLoadRecordedDataMode LoadingMode = EChaosVDLoadRecordedDataMode::MultiSource);
 	void LoadCombinedMultiRecording(const FString& FilePath);
 
 	bool ConnectToLiveSession(uint32 SessionID, const FString& InSessionAddress, EChaosVDLoadRecordedDataMode LoadingMode = EChaosVDLoadRecordedDataMode::MultiSource);
-
-	UE_DEPRECATED(5.8, "Please use the version that takes a remote Session ID")
-	bool ConnectToLiveSession_Direct(EChaosVDLoadRecordedDataMode LoadingMode = EChaosVDLoadRecordedDataMode::MultiSource);
-	
 	bool ConnectToLiveSession_Direct(FGuid RemoteSessionID, EChaosVDLoadRecordedDataMode LoadingMode = EChaosVDLoadRecordedDataMode::MultiSource);
 	bool ConnectToLiveSession_Relay(FGuid RemoteSessionID, EChaosVDLoadRecordedDataMode LoadingMode = EChaosVDLoadRecordedDataMode::MultiSource);
 
-	void OpenSession(const FChaosVDTraceSessionDescriptor& SessionDescriptor, EChaosVDLoadRecordedDataMode LoadingMode = EChaosVDLoadRecordedDataMode::MultiSource);
+	void OpenSession(const UE::TraceBasedDebuggers::FTraceSessionDescriptor& SessionDescriptor, EChaosVDLoadRecordedDataMode LoadingMode = EChaosVDLoadRecordedDataMode::MultiSource);
 
 	bool SaveOpenSessionToCombinedFile(const FString& InTargetFilePath = FString());
 
 	bool CanCombineOpenSessions() const;
-	
+
+PRAGMA_DISABLE_DEPRECATION_WARNINGS
+
+	UE_DEPRECATED(5.8, "Use GetTraceSessionDescriptors")
+	TArrayView<FChaosVDTraceSessionDescriptor> GetCurrentSessionDescriptors()
+	{
+		return {};
+	}
+
+	UE_DEPRECATED(5.8, "Please use the version that takes a remote Session ID")
+	bool ConnectToLiveSession_Direct(EChaosVDLoadRecordedDataMode LoadingMode = EChaosVDLoadRecordedDataMode::MultiSource);
+
+	UE_DEPRECATED(5.8, "Use the version taking UE::TraceBasedDebuggers::FTraceSessionDescriptor in parameter")
+	void OpenSession(const FChaosVDTraceSessionDescriptor& SessionDescriptor, EChaosVDLoadRecordedDataMode LoadingMode = EChaosVDLoadRecordedDataMode::MultiSource)
+	{
+	}
+
+	UE_DEPRECATED(5.8, "Use OnTraceSessionOpened instead")
 	FSessionStateChangedDelegate& OnSessionOpened()
 	{
-		return OnSessionOpenedDelegate;
+		static FSessionStateChangedDelegate SessionOpenedDelegate;
+		return SessionOpenedDelegate;
 	}
 
+	UE_DEPRECATED(5.8, "Use OnTraceSessionClosed instead")
 	FSessionStateChangedDelegate& OnSessionClosed()
 	{
-		return OnSessionClosedDelegate;
+		static FSessionStateChangedDelegate SessionClosedDelegate;
+		return SessionClosedDelegate;
 	}
 
+	UE_DEPRECATED(5.8, "Use GetTraceSessionDescriptor")
 	const FChaosVDTraceSessionDescriptor& GetSessionDescriptor(FGuid RemoteSessionID) const;
 
+PRAGMA_ENABLE_DEPRECATION_WARNINGS
+
+	FTraceSessionStateChangedDelegate& OnTraceSessionOpened()
+	{
+		return OnSessionOpenedDelegate;
+	}
+
+	FTraceSessionStateChangedDelegate& OnTraceSessionClosed()
+	{
+		return OnSessionClosedDelegate;
+	}
+
+	CHAOSVD_API const UE::TraceBasedDebuggers::FTraceSessionDescriptor& GetTraceSessionDescriptor(FGuid RemoteSessionID) const;
+
+	CHAOSVD_API TArrayView<UE::TraceBasedDebuggers::FTraceSessionDescriptor> GetTraceSessionDescriptors();
+
 private:
 
 	void LoadRecording_Internal(const TFunction<FString(const TSharedPtr<FChaosVDRecording>&)>&, EChaosVDLoadRecordedDataMode LoadingMode = EChaosVDLoadRecordedDataMode::MultiSource);
@@ -112,7 +146,7 @@ private:
 
 	FGuid InstanceGUID;
 
-	TArray<FChaosVDTraceSessionDescriptor> CurrentSessionDescriptors;
+	TArray<UE::TraceBasedDebuggers::FTraceSessionDescriptor> CurrentSessionDescriptors;
 
 	TSharedPtr<FChaosVDScene> CurrentScene;
 	TSharedPtr<FChaosVDPlaybackController> PlaybackController;
@@ -123,6 +157,6 @@ private:
 
 	FDelegateHandle DataChannelStateUpdatedHandle;
 	
-	FSessionStateChangedDelegate OnSessionOpenedDelegate;
-	FSessionStateChangedDelegate OnSessionClosedDelegate;
-};
+	FTraceSessionStateChangedDelegate OnSessionOpenedDelegate;
+	FTraceSessionStateChangedDelegate OnSessionClosedDelegate;
+};
\ No newline at end of file
diff --git a/Engine/Plugins/ChaosVD/Source/ChaosVD/Public/ChaosVDModule.h b/Engine/Plugins/ChaosVD/Source/ChaosVD/Public/ChaosVDModule.h
--- a/Engine/Plugins/ChaosVD/Source/ChaosVD/Public/ChaosVDModule.h
+++ b/Engine/Plugins/ChaosVD/Source/ChaosVD/Public/ChaosVDModule.h
@@ -9,13 +9,17 @@
 #include "Misc/Guid.h"
 
 class SWidget;
-class FChaosVDTraceManager;
 class FChaosVDEngine;
 class SChaosVDMainTab;
 class SDockTab;
 class FSpawnTabArgs;
 struct FGuid;
 
+namespace UE::TraceBasedDebuggers
+{
+class FTraceSessionsManager;
+}
+
 /** Main module class for the Chaos Visual Debugger editor */
 DECLARE_LOG_CATEGORY_EXTERN(LogChaosVDEditor, Log, VeryVerbose);
 class FChaosVDModule : public IModuleInterface
@@ -31,7 +35,16 @@ public:
 	virtual void ShutdownModule() override;
 
 	/** Returns the Trace manager instance used by CVD */
-	TSharedPtr<FChaosVDTraceManager>& GetTraceManager() { return ChaosVDTraceManager; }
+	TSharedPtr<UE::TraceBasedDebuggers::FTraceSessionsManager>& GetTraceSessionsManager()
+	{
+		return ChaosVDTraceManager;
+	}
+
+	UE_DEPRECATED(5.8, "Use GetTraceSessionsManager")
+	TSharedPtr<class FChaosVDTraceManager> GetTraceManager()
+	{
+		return nullptr;
+	}
 
 	/** Spawns a new instance of CVD */
 	void SpawnCVDTab();
@@ -61,7 +74,7 @@ private:
 
 	TMap<FGuid, TWeakPtr<SDockTab>> ActiveCVDTabs;
 
-	TSharedPtr<FChaosVDTraceManager> ChaosVDTraceManager;
+	TSharedPtr<UE::TraceBasedDebuggers::FTraceSessionsManager> ChaosVDTraceManager;
 
 	static FString ChaosVisualDebuggerProgramName;
 
diff --git a/Engine/Plugins/ChaosVD/Source/ChaosVD/Public/ChaosVDPlaybackController.h b/Engine/Plugins/ChaosVD/Source/ChaosVD/Public/ChaosVDPlaybackController.h
--- a/Engine/Plugins/ChaosVD/Source/ChaosVD/Public/ChaosVDPlaybackController.h
+++ b/Engine/Plugins/ChaosVD/Source/ChaosVD/Public/ChaosVDPlaybackController.h
@@ -1,4 +1,4 @@
-// Copyright Epic Games, Inc. All Rights Reserved.
+﻿// Copyright Epic Games, Inc. All Rights Reserved.
 
 #pragma once
 
@@ -13,13 +13,17 @@
 struct FChaosVDPlaybackEngineSnapshot;
 class IChaosVDPlaybackControllerInstigator;
 class UChaosVDCoreSettings;
-struct FChaosVDTraceSessionDescriptor;
 struct FChaosVDTrackInfo;
 class FChaosVDScene;
 struct FChaosVDRecording;
 class FChaosVDPlaybackController;
 class FString;
 
+namespace UE::TraceBasedDebuggers
+{
+struct FTraceSessionDescriptor;
+}
+
 enum class EChaosVDPlaybackButtonsID : uint8;
 
 DECLARE_MULTICAST_DELEGATE_OneParam(FChaosVDInitialFirstFramePlayed, TWeakPtr<FChaosVDPlaybackController>)
@@ -124,7 +128,12 @@ public:
 	void DeInitialize();
 
 	/** Loads a recording using a CVD Trace Session Descriptor */
-	bool LoadChaosVDRecordingFromTraceSession(const FChaosVDTraceSessionDescriptor& InSessionDescriptor);
+	bool LoadChaosVDRecordingFromTraceSession(const UE::TraceBasedDebuggers::FTraceSessionDescriptor& InSessionDescriptor);
+
+PRAGMA_DISABLE_DEPRECATION_WARNINGS
+	UE_DEPRECATED(5.8, "Use the version taking UE::TraceBasedDebuggers::FTraceSessionDescriptor in parameter")
+	bool LoadChaosVDRecordingFromTraceSession(const struct FChaosVDTraceSessionDescriptor& InSessionDescriptor);
+PRAGMA_ENABLE_DEPRECATION_WARNINGS
 
 	/** Unloads the currently loaded recording
 	 * @param UnloadOptions Options flags to change the steps performed during the unload
diff --git a/Engine/Plugins/ChaosVD/Source/ChaosVD/Public/Trace/ChaosVDTraceManager.h b/Engine/Plugins/ChaosVD/Source/ChaosVD/Public/Trace/ChaosVDTraceManager.h
--- a/Engine/Plugins/ChaosVD/Source/ChaosVD/Public/Trace/ChaosVDTraceManager.h
+++ b/Engine/Plugins/ChaosVD/Source/ChaosVD/Public/Trace/ChaosVDTraceManager.h
@@ -2,13 +2,15 @@
 
 #pragma once
 
-#include "ChaosVDModule.h"
 #include "Chaos/Serialization/SerializedDataBuffer.h"
+#include "ChaosVDModule.h"
 #include "Containers/SpscQueue.h"
 #include "Containers/StringFwd.h"
 #include "Framework/Notifications/NotificationManager.h"
 #include "Templates/SharedPointer.h"
 #include "Trace/StoreClient.h"
+#include "TraceDataStreamTypes.h"
+#include "TraceSessionsManager.h"
 
 struct FChaosVDRelayTraceDataMessage;
 enum class EChaosVDLoadRecordedDataMode : uint8;
@@ -17,148 +19,69 @@ class FChaosVDTraceModule;
 
 namespace TraceServices { class IAnalysisSession; }
 
-namespace Chaos::VD
+/**
+ * Objects that allows us to use TLS to temporarily store and access a ptr to an existing instance.
+ * This is temporary to work around the lack of an API method we need in the trace API, and will be removed in the future.
+ * Either when we add that to the API, or find another way to pass an existing CVD recording to the trace provider before analysis starts
+ */
+UE_EXPERIMENTAL(5.8, "This class is temporary and will be removed when a method in the trace API allows us to pass an existing recording to the trace provider before analysis starts.")
+class FChaosVDTraceManagerThreadContext : public TThreadSingleton<FChaosVDTraceManagerThreadContext>
 {
-	/** Object that allows showing a progress bar in the status bar of the editor,
-	 * where the progress can be updated from any thread */
-	struct FAsyncProgressNotification : FTSTickerObjectBase
-	{
-		FAsyncProgressNotification(const FText& InMessage) : NotificationTitle(InMessage)
-		{
-		}
-
-		virtual ~FAsyncProgressNotification() override;
-
-		/** Updates the current progress
-		 * @param InCurrentProgress Value between 0 and 100, representing the progress we want to display
-		 */
-		void EnterProgress(int32 InCurrentProgress);
-
-	private:
-		virtual bool Tick(float DeltaTime) override;
+public:
 
-		float ElapsedTimeSinceLastUpdate = 0.0f;
+	TWeakPtr<FChaosVDRecording> PendingExternalRecordingWeakPtr;
+};
 
-		FProgressNotificationHandle NotificationHandle;
-		std::atomic<int32> CurrentProgress = 0;
-		FText NotificationTitle;
+namespace Chaos::VD
+{
+PRAGMA_DISABLE_DEPRECATION_WARNINGS
+	struct UE_DEPRECATED(5.8, "This struct will no longer be exposed publicly") FAsyncProgressNotification : FTSTickerObjectBase
+	{
+		FAsyncProgressNotification(const FText&) {}
+		void EnterProgress(int32 InCurrentProgress) {}
 	};
 
-	class IRemoteSessionDataStream
+	class UE_DEPRECATED(5.8, "Use UE::TraceBasedDebuggers::IRemoteSessionDataStream instead") IRemoteSessionDataStream : public UE::TraceBasedDebuggers::IRemoteSessionDataStream
 	{
-	public:
-		virtual ~IRemoteSessionDataStream() = default;
-		virtual FGuid GetOwningRemoteSessionID() const = 0;
-		virtual uint64 UpdateAndGetBytesReadSinceLastMeasurement() = 0;
 	};
 
-	class FSessionStreamStatsUpdater : FTSTickerObjectBase
+	class UE_DEPRECATED(5.8, "Use UE::TraceBasedDebuggers::FSessionStreamStatsUpdater instead") FSessionStreamStatsUpdater : FTSTickerObjectBase
 	{
 	public:
-		virtual ~FSessionStreamStatsUpdater() override;
-
-		explicit FSessionStreamStatsUpdater(IRemoteSessionDataStream* const DataStreamInstance)
-			: DataStreamInstance(DataStreamInstance)
-		{
-		}
-
-	private:
-		virtual bool Tick(float DeltaTime) override;
-	
-		IRemoteSessionDataStream* DataStreamInstance;
-
-		double ElapsedTimeSinceLastUpdate = 0.0;
+		explicit FSessionStreamStatsUpdater(IRemoteSessionDataStream* const DataStreamInstance) {}
 	};
 
-	class FDirectSocketStream : public UE::Trace::FDirectSocketStream, public IRemoteSessionDataStream
+	class UE_DEPRECATED(5.8, "Use UE::TraceBasedDebuggers::FDirectSocketStream instead") FDirectSocketStream : public UE::Trace::FDirectSocketStream, public IRemoteSessionDataStream
 	{
 	public:
-
 		FDirectSocketStream() = delete;
-
-		explicit FDirectSocketStream(FGuid InRemoteSessionID) : BytesReadSinceLastMeasurement(0), RemoteSessionID(InRemoteSessionID), StatsUpdater(this)
-		{
-		}
-
-		explicit FDirectSocketStream(FGuid InRemoteSessionID, TUniquePtr<FArchiveFileWriterGeneric>&& InFileWriter)
-			: UE::Trace::FDirectSocketStream(MoveTemp(InFileWriter)),
-			BytesReadSinceLastMeasurement(0),
-			RemoteSessionID(InRemoteSessionID),
-			StatsUpdater(this)
+		explicit FDirectSocketStream(FGuid) {}
+		explicit FDirectSocketStream(FGuid, TUniquePtr<FArchiveFileWriterGeneric>&& InFileWriter)
+			: UE::Trace::FDirectSocketStream(MoveTemp(InFileWriter))
 		{
 		}
 
-		virtual int32 Read(void* Data, uint32 Size) override;
-
-		virtual FGuid GetOwningRemoteSessionID() const override;
-		virtual uint64 UpdateAndGetBytesReadSinceLastMeasurement() override;
-	
-	private:
-		std::atomic<uint64> BytesReadSinceLastMeasurement;
-		FGuid RemoteSessionID;
-
-		FSessionStreamStatsUpdater StatsUpdater;
+		virtual int32 Read(void* Data, uint32 Size) override { return 0; }
+		virtual FGuid GetOwningRemoteSessionID() const override { return FGuid{}; }
+		virtual uint64 UpdateAndGetBytesReadSinceLastMeasurement() override { return 0; }
 	};
 
-	/** A trace data stream object that listens for data coming from a relay stream
-	 * via CVD's Engine Editor Bridge object */
-	class FRelayDataStream : public UE::Trace::IInDataStream,  public IRemoteSessionDataStream
+	class UE_DEPRECATED(5.8, "Use UE::TraceBasedDebuggers::FRelayDataStream instead") FRelayDataStream : public UE::Trace::IInDataStream,  public IRemoteSessionDataStream
 	{
 	public:
-		explicit FRelayDataStream(FGuid InRemoteSessionID);
-		explicit FRelayDataStream(FGuid InRemoteSessionID, TUniquePtr<FArchiveFileWriterGeneric>&& InFileWriter);
-
-		virtual ~FRelayDataStream() override;
-
-		void EnqueueRelayedData(const TConstArrayView<uint8> InTraceDataBuffer);
-		void EnqueueRelayedData(TArray<uint8>&& InTraceDataBuffer);
-
-		virtual FGuid GetOwningRemoteSessionID() const override;
-		virtual uint64 UpdateAndGetBytesReadSinceLastMeasurement() override;
-
-	private:
-		virtual int32 Read(void* Dest, uint32 DestSize) override;
-		virtual void Close() override;
-		virtual bool WaitUntilReady() override;
-
-		TSpscQueue<FSerializedDataBuffer> DataQueue;
-		FEvent*	CancelEvent = nullptr;
-		FEvent*	WaitForDataEvent = nullptr;
-		TOptional<FSerializedDataBuffer> CurrentPacket;
-		uint32 CurrentPacketOffset = 0;
-		bool bClosed = false;
+		explicit FRelayDataStream(FGuid InRemoteSessionID) {}
+		explicit FRelayDataStream(FGuid InRemoteSessionID, TUniquePtr<FArchiveFileWriterGeneric>&& InFileWriter) {}
 
-		FGuid RemoteSessionID = FGuid();
+		void EnqueueRelayedData(const TConstArrayView<uint8> InTraceDataBuffer) {}
+		void EnqueueRelayedData(TArray<uint8>&& InTraceDataBuffer) {}
 
-		TUniquePtr<FArchiveFileWriterGeneric> FileWriter;
-
-		std::atomic<uint64> BytesReadSinceLastMeasurement;
-		FSessionStreamStatsUpdater StatsUpdater;
+		virtual FGuid GetOwningRemoteSessionID() const override { return FGuid{}; }
+		virtual uint64 UpdateAndGetBytesReadSinceLastMeasurement() override { return 0; }
 	};
-}
+} // Chaos::VD
 
-/**
- * Structure containing info about a trace session used by Chaos Visual Debugger
- */
-struct FChaosVDTraceSessionDescriptor
+struct UE_DEPRECATED(5.8, "Use UE::TraceBasedDebuggers::FTraceSessionDescriptor instead") FChaosVDTraceSessionDescriptor : UE::TraceBasedDebuggers::FTraceSessionDescriptor
 {
-
-	/** The trace session name. Used to get the trace analysis instance for this CVD session*/
-	FString SessionName;
-	/** Port number used to establish a connection to this session. The value is 0 if no specific port was used*/
-	uint16 SessionPort = 0;
-	/** Session ID of this CVD session in the Remote Session management system*/
-	FGuid RemoteSessionID;
-
-	/** True if this is a session recorded live (not loaded from disk)*/
-	bool bIsLiveSession = false;
-
-	/** Returns true if valid. For a session to be valued needs at minimum have a Trace Session name */
-	bool IsValid() const
-	{
-		return !SessionName.IsEmpty();
-	}
-
 	friend bool operator==(const FChaosVDTraceSessionDescriptor& Lhs, const FChaosVDTraceSessionDescriptor& RHS)
 	{
 		return Lhs.SessionName == RHS.SessionName
@@ -171,40 +94,20 @@ struct FChaosVDTraceSessionDescriptor
 	}
 };
 
-/** Objects that allows us to use TLS to temporarily store and access a ptr to an existing instance.
- * This is temporary to workaround the lack of an API method we need in the trace API, and will be removed in the future.
- * either when we add that to the API, or find another way to pass an existing CVD recording to the trace provider before anaislis starts
- */
-class FChaosVDTraceManagerThreadContext : public TThreadSingleton<FChaosVDTraceManagerThreadContext>
-{
-public:
-
-	TWeakPtr<FChaosVDRecording> PendingExternalRecordingWeakPtr;
-};
-
-/** Manager class used by Chaos VD to interact/control UE Trace systems */
-class FChaosVDTraceManager
+class UE_DEPRECATED(5.8, "Use UE::TraceBasedDebuggers::FTraceSessionsManager instead") FChaosVDTraceManager : public UE::TraceBasedDebuggers::FTraceSessionsManager
 {
 public:
 	FChaosVDTraceManager();
 	~FChaosVDTraceManager();
 
-	/** Load a trace file and starts analyzing it
-	 * @param InTraceFilename File Name including (Path Included) of the Trace file to load
-	 */
 	CHAOSVD_API FString LoadTraceFile(const FString& InTraceFilename, const TSharedPtr<FChaosVDRecording>& ExistingRecordingPtr = nullptr);
 
 	FString LoadTraceFile(TUniquePtr<IFileHandle>&& InFileHandle, const FString& InTraceSessionName, const TSharedPtr<FChaosVDRecording>& ExistingRecordingPtr);
-	
-	/**
-	 * Connects to a live Trace Session and starts analyzing it.
-	 * @param InSessionHost Trace Store Address for this session
-	 * @param SessionID Trace ID in the Trace Store provided as host
-	 * @param ExistingRecordingPtr Ptr to an existing recording instance to append the new CVD data to, if any
-	 * @return Session Name
-	 */
-	FString ConnectToLiveSession(FStringView InSessionHost, uint32 SessionID, const TSharedPtr<FChaosVDRecording>& ExistingRecordingPtr = nullptr);
 
+	FString ConnectToLiveSession(FStringView InSessionHost, uint32 SessionID, const TSharedPtr<FChaosVDRecording>& ExistingRecordingPtr = nullptr)
+	{
+		return FString();
+	}
 
 	UE_DEPRECATED(5.8, "Please use the version that takes a Remote Session ID")
 	FString ConnectToLiveSession_Direct( uint16& OutSessionPort, const TSharedPtr<FChaosVDRecording>& ExistingRecordingPtr = nullptr)
@@ -212,123 +115,23 @@ public:
 		return FString();
 	};
 
-	/**
-	 * Opens up a direct trace stream and waits for data
-	 * @param RemoteSessionID Session ID to which we want to connect to
-	 * @param OutSessionPort Number of the port we open for this session
-	 * @param ExistingRecordingPtr Ptr to an existing recording instance to append the new CVD data to, if any
-	 * @return Session Name
-	 */
-	FString ConnectToLiveSession_Direct(FGuid RemoteSessionID, uint16& OutSessionPort, const TSharedPtr<FChaosVDRecording>& ExistingRecordingPtr = nullptr);
-
-	/**
-	 * Opens up a relay trace stream and waits for data
-	 * @param RemoteSessionID Session ID to which we want to connect to
-	 * @param ExistingRecordingPtr Ptr to an existing recording instance to append the new CVD data to, if any
-	 * @return Session Name
-	 */
-	FString ConnectToLiveSession_Relay(FGuid RemoteSessionID, const TSharedPtr<FChaosVDRecording>& ExistingRecordingPtr = nullptr);
-
-	/** Returns the path to the local trace store */
-	FString GetLocalTraceStoreDirPath();
-
-	/** Returns a ptr to the session registered with the provided session name. Null if no session is found */
-	CHAOSVD_API TSharedPtr<const TraceServices::IAnalysisSession> GetSession(const FString& InSessionName);
-
-	/** Stops and de-registers a trace session registered with the provided session name */
-	CHAOSVD_API void CloseSession(const FString& InSessionName);
-
-	/** Stops a trace session registered with the provided session name */
-	void StopSession(const FString& InSessionName);
-
-	template<typename TVisitor>
-	static void EnumerateActiveSessions(FStringView InSessionHost, TVisitor Callback);
-
-	static const UE::Trace::FStoreClient::FSessionInfo* GetTraceSessionInfo(FStringView InSessionHost, FGuid TraceGuid);
-
-	template<typename TVisitor>
-	void EnumerateActiveSessions(TVisitor Callback);
-
-	/** Access the trace store at the provided host address, and returns the file name for the trace file
-	 * @param InSessionHost Address of the Trace Store
-	 * @param SessionID Session id of the trace data from which get the file name from
-	 * @return Full Filename with path to the trace file if exist. Returns an empty string if a trace file cannot be found
-	 */
-	FString GetTraceFileNameFromStoreForSession(FStringView InSessionHost, uint32 SessionID);
-
-	/** Access the trace store at the provided host address, and returns a trace data stream ready to use
-	 * @param InSessionHost Address of the Trace Store
-	 * @param SessionID Session id of the trace data
-	 */
-	UE::Trace::FStoreClient::FTraceData GetTraceDataStreamFromStore(FStringView InSessionHost, uint32 SessionID);
-
-	static FString GetDefaultSaveDirectoryPath();
-
-private:
-
-	TUniquePtr<FArchiveFileWriterGeneric> OpenTraceFileForWrite(const FString& InDirectoryToSavePath);
-
-	bool ConnectToLiveSession_Internal(uint32 SessionID, const TSharedPtr<FChaosVDRecording>& ExistingRecordingPtr, const FString& InRequestedSessionName, UE::Trace::FStoreClient::FTraceData&& InTraceDataStream);
-
-	TUniquePtr<UE::Trace::IInDataStream> CreateFileDataStream(TUniquePtr<IFileHandle>&& InFileHandle);
-	TUniquePtr<Chaos::VD::FRelayDataStream> CreateRelayDataStream(FGuid RemoteSessionID);
-
-	/** Temporary workaround method to set an existing recording structure in CVD's trace provider before the trace analysis starts
-	 * in the Trace analysis thread.*/
-	void SetPendingExternalRecordingToProcess(const TSharedPtr<FChaosVDRecording>& ExistingRecordingPtr);
-
-	/** The trace analysis session. */
-	TMap<FString, TSharedPtr<const TraceServices::IAnalysisSession>> AnalysisSessionByName;
-
-	TSharedPtr<FChaosVDTraceModule> ChaosVDTraceModule;
-
-	static uint16 CurrentDirectModeStartingPort;
-	static constexpr uint16 DefaultStartingPort = 1993;
-};
-
-template <typename TCallback>
-void FChaosVDTraceManager::EnumerateActiveSessions(FStringView InSessionHost, TCallback Callback)
-{
-	if (InSessionHost.IsEmpty())
+	FString ConnectToLiveSession_Direct(FGuid RemoteSessionID, uint16& OutSessionPort, const TSharedPtr<FChaosVDRecording>& ExistingRecordingPtr = nullptr)
 	{
-		UE_LOG(LogChaosVDEditor, Error, TEXT("[%s] Failed to connect to trace store. Provided session host is empty"), ANSI_TO_TCHAR(__FUNCTION__));
-		return;
+		return FString();
 	}
 
-	using namespace UE::Trace;
-	const FStoreClient* StoreClient = FStoreClient::Connect(InSessionHost.GetData());
-
-	if (!StoreClient)
+	FString ConnectToLiveSession_Relay(FGuid RemoteSessionID, const TSharedPtr<FChaosVDRecording>& ExistingRecordingPtr = nullptr)
 	{
-		UE_LOG(LogChaosVDEditor, Error, TEXT("[%s] Failed to connect to trace store at [%s]"), ANSI_TO_TCHAR(__FUNCTION__), InSessionHost.GetData())
-		return;
+		return FString();
 	}
 
-	const uint32 SessionCount = StoreClient->GetSessionCount();
-
-	for (uint32 SessionIndex = 0; SessionIndex < SessionCount; SessionIndex++)
+	template<typename TVisitor>
+	static void EnumerateActiveSessions(FStringView InSessionHost, TVisitor Callback)
 	{
-		if (const FStoreClient::FSessionInfo* SessionInfo = StoreClient->GetSessionInfo(SessionIndex))
-		{
-			if (!Callback(*SessionInfo))
-			{
-				return;
-			}
-		}	
+		// deprecated
 	}
-}
 
-template <typename TVisitor>
-void FChaosVDTraceManager::EnumerateActiveSessions(TVisitor Callback)
-{
-	for (const TPair<FString, TSharedPtr<const TraceServices::IAnalysisSession>>& SessionWithName : AnalysisSessionByName)
-	{
-		if (SessionWithName.Value)
-		{
-			if (!Callback(SessionWithName.Value.ToSharedRef()))
-			{
-				return;
-			}
-		}
-	}
-}
+	static FString GetDefaultSaveDirectoryPath();
+};
+PRAGMA_ENABLE_DEPRECATION_WARNINGS
+
diff --git a/Engine/Source/Runtime/Experimental/ChaosVDData/Public/ChaosVDRecording.h b/Engine/Source/Runtime/Experimental/ChaosVDData/Public/ChaosVDRecording.h
--- a/Engine/Source/Runtime/Experimental/ChaosVDData/Public/ChaosVDRecording.h
+++ b/Engine/Source/Runtime/Experimental/ChaosVDData/Public/ChaosVDRecording.h
@@ -1,4 +1,4 @@
-// Copyright Epic Games, Inc. All Rights Reserved.
+﻿// Copyright Epic Games, Inc. All Rights Reserved.
 
 #pragma once
 
diff --git a/Engine/Source/Runtime/Experimental/TraceBasedDebuggers/Private/RemoteSessionsManager.cpp b/Engine/Source/Runtime/Experimental/TraceBasedDebuggers/Private/RemoteSessionsManager.cpp
--- a/Engine/Source/Runtime/Experimental/TraceBasedDebuggers/Private/RemoteSessionsManager.cpp
+++ b/Engine/Source/Runtime/Experimental/TraceBasedDebuggers/Private/RemoteSessionsManager.cpp
@@ -16,8 +16,6 @@
 #define TRACE_BASED_DEBUGGERS_WITHOUT_TRACE 0
 #endif
 
-DEFINE_LOG_CATEGORY(LogTraceBasedDebuggers)
-
 namespace UE::TraceBasedDebuggers
 {
 const FGuid FRemoteSessionsManager::AllRemoteSessionsWrapperGUID = FGuid::NewGuid();
@@ -500,6 +498,10 @@ void FRemoteSessionsManager::HandleRecordingStatusUpdateMessage(const FRecording
 			return;
 		}
 
+		// Update recording state before broadcasting so listeners can use the information to only
+		// process request from their debugger Id.
+		SessionInfoPtr->LastKnownRecordingState = Message;
+
 		if (bLastKnownDebuggerIsRecording != Message.bIsRecording)
 		{
 			if (Message.bIsRecording)
@@ -511,8 +513,6 @@ void FRemoteSessionsManager::HandleRecordingStatusUpdateMessage(const FRecording
 				RecordingStoppedDelegate.Broadcast(SessionInfoPtr);
 			}
 		}
-
-		SessionInfoPtr->LastKnownRecordingState = Message;
 	}
 	else
 	{
diff --git a/Engine/Source/Runtime/Experimental/TraceBasedDebuggers/Private/TraceBasedDebuggerTypes.cpp b/Engine/Source/Runtime/Experimental/TraceBasedDebuggers/Private/TraceBasedDebuggerTypes.cpp
new file mode 100644
--- /dev/null
+++ b/Engine/Source/Runtime/Experimental/TraceBasedDebuggers/Private/TraceBasedDebuggerTypes.cpp
@@ -0,0 +1,5 @@
+// Copyright Epic Games, Inc. All Rights Reserved.
+
+#include "TraceBasedDebuggerTypes.h"
+
+DEFINE_LOG_CATEGORY(LogTraceBasedDebuggers)
diff --git a/Engine/Source/Runtime/Experimental/TraceBasedDebuggers/Public/RemoteSessionsManager.h b/Engine/Source/Runtime/Experimental/TraceBasedDebuggers/Public/RemoteSessionsManager.h
--- a/Engine/Source/Runtime/Experimental/TraceBasedDebuggers/Public/RemoteSessionsManager.h
+++ b/Engine/Source/Runtime/Experimental/TraceBasedDebuggers/Public/RemoteSessionsManager.h
@@ -5,6 +5,7 @@
 #include "Containers/Ticker.h"
 #include "IMessageContext.h"
 #include "MessageEndpoint.h"
+#include "TraceBasedDebuggerTypes.h"
 #include "Misc/DateTime.h"
 #include "Misc/Guid.h"
 #include "Templates/SharedPointer.h"
@@ -17,8 +18,6 @@
 class IMessageBus;
 class IMessageContext;
 
-UE_API DECLARE_LOG_CATEGORY_EXTERN(LogTraceBasedDebuggers, Log, All);
-
 namespace UE::TraceBasedDebuggers
 {
 struct FFullSessionInfoRequestMessage;
diff --git a/Engine/Source/Runtime/Experimental/TraceBasedDebuggers/Public/TraceBasedDebuggerTypes.h b/Engine/Source/Runtime/Experimental/TraceBasedDebuggers/Public/TraceBasedDebuggerTypes.h
new file mode 100644
--- /dev/null
+++ b/Engine/Source/Runtime/Experimental/TraceBasedDebuggers/Public/TraceBasedDebuggerTypes.h
@@ -0,0 +1,11 @@
+// Copyright Epic Games, Inc. All Rights Reserved.
+
+#pragma once
+
+#include "Logging/LogCategory.h"
+
+#define UE_API TRACEBASEDDEBUGGERS_API
+
+UE_API DECLARE_LOG_CATEGORY_EXTERN(LogTraceBasedDebuggers, Log, All);
+
+#undef UE_API
\ No newline at end of file
diff --git a/Engine/Source/Runtime/Experimental/TraceBasedDebuggersAnalysis/Private/TraceBasedDebuggersAnalysis.cpp b/Engine/Source/Runtime/Experimental/TraceBasedDebuggersAnalysis/Private/TraceBasedDebuggersAnalysis.cpp
new file mode 100644
--- /dev/null
+++ b/Engine/Source/Runtime/Experimental/TraceBasedDebuggersAnalysis/Private/TraceBasedDebuggersAnalysis.cpp
@@ -0,0 +1,5 @@
+// Copyright Epic Games, Inc. All Rights Reserved.
+
+#include "Modules/ModuleManager.h"
+
+IMPLEMENT_MODULE(FDefaultModuleImpl, TraceBasedDebuggersAnalysis);
\ No newline at end of file
diff --git a/Engine/Source/Runtime/Experimental/TraceBasedDebuggersAnalysis/Private/TraceDataStreamTypes.cpp b/Engine/Source/Runtime/Experimental/TraceBasedDebuggersAnalysis/Private/TraceDataStreamTypes.cpp
new file mode 100644
--- /dev/null
+++ b/Engine/Source/Runtime/Experimental/TraceBasedDebuggersAnalysis/Private/TraceDataStreamTypes.cpp
@@ -0,0 +1,283 @@
+// Copyright Epic Games, Inc. All Rights Reserved.
+
+#if UE_WITH_TRACE_BASED_DEBUGGERS_ANALYSIS
+
+#include "TraceDataStreamTypes.h"
+
+#include "HAL/IConsoleManager.h"
+#include "RemoteSessionsManager.h"
+#include "SessionInfo.h"
+
+namespace UE::TraceBasedDebuggers
+{
+
+float IntervalBetweenUpdatesSeconds = 1.0f;
+
+static FAutoConsoleVariableRef CVarMaxSendQueueBytes(
+	TEXT("tracebaseddebuggers.StreamStatsUpdateInterval"),
+	IntervalBetweenUpdatesSeconds,
+	TEXT("Delay, in seconds before fetching statistics from the data stream instance for the associated remote session."));
+
+//----------------------------------------------------------------------//
+// FSessionStreamStatsUpdater
+//----------------------------------------------------------------------//
+FSessionStreamStatsUpdater::~FSessionStreamStatsUpdater()
+{
+}
+
+bool FSessionStreamStatsUpdater::Tick(const float DeltaTime)
+{
+	ElapsedTimeSinceLastUpdate += DeltaTime;
+
+	if (ElapsedTimeSinceLastUpdate < IntervalBetweenUpdatesSeconds)
+	{
+		return true;
+	}
+
+	if (!ensure(DataStreamInstance))
+	{
+		return false;
+	}
+
+	ElapsedTimeSinceLastUpdate = 0.0;
+
+	using namespace UE::TraceBasedDebuggers;
+	const TSharedPtr<FRemoteSessionsManager> RemoteSessionManager = DataStreamInstance->GetSessionsManager();
+	const TSharedPtr<FSessionInfo> SessionInfo = RemoteSessionManager ? RemoteSessionManager->GetSessionInfo(DataStreamInstance->GetOwningRemoteSessionID()).Pin() : nullptr;
+	if (!SessionInfo)
+	{
+		return false;
+	}
+
+	SessionInfo->SetReceivedBytesPerSecond(DataStreamInstance->UpdateAndGetBytesReadSinceLastMeasurement());
+
+	return true;
+}
+
+//----------------------------------------------------------------------//
+// FDirectSocketStream
+//----------------------------------------------------------------------//
+int32 FDirectSocketStream::Read(void* Data, const uint32 Size)
+{
+	const int32 BytesRead = Trace::FDirectSocketStream::Read(Data, Size);
+
+	if (BytesRead > 0)
+	{
+		BytesReadSinceLastMeasurement += BytesRead;
+	}
+
+	return BytesRead;
+}
+
+FGuid FDirectSocketStream::GetOwningRemoteSessionID() const
+{
+	return RemoteSessionID;
+}
+
+uint64 FDirectSocketStream::UpdateAndGetBytesReadSinceLastMeasurement()
+{
+	const uint64 BytesRead = BytesReadSinceLastMeasurement;
+	BytesReadSinceLastMeasurement = 0;
+	return BytesRead;
+}
+
+TSharedPtr<FRemoteSessionsManager> FDirectSocketStream::GetSessionsManager() const
+{
+	ensureAlways(WeakSessionManager.IsValid());
+	return WeakSessionManager.Pin();
+}
+
+//----------------------------------------------------------------------//
+// FRelayDataStream
+//----------------------------------------------------------------------//
+FRelayDataStream::FRelayDataStream(
+	const TSharedRef<FRemoteSessionsManager>& InSessionsManager
+	, const TSharedRef<IDataRelayTransport>& InDataRelay
+	, const FGuid& InRemoteSessionID
+	, TUniquePtr<FArchiveFileWriterGeneric>&& InFileWriter)
+	: WeakSessionManager(InSessionsManager.ToWeakPtr()),
+	CancelEvent(FPlatformProcess::GetSynchEventFromPool(false)),
+	WaitForDataEvent(FPlatformProcess::GetSynchEventFromPool(false)),
+	RemoteSessionID(InRemoteSessionID),
+	FileWriter(MoveTemp(InFileWriter)),
+#if WITH_TRACE_BASED_DEBUGGERS
+	RelayTransportInstance(InDataRelay),
+#endif
+	BytesReadSinceLastMeasurement(0),
+	StatsUpdater(this)
+{
+#if WITH_TRACE_BASED_DEBUGGERS
+	RelayTransportInstance->RegisterRelayDataReceiverForSessionID(RemoteSessionID
+		, IDataRelayTransport::FProcessReceivedRelayDataDelegate::CreateRaw(this, &FRelayDataStream::EnqueueRelayedData));
+#endif
+}
+
+FRelayDataStream::~FRelayDataStream()
+{
+	FPlatformProcess::ReturnSynchEventToPool(CancelEvent);
+	FPlatformProcess::ReturnSynchEventToPool(WaitForDataEvent);
+
+#if WITH_TRACE_BASED_DEBUGGERS
+	RelayTransportInstance->UnregisterRelayDataReceiverForSessionID(RemoteSessionID);
+#endif
+}
+
+void FRelayDataStream::EnqueueRelayedData(const TConstArrayView<uint8> InTraceDataBuffer)
+{
+	// An empty data array signals an end of file event from the relay instance
+	if (InTraceDataBuffer.IsEmpty())
+	{
+		Close();
+		return;
+	}
+
+	WaitForDataEvent->Trigger();
+
+	FSerializedDataBuffer DataBuffer(InTraceDataBuffer);
+	DataQueue.Enqueue(MoveTemp(DataBuffer));
+}
+
+void FRelayDataStream::EnqueueRelayedData(TArray<uint8>&& InTraceDataBuffer)
+{
+	// An empty data array signals an end of file event from the relay instance
+	if (InTraceDataBuffer.IsEmpty())
+	{
+		Close();
+		return;
+	}
+
+	WaitForDataEvent->Trigger();
+
+	FSerializedDataBuffer DataBuffer(MoveTemp(InTraceDataBuffer));
+	DataQueue.Enqueue(MoveTemp(DataBuffer));
+}
+
+FGuid FRelayDataStream::GetOwningRemoteSessionID() const
+{
+	return RemoteSessionID;
+}
+
+uint64 FRelayDataStream::UpdateAndGetBytesReadSinceLastMeasurement()
+{
+	const uint64 BytesRead = BytesReadSinceLastMeasurement;
+	BytesReadSinceLastMeasurement = 0;
+	return BytesRead;
+}
+
+TSharedPtr<FRemoteSessionsManager> FRelayDataStream::GetSessionsManager() const
+{
+	ensureAlways(WeakSessionManager.IsValid());
+	return WeakSessionManager.Pin();
+}
+
+int32 FRelayDataStream::Read(void* Dest, const uint32 DestSize)
+{
+	if (DestSize == 0)
+	{
+		return 0;
+	}
+
+	if (Dest == nullptr)
+	{
+		return 0;
+	}
+
+	if (!ensure(CancelEvent))
+	{
+		return 0;
+	}
+
+	uint8* DestBuffer = static_cast<uint8*>(Dest);
+	uint32 BytesCopied = 0;
+
+	// We need to wait for some data to be queued. Returning 
+	// zero bytes read is interpreted as eof.
+	while (!CancelEvent->Wait(0, true))
+	{
+		if (!CurrentPacket)
+		{
+			CurrentPacket = DataQueue.Dequeue();
+			CurrentPacketOffset = 0;
+		}
+
+		// Consume as many queued packets as possible and copy into
+		// the destination buffer
+		while (CurrentPacket && BytesCopied < DestSize)
+		{
+			const TArray<uint8>& SrcBuffer = CurrentPacket->GetDataAsByteArrayRef();
+			if (SrcBuffer.IsEmpty())
+			{
+				// We should never get an empty packet. From this point on we can't trust the integrity of the data
+				// so we need to close the stream
+				Close();
+				break;
+			}
+
+			const uint32 ReceivedBufferSize = static_cast<uint32>(SrcBuffer.Num());
+
+			// Copy as much of the current buffer as possible into the remaining space 
+			// in the destination buffer
+			const uint32 BytesToCopy = FMath::Min(ReceivedBufferSize - CurrentPacketOffset, DestSize - BytesCopied);
+			FMemory::Memcpy(&DestBuffer[BytesCopied], SrcBuffer.GetData() + CurrentPacketOffset, BytesToCopy);
+			BytesCopied += BytesToCopy;
+			CurrentPacketOffset += BytesToCopy;
+
+			if (CurrentPacketOffset == ReceivedBufferSize)
+			{
+				// Current packet buffer has been fully copied, goto next
+				CurrentPacket = DataQueue.Dequeue();
+				CurrentPacketOffset = 0;
+			}
+		}
+
+		// If at least some bytes have been written return those. Otherwise, 
+		// we'll wait for some data to arrive.
+		if (BytesCopied > 0)
+		{
+			BytesReadSinceLastMeasurement += BytesCopied;
+
+			if (FileWriter.IsValid() && !FileWriter->IsError())
+			{
+				FileWriter->Serialize(Dest, BytesCopied);
+			}
+			break;
+		}
+		else
+		{
+			WaitForDataEvent->Reset();
+			WaitForDataEvent->Wait();
+		}
+	}
+
+	return BytesCopied;
+}
+
+void FRelayDataStream::Close()
+{
+	if (CancelEvent)
+	{
+		CancelEvent->Trigger();
+	}
+
+	if (WaitForDataEvent)
+	{
+		WaitForDataEvent->Trigger();
+	}
+
+	if (FileWriter.IsValid())
+	{
+		FileWriter->Close();
+		FileWriter.Reset();
+	}
+
+	bClosed = true;
+}
+
+bool FRelayDataStream::WaitUntilReady()
+{
+	return DataQueue.Peek() != nullptr;
+}
+
+} // UE::TraceBasedDebuggers
+
+#endif // UE_WITH_TRACE_BASED_DEBUGGERS_ANALYSIS
\ No newline at end of file
diff --git a/Engine/Source/Runtime/Experimental/TraceBasedDebuggersAnalysis/Private/TraceSessionsManager.cpp b/Engine/Source/Runtime/Experimental/TraceBasedDebuggersAnalysis/Private/TraceSessionsManager.cpp
new file mode 100644
--- /dev/null
+++ b/Engine/Source/Runtime/Experimental/TraceBasedDebuggersAnalysis/Private/TraceSessionsManager.cpp
@@ -0,0 +1,600 @@
+// Copyright Epic Games, Inc. All Rights Reserved.
+
+#if UE_WITH_TRACE_BASED_DEBUGGERS_ANALYSIS
+
+#include "TraceSessionsManager.h"
+
+#include "Containers/Ticker.h"
+#include "EngineEditorBridge.h"
+#include "Features/IModularFeatures.h"
+#include "Framework/Notifications/NotificationManager.h"
+#include "Internationalization/Internationalization.h"
+#include "Internationalization/Text.h"
+#include "Modules/ModuleManager.h"
+#include "Trace/DataStream.h"
+#include "Trace/StoreClient.h"
+#include "TraceBasedDebuggerRuntime.h"
+#include "TraceBasedDebuggerTypes.h"
+#include "TraceDataStreamTypes.h"
+#include "TraceServices/AnalysisService.h"
+#include "TraceServices/ITraceServicesModule.h"
+#include "TraceServices/Model/AnalysisSession.h"
+#include "TraceServices/ModuleService.h"
+
+namespace UE::TraceBasedDebuggers
+{
+
+struct FFileDataStream : public Trace::IInDataStream
+{
+	explicit FFileDataStream(TUniquePtr<IFileHandle>&& FileHandle)
+		: AsyncProgressNotification(FAsyncProgressNotification(NSLOCTEXT("TraceBasedDebuggers", "FileLoadingProgressNotification", "Loading File")))
+		, Handle(MoveTemp(FileHandle))
+	{
+		check(Handle);
+		FileSize = Handle->Size();
+		RemainingDataSize = FileSize;
+		constexpr int32 FileFullyLoadedPercentage = 100;
+
+		check(FileSize > 0);
+		ProgressPercentagePerByte = static_cast<double>(FileFullyLoadedPercentage) / static_cast<double>(FileSize);
+	}
+
+	virtual int32 Read(void* Data, uint32 Size) override
+	{
+		if (RemainingDataSize <= 0)
+		{
+			return 0;
+		}
+
+		if (Size > RemainingDataSize)
+		{
+			Size = static_cast<uint32>(RemainingDataSize);
+		}
+
+		RemainingDataSize -= Size;
+
+		const double ProgressPercentage = static_cast<double>(FileSize - RemainingDataSize) * ProgressPercentagePerByte;
+		AsyncProgressNotification.EnterProgress(FMath::CeilToInt32(ProgressPercentage));
+
+		check(Handle);
+		if (!Handle->Read(static_cast<uint8*>(Data), Size))
+		{
+			return 0;
+		}
+		return Size;
+	}
+
+	virtual void Close() override
+	{
+		Handle.Reset();
+	}
+
+private:
+
+	/**
+	 * Object that allows showing a progress bar in the status bar of the editor,
+	 * where the progress can be updated from any thread
+	 */
+	struct FAsyncProgressNotification : FTSTickerObjectBase
+	{
+		FAsyncProgressNotification(const FText& InMessage) : NotificationTitle(InMessage)
+		{
+		}
+
+		virtual ~FAsyncProgressNotification() override;
+
+		/** Updates the current progress
+		 * @param InCurrentProgress Value between 0 and 100, representing the progress we want to display
+		 */
+		void EnterProgress(int32 InCurrentProgress);
+
+	private:
+		virtual bool Tick(float DeltaTime) override;
+
+		float ElapsedTimeSinceLastUpdate = 0.0f;
+
+		FProgressNotificationHandle NotificationHandle;
+		std::atomic<int32> CurrentProgress = 0;
+		FText NotificationTitle;
+	};
+	FAsyncProgressNotification AsyncProgressNotification;
+
+	TUniquePtr<IFileHandle> Handle;
+	uint64 RemainingDataSize = 0;
+	uint64 FileSize = 0;
+	double ProgressPercentagePerByte = 0.0;
+};
+
+FFileDataStream::FAsyncProgressNotification::~FAsyncProgressNotification()
+{
+	if (NotificationHandle.IsValid())
+	{
+		// FSlateNotificationManager is not thread safe, so we need to make sure this runs in the game thread
+		if (IsInGameThread())
+		{
+			FSlateNotificationManager::Get().CancelProgressNotification(NotificationHandle);
+		}
+		else
+		{
+			constexpr float Delay = 0.0f;
+			FTSTicker::GetCoreTicker().AddTicker(TEXT("CancelCVDAsyncNotification"), Delay, [HandleCopy = NotificationHandle](float DeltaTime)
+				{
+					FSlateNotificationManager::Get().CancelProgressNotification(HandleCopy);
+					return false;
+				});
+		}
+	}
+}
+
+void FFileDataStream::FAsyncProgressNotification::EnterProgress(const int32 InCurrentProgress)
+{
+	CurrentProgress = InCurrentProgress;
+}
+
+bool FFileDataStream::FAsyncProgressNotification::Tick(const float DeltaTime)
+{
+	constexpr int32 MaxProgress = 100;
+	constexpr float TimeBetweenUpdatesSeconds = 0.5f;
+
+	ElapsedTimeSinceLastUpdate += DeltaTime;
+	if (ElapsedTimeSinceLastUpdate < TimeBetweenUpdatesSeconds)
+	{
+		return true;
+	}
+
+	ElapsedTimeSinceLastUpdate = 0.0f;
+
+	if (CurrentProgress >= MaxProgress)
+	{
+		if (NotificationHandle.IsValid())
+		{
+			FSlateNotificationManager::Get().UpdateProgressNotification(NotificationHandle, MaxProgress, MaxProgress, NotificationTitle);
+		}
+		NotificationHandle = FProgressNotificationHandle();
+
+		// We are done and we can stop ticking
+		return false;
+	}
+
+	if (!NotificationHandle.IsValid())
+	{
+		NotificationHandle = FSlateNotificationManager::Get().StartProgressNotification(NotificationTitle, MaxProgress);
+	}
+	else
+	{
+		FSlateNotificationManager::Get().UpdateProgressNotification(NotificationHandle, CurrentProgress, MaxProgress, NotificationTitle);
+	}
+
+	return true;
+}
+
+FTraceSessionsManager::FTraceSessionsManager(
+	FRuntimeModule& RuntimeModule
+	, FEngineEditorBridge& Bridge
+	, const TSharedPtr<TraceServices::IModule>& TraceModule
+	, FStringView SaveDirectorySubPathInUserDir
+)
+	: RuntimeModule(RuntimeModule)
+	, EngineEditorBridge(Bridge)
+	, TraceModule(TraceModule)
+	, SaveDirectorySubPathInUserDir(SaveDirectorySubPathInUserDir)
+{
+	if (TraceModule)
+	{
+		IModularFeatures::Get().RegisterModularFeature(TraceServices::ModuleFeatureName, TraceModule.Get());
+	}
+}
+
+FTraceSessionsManager::~FTraceSessionsManager()
+{
+	if (TraceModule)
+	{
+		IModularFeatures::Get().UnregisterModularFeature(TraceServices::ModuleFeatureName, TraceModule.Get());
+	}
+}
+
+FString FTraceSessionsManager::LoadTraceFile(const FString& InTraceFilename, const TFunction<void()>& PreStartAnalysisFunc)
+{
+	CloseSession(InTraceFilename);
+
+	IPlatformFile& FileSystem = IPlatformFile::GetPlatformPhysical();
+
+	constexpr bool bAllowWrite = false;
+	IFileHandle* Handle = FileSystem.OpenRead(*InTraceFilename, bAllowWrite);
+	if (!Handle)
+	{
+		return FString();
+	}
+
+	return LoadTraceFile(TUniquePtr<IFileHandle>(Handle), InTraceFilename, PreStartAnalysisFunc);
+}
+
+FString FTraceSessionsManager::LoadTraceFile(TUniquePtr<IFileHandle>&& InFileHandle, const FString& InTraceSessionName, const TFunction<void()>& PreStartAnalysisFunc)
+{
+	if (InFileHandle->Size() == 0)
+	{
+		return FString();
+	}
+
+	ITraceServicesModule& TraceServicesModule = FModuleManager::LoadModuleChecked<ITraceServicesModule>("TraceServices");
+	if (const TSharedPtr<TraceServices::IAnalysisService> TraceAnalysisService = TraceServicesModule.GetAnalysisService())
+	{
+		if (PreStartAnalysisFunc)
+		{
+			PreStartAnalysisFunc();
+		}
+
+		if (const TSharedPtr<const TraceServices::IAnalysisSession> NewSession = TraceAnalysisService->StartAnalysis(~0, *InTraceSessionName, CreateFileDataStream(MoveTemp(InFileHandle))))
+		{
+			AnalysisSessionByName.Add(InTraceSessionName, NewSession.ToSharedRef());
+
+			return NewSession->GetName();
+		}
+	}
+
+	return FString();
+}
+
+FString FTraceSessionsManager::ConnectToLiveSession(const FStringView InSessionHost, const uint32 SessionID, TFunction<void()> PreStartAnalysisFunc)
+{
+	if (InSessionHost.IsEmpty())
+	{
+		return FString();
+	}
+
+	using namespace UE::Trace;
+	FString LiveSessionName = FString::Printf(TEXT("LiveSession[%.*s - %u]"), InSessionHost.Len(), InSessionHost.GetData(), SessionID);
+
+	FStoreClient::FTraceData TraceData = GetTraceDataStreamFromStore(InSessionHost, SessionID);
+	if (!TraceData.IsValid())
+	{
+		UE_LOG(LogTraceBasedDebuggers, Error, TEXT("[%hs] Failed to get trace data stream from the trace store | Host [%s] | Session ID [%u]"), __func__, InSessionHost.GetData(), SessionID);
+		return FString();
+	}
+
+	if (ConnectToLiveSession_Internal(SessionID, PreStartAnalysisFunc, LiveSessionName, MoveTemp(TraceData)))
+	{
+		return LiveSessionName;
+	}
+
+	UE_LOG(LogTraceBasedDebuggers, Error, TEXT("[%hs] Failed connect to session | Host [%s] | Session ID [%u]"), __func__, InSessionHost.GetData(), SessionID);
+
+	return FString();
+}
+
+FString FTraceSessionsManager::ConnectToLiveSession_Direct(FGuid RemoteSessionID, uint16& OutSessionPort, TFunction<void()> PreStartAnalysisFunc)
+{
+	TUniquePtr<FArchiveFileWriterGeneric> TraceFileWriter = nullptr;
+
+	if (bSaveRecordingsToDisk.IsSet() && bSaveRecordingsToDisk.Get())
+	{
+		TraceFileWriter = OpenTraceFileForWrite(GetSaveDirPath());
+		if (!TraceFileWriter)
+		{
+			UE_LOG(LogTraceBasedDebuggers, Error, TEXT("[%hs] Failed to open open file for direct trace stream"), __func__);
+			return FString();
+		}
+	}
+
+	TUniquePtr<FDirectSocketStream> DirectSocketStream = MakeUnique<FDirectSocketStream>(
+		EngineEditorBridge.GetSessionsManager().ToSharedRef()
+		, RemoteSessionID
+		, MoveTemp(TraceFileWriter));
+
+	OutSessionPort = DirectSocketStream->StartListening(CurrentDirectModeStartingPort);
+
+	if (OutSessionPort == 0)
+	{
+		UE_LOG(LogTraceBasedDebuggers, Error, TEXT("[%hs] Failed to open direct trace stream socket"), __func__);
+		return FString();
+	}
+
+	// If the port was in use, FDirectSocketStream will try multiple ports in sequence, therefore the next port we should try while we have recording active (common case in Multi-session mode)
+	// is the next one from the last success
+	CurrentDirectModeStartingPort = OutSessionPort + 1;
+
+	FString DirectSessionName = FString::Printf(TEXT("DirectSession[127.0.0.1:%u]"), OutSessionPort);
+
+	if (ConnectToLiveSession_Internal(~0u, PreStartAnalysisFunc, DirectSessionName, MoveTemp(DirectSocketStream)))
+	{
+		return DirectSessionName;
+	}
+	else
+	{
+		UE_LOG(LogTraceBasedDebuggers, Error, TEXT("[%hs] Failed start trace analysis with direct trace stream socket | [%s]"), __func__, *DirectSessionName);
+		return FString();
+	}
+}
+
+FString FTraceSessionsManager::ConnectToLiveSession_Relay(const FGuid RemoteSessionID, TFunction<void()> PreStartAnalysisFunc)
+{
+	FString RelaySessionName = FString::Printf(TEXT("RelaySession[%s]"), *RemoteSessionID.ToString());
+
+	TUniquePtr<FRelayDataStream> RelayDataStream = CreateRelayDataStream(RemoteSessionID);
+
+	if (ConnectToLiveSession_Internal(~0u, PreStartAnalysisFunc, RelaySessionName, MoveTemp(RelayDataStream)))
+	{
+		return RelaySessionName;
+	}
+	else
+	{
+		UE_LOG(LogTraceBasedDebuggers, Error, TEXT("[%hs] Failed start trace analysis with relay trace stream | [%s]"), __func__, *RelaySessionName);
+		return FString();
+	}
+}
+
+
+FString FTraceSessionsManager::GetLocalTraceStoreDirPath()
+{
+	const Trace::FStoreClient* StoreClient = Trace::FStoreClient::Connect(TEXT("localhost"));
+
+	if (!StoreClient)
+	{
+		UE_LOG(LogTraceBasedDebuggers, Error, TEXT("[%hs] Failed to connect to local Trace Store client"), __func__);
+		return TEXT("");
+	}
+
+	const Trace::FStoreClient::FStatus* Status = StoreClient->GetStatus();
+	if (!Status)
+	{
+		UE_LOG(LogTraceBasedDebuggers, Error, TEXT("[%hs] Failed to to get Trace Store status"), __func__);
+		return TEXT("");
+	}
+
+	return FString(Status->GetStoreDir());
+}
+
+TSharedPtr<const TraceServices::IAnalysisSession> FTraceSessionsManager::GetSession(const FString& InSessionName)
+{
+	if (TSharedPtr<const TraceServices::IAnalysisSession>* FoundSession = AnalysisSessionByName.Find(InSessionName))
+	{
+		return *FoundSession;
+	}
+
+	return nullptr;
+}
+
+void FTraceSessionsManager::CloseSession(const FString& InSessionName)
+{
+	if (const TSharedPtr<const TraceServices::IAnalysisSession>* Session = AnalysisSessionByName.Find(InSessionName))
+	{
+		if (Session->IsValid())
+		{
+			(*Session)->Stop(true);
+		}
+
+		AnalysisSessionByName.Remove(InSessionName);
+
+		if (AnalysisSessionByName.IsEmpty())
+		{
+			// If we don't have any more session left, we can reset back to the original starting port
+			// As all sockets we had in use are closed at this point
+			CurrentDirectModeStartingPort = DefaultStartingPort;
+		}
+	}
+}
+
+void FTraceSessionsManager::StopSession(const FString& InSessionName)
+{
+	if (const TSharedPtr<const TraceServices::IAnalysisSession>* Session = AnalysisSessionByName.Find(InSessionName))
+	{
+		if (Session->IsValid())
+		{
+			(*Session)->Stop(true);
+		}
+	}
+}
+
+const Trace::FStoreClient::FSessionInfo* FTraceSessionsManager::GetTraceSessionInfo(const FStringView InSessionHost, const FGuid TraceGuid)
+{
+	if (InSessionHost.IsEmpty())
+	{
+		UE_LOG(LogTraceBasedDebuggers, Error, TEXT("[%hs] Failed to connect to trace store. Provided session host is empty"), __func__);
+		return nullptr;
+	}
+
+	using namespace UE::Trace;
+	const FStoreClient* StoreClient = FStoreClient::Connect(InSessionHost.GetData());
+
+	if (!StoreClient)
+	{
+		UE_LOG(LogTraceBasedDebuggers, Error, TEXT("[%hs] Failed to connect to trace store at [%s]"), __func__, InSessionHost.GetData())
+		return nullptr;
+	}
+
+	return StoreClient->GetSessionInfoByGuid(TraceGuid);
+}
+
+FString FTraceSessionsManager::GetTraceFileNameFromStoreForSession(const FStringView InSessionHost, const uint32 SessionID)
+{
+	using namespace UE::Trace;
+
+	if (InSessionHost.IsEmpty())
+	{
+		return FString();
+	}
+
+	const FStoreClient* StoreClient = FStoreClient::Connect(InSessionHost.GetData());
+	if (!StoreClient)
+	{
+		return FString();
+	}
+
+	// Note: The following calls to the store client are scoped because FStoreClient::FTraceInfo and FStoreClient::FStatus share the same buffer
+	// under the hood, therefor although the ptr will still be valid, the data will not after we obtain one after the other
+	FString FileName;
+	{
+		const FStoreClient::FTraceInfo* TraceInfo = StoreClient->GetTraceInfoById(SessionID);
+		if (!TraceInfo)
+		{
+			return FString();
+		}
+
+		const FUtf8StringView Utf8NameView = TraceInfo->GetName();
+		FileName = FString(Utf8NameView);
+		if (!FileName.EndsWith(TEXT(".utrace")))
+		{
+			FileName += TEXT(".utrace");
+		}
+	}
+
+	FString TraceStorePath;
+	{
+		const FStoreClient::FStatus* StoreStatus = StoreClient->GetStatus();
+		if (!StoreStatus)
+		{
+			return FString();
+		}
+
+		TraceStorePath = FString(StoreStatus->GetStoreDir());
+	}
+
+	FString FullFileName = FPaths::Combine(TraceStorePath, FileName);
+	FPaths::NormalizeFilename(FullFileName);
+
+	return FullFileName;
+}
+
+
+Trace::FStoreClient::FTraceData FTraceSessionsManager::GetTraceDataStreamFromStore(const FStringView InSessionHost, const uint32 SessionID)
+{
+	using namespace UE::Trace;
+	const FStoreClient* StoreClient = FStoreClient::Connect(InSessionHost.GetData());
+	if (!StoreClient)
+	{
+		return FStoreClient::FTraceData();
+	}
+
+	return StoreClient->ReadTrace(SessionID);
+}
+
+FString FTraceSessionsManager::GetSaveDirPath() const
+{
+	if (SavePathOverride.IsSet()
+		&& SavePathOverride.Get().FilePath.Len() > 0)
+	{
+		return SavePathOverride.Get().FilePath;
+	}
+
+	return GetDefaultSaveDirPath();
+}
+
+FString FTraceSessionsManager::GetDefaultSaveDirPath() const
+{
+	return FPaths::Combine(FPlatformProcess::UserDir(), SaveDirectorySubPathInUserDir);
+}
+
+TUniquePtr<FArchiveFileWriterGeneric> FTraceSessionsManager::OpenTraceFileForWrite(const FString& InDirectoryToSavePath) const
+{
+	IPlatformFile& PlatformFile = FPlatformFileManager::Get().GetPlatformFile();
+
+	if (InDirectoryToSavePath.IsEmpty())
+	{
+		UE_LOG(LogTraceBasedDebuggers, Error, TEXT("[%hs] Failed open file for write | No save directory was provided"), __func__);
+		return nullptr;
+	}
+
+	PlatformFile.CreateDirectory(*InDirectoryToSavePath);
+	if (!PlatformFile.DirectoryExists(*InDirectoryToSavePath))
+	{
+		UE_LOG(LogTraceBasedDebuggers, Error, TEXT("[%hs] Failed open file for write | Failed to create or access provided save directory [%s]"), __func__, *InDirectoryToSavePath);
+		return nullptr;
+	}
+
+	TUniquePtr<FArchiveFileWriterGeneric> FileWriter = nullptr;
+
+#if WITH_TRACE_BASED_DEBUGGERS
+
+	constexpr int32 MaxAttempts = 10;
+	int32 CurrentAttempts = 0;
+	while (CurrentAttempts < MaxAttempts)
+	{
+		FString OutFileName;
+		RuntimeModule.GenerateRecordingFileName(OutFileName);
+		OutFileName = FPaths::Combine(InDirectoryToSavePath, OutFileName);
+
+		if (IFileHandle* FileHandle = PlatformFile.OpenWrite(*OutFileName))
+		{
+			FileWriter = MakeUnique<FArchiveFileWriterGeneric>(FileHandle, *OutFileName, FileHandle->Tell());
+			break;
+		}
+
+		CurrentAttempts++;
+
+		FPlatformProcess::Sleep(0.1f);
+	}
+
+#endif
+
+	return MoveTemp(FileWriter);
+}
+
+bool FTraceSessionsManager::ConnectToLiveSession_Internal(const uint32 SessionID, const TFunction<void()>& PreStartAnalysisFunc, const FString& InRequestedSessionName, Trace::FStoreClient::FTraceData&& InTraceDataStream)
+{
+	if (!InTraceDataStream.IsValid())
+	{
+		return false;
+	}
+
+	ITraceServicesModule& TraceServicesModule = FModuleManager::LoadModuleChecked<ITraceServicesModule>("TraceServices");
+	if (const TSharedPtr<TraceServices::IAnalysisService> TraceAnalysisService = TraceServicesModule.GetAnalysisService())
+	{
+		// Close this session in case we were already analysing it
+		CloseSession(InRequestedSessionName);
+
+		if (PreStartAnalysisFunc)
+		{
+			PreStartAnalysisFunc();
+		}
+
+		if (const TSharedPtr<const TraceServices::IAnalysisSession> NewSession = TraceAnalysisService->StartAnalysis(SessionID, *InRequestedSessionName, MoveTemp(InTraceDataStream)))
+		{
+			AnalysisSessionByName.Add(InRequestedSessionName, NewSession.ToSharedRef());
+			return true;
+		}
+	}
+
+	return false;
+}
+
+TUniquePtr<Trace::IInDataStream> FTraceSessionsManager::CreateFileDataStream(TUniquePtr<IFileHandle>&& InFileHandle)
+{
+	check(InFileHandle);
+	return TUniquePtr<Trace::IInDataStream>(new FFileDataStream(MoveTemp(InFileHandle)));
+}
+
+TUniquePtr<FRelayDataStream> FTraceSessionsManager::CreateRelayDataStream(FGuid RemoteSessionID) const
+{
+	TUniquePtr<FArchiveFileWriterGeneric> TraceFileWriter = nullptr;
+
+	if (bSaveRecordingsToDisk.IsSet() && bSaveRecordingsToDisk.Get())
+	{
+		TraceFileWriter = OpenTraceFileForWrite(GetSaveDirPath());
+		if (!TraceFileWriter)
+		{
+			UE_LOG(LogTraceBasedDebuggers, Error, TEXT("[%hs] Failed to open file for relay trace stream"), __func__);
+			return nullptr;
+		}
+	}
+
+	const TSharedPtr<IDataRelayTransport> DataRelay = EngineEditorBridge.GetTraceRelayTransportInstance();
+	const TSharedPtr<FRemoteSessionsManager> SessionsManager = EngineEditorBridge.GetSessionsManager();
+	if (ensure(DataRelay && SessionsManager))
+	{
+		TUniquePtr<FRelayDataStream> NewDataStream = MakeUnique<FRelayDataStream>(
+			SessionsManager.ToSharedRef()
+			, DataRelay.ToSharedRef()
+			, RemoteSessionID
+			, MoveTemp(TraceFileWriter));
+
+		return MoveTemp(NewDataStream);
+	}
+
+	return nullptr;
+}
+
+uint16 FTraceSessionsManager::CurrentDirectModeStartingPort = DefaultStartingPort;
+
+} // UE::TraceBasedDebuggers
+
+#endif // UE_WITH_TRACE_BASED_DEBUGGERS_ANALYSIS
\ No newline at end of file
diff --git a/Engine/Source/Runtime/Experimental/TraceBasedDebuggersAnalysis/Public/TraceDataStreamTypes.h b/Engine/Source/Runtime/Experimental/TraceBasedDebuggersAnalysis/Public/TraceDataStreamTypes.h
new file mode 100644
--- /dev/null
+++ b/Engine/Source/Runtime/Experimental/TraceBasedDebuggersAnalysis/Public/TraceDataStreamTypes.h
@@ -0,0 +1,177 @@
+﻿// Copyright Epic Games, Inc. All Rights Reserved.
+
+#pragma once
+
+#if UE_WITH_TRACE_BASED_DEBUGGERS_ANALYSIS
+
+#include "Containers/SpscQueue.h"
+#include "Containers/Ticker.h"
+#include "Misc/Guid.h"
+#include "Trace/DataStream.h"
+
+#define UE_API TRACEBASEDDEBUGGERSANALYSIS_API
+
+namespace TraceServices
+{
+class IAnalysisSession;
+}
+
+namespace UE::TraceBasedDebuggers
+{
+struct FRemoteSessionsManager;
+class IDataRelayTransport;
+
+/** Wrapper around a data buffer used for serialization of trace-based debuggers data */
+struct FSerializedDataBuffer
+{
+	FSerializedDataBuffer(const TConstArrayView<uint8> ConstView)
+		: DataBuffer(ConstView.GetData(), ConstView.Num())
+	{
+	}
+
+	FSerializedDataBuffer(TArray<uint8>&& InByteArray)
+		: DataBuffer(MoveTemp(InByteArray))
+	{
+	}
+
+	FSerializedDataBuffer(const TArray<uint8>& InByteArray)
+		: DataBuffer(InByteArray)
+	{
+	}
+
+	FSerializedDataBuffer()
+	{
+	}
+
+	TArray<uint8>& GetDataAsByteArrayRef()
+	{
+		return DataBuffer;
+	}
+
+	int32 GetSize() const
+	{
+		return DataBuffer.Num();
+	}
+
+	void Serialize(FArchive& Ar)
+	{
+		Ar << DataBuffer;
+	}
+
+private:
+	TArray<uint8> DataBuffer;
+};
+
+inline FArchive& operator<<(FArchive& Ar, FSerializedDataBuffer& Data)
+{
+	Data.Serialize(Ar);
+	return Ar;
+}
+
+class IRemoteSessionDataStream
+{
+public:
+	virtual ~IRemoteSessionDataStream() = default;
+	virtual FGuid GetOwningRemoteSessionID() const = 0;
+	virtual uint64 UpdateAndGetBytesReadSinceLastMeasurement() = 0;
+	virtual TSharedPtr<FRemoteSessionsManager> GetSessionsManager() const = 0;
+};
+
+class FSessionStreamStatsUpdater : FTSTickerObjectBase
+{
+public:
+	virtual ~FSessionStreamStatsUpdater() override;
+
+	explicit FSessionStreamStatsUpdater(IRemoteSessionDataStream* const DataStreamInstance)
+		: DataStreamInstance(DataStreamInstance)
+	{
+	}
+
+private:
+	virtual bool Tick(float DeltaTime) override;
+
+	IRemoteSessionDataStream* DataStreamInstance;
+
+	double ElapsedTimeSinceLastUpdate = 0.0;
+};
+
+class FDirectSocketStream : public Trace::FDirectSocketStream, public IRemoteSessionDataStream
+{
+public:
+
+	FDirectSocketStream() = delete;
+
+	explicit FDirectSocketStream(const TSharedRef<FRemoteSessionsManager>& InSessionsManager
+		, const FGuid& InRemoteSessionID
+		, TUniquePtr<FArchiveFileWriterGeneric>&& InFileWriter
+		)
+		: Trace::FDirectSocketStream(MoveTemp(InFileWriter))
+		, WeakSessionManager(InSessionsManager.ToWeakPtr())
+		, BytesReadSinceLastMeasurement(0)
+		, RemoteSessionID(InRemoteSessionID)
+		, StatsUpdater(this)
+	{
+	}
+
+	virtual int32 Read(void* Data, uint32 Size) override;
+
+	virtual FGuid GetOwningRemoteSessionID() const override;
+	virtual uint64 UpdateAndGetBytesReadSinceLastMeasurement() override;
+	virtual TSharedPtr<FRemoteSessionsManager> GetSessionsManager() const override;
+
+private:
+	TWeakPtr<FRemoteSessionsManager> WeakSessionManager;
+	std::atomic<uint64> BytesReadSinceLastMeasurement;
+	FGuid RemoteSessionID;
+
+	FSessionStreamStatsUpdater StatsUpdater;
+};
+
+/** A trace data stream object that listens for data coming from a relay stream via provided DataRelayTransport instance */
+class FRelayDataStream : public Trace::IInDataStream, public IRemoteSessionDataStream
+{
+public:
+	explicit FRelayDataStream(const TSharedRef<FRemoteSessionsManager>& InSessionsManager
+		, const TSharedRef<IDataRelayTransport>& InDataRelay
+		, const FGuid& InRemoteSessionID
+		, TUniquePtr<FArchiveFileWriterGeneric>&& InFileWriter);
+
+	virtual ~FRelayDataStream() override;
+
+	void EnqueueRelayedData(const TConstArrayView<uint8> InTraceDataBuffer);
+	void EnqueueRelayedData(TArray<uint8>&& InTraceDataBuffer);
+
+	virtual FGuid GetOwningRemoteSessionID() const override;
+	virtual uint64 UpdateAndGetBytesReadSinceLastMeasurement() override;
+	virtual TSharedPtr<FRemoteSessionsManager> GetSessionsManager() const override;
+
+private:
+	virtual int32 Read(void* Dest, uint32 DestSize) override;
+	virtual void Close() override;
+	virtual bool WaitUntilReady() override;
+
+	TWeakPtr<FRemoteSessionsManager> WeakSessionManager;
+	TSpscQueue<FSerializedDataBuffer> DataQueue;
+	FEvent* CancelEvent = nullptr;
+	FEvent* WaitForDataEvent = nullptr;
+	TOptional<FSerializedDataBuffer> CurrentPacket;
+
+	FGuid RemoteSessionID = FGuid();
+
+	TUniquePtr<FArchiveFileWriterGeneric> FileWriter;
+
+#if WITH_TRACE_BASED_DEBUGGERS
+	TSharedPtr<IDataRelayTransport> RelayTransportInstance;
+#endif
+
+	std::atomic<uint64> BytesReadSinceLastMeasurement;
+	FSessionStreamStatsUpdater StatsUpdater;
+	uint32 CurrentPacketOffset = 0;
+	bool bClosed = false;
+};
+
+} // UE::TraceBasedDebuggers
+
+#undef UE_API
+
+#endif // UE_WITH_TRACE_BASED_DEBUGGERS_ANALYSIS
\ No newline at end of file
diff --git a/Engine/Source/Runtime/Experimental/TraceBasedDebuggersAnalysis/Public/TraceSessionsManager.h b/Engine/Source/Runtime/Experimental/TraceBasedDebuggersAnalysis/Public/TraceSessionsManager.h
new file mode 100644
--- /dev/null
+++ b/Engine/Source/Runtime/Experimental/TraceBasedDebuggersAnalysis/Public/TraceSessionsManager.h
@@ -0,0 +1,220 @@
+﻿// Copyright Epic Games, Inc. All Rights Reserved.
+
+#pragma once
+
+#if UE_WITH_TRACE_BASED_DEBUGGERS_ANALYSIS
+
+#include "Misc/Attribute.h"
+#include "Misc/Guid.h"
+#include "Templates/SharedPointer.h"
+#include "Trace/StoreClient.h"
+#include "UObject/SoftObjectPath.h"
+
+namespace UE::Trace
+{
+class IInDataStream;
+}
+
+class FArchiveFileWriterGeneric;
+
+namespace TraceServices
+{
+class IModule;
+class IAnalysisSession;
+}
+
+namespace UE::TraceBasedDebuggers
+{
+class FEngineEditorBridge;
+class FRuntimeModule;
+class FRelayDataStream;
+
+#define UE_API TRACEBASEDDEBUGGERSANALYSIS_API
+
+/**
+ * Structure containing info about a trace session used by trace-based debuggers
+ */
+struct FTraceSessionDescriptor
+{
+	/** The trace session name. Used to get the trace analysis instance for this session*/
+	FString SessionName;
+
+	/** Session ID of this session in the Remote Session management system*/
+	FGuid RemoteSessionID;
+
+	/** Port number used to establish a connection to this session. The value is 0 if no specific port was used*/
+	uint16 SessionPort = 0;
+
+	/** True if this is a session recorded live (not loaded from disk)*/
+	bool bIsLiveSession = false;
+
+	/** Returns true if valid. For a session to be valued needs at minimum have a Trace Session name */
+	bool IsValid() const
+	{
+		return !SessionName.IsEmpty();
+	}
+
+	friend bool operator==(const FTraceSessionDescriptor& Lhs, const FTraceSessionDescriptor& RHS)
+	{
+		return Lhs.SessionName == RHS.SessionName
+			&& Lhs.RemoteSessionID == RHS.RemoteSessionID;
+	}
+
+	friend bool operator!=(const FTraceSessionDescriptor& Lhs, const FTraceSessionDescriptor& RHS)
+	{
+		return !(Lhs == RHS);
+	}
+};
+
+/** Manager class used by trace-based debuggers to interact/control UE Trace systems */
+class FTraceSessionsManager
+{
+#if NO_LOGGING
+	typedef FNoLoggingCategory FLogCategoryAlias;
+#else
+	typedef FLogCategoryBase FLogCategoryAlias;
+#endif
+
+
+public:
+	/**
+	 * @param RuntimeModule Debugger runtime module used to provide names for the trace files
+	 * @param Bridge The engine-editor bridge instance used to provide access to the remote sessions manager
+	 * @param TraceModule Optional trace module that will be registered/deregistered as a ModularFeature for the TraceServices
+	 * @param SaveDirectorySubPathInUserDir Subdirectory path that will be used by GetDefaultSaveDirPath
+	 */
+	UE_API FTraceSessionsManager(FRuntimeModule& RuntimeModule
+		, FEngineEditorBridge& Bridge
+		, const TSharedPtr<TraceServices::IModule>& TraceModule
+		, FStringView SaveDirectorySubPathInUserDir);
+	UE_API ~FTraceSessionsManager();
+
+	void SetSavePathOverride(const TAttribute<FFilePath>& InSavePathOverride)
+	{
+		SavePathOverride = InSavePathOverride;
+	}
+
+	void SetSaveRecordingsToDisk(const TAttribute<bool>& bInSaveRecordingsToDisk)
+	{
+		bSaveRecordingsToDisk = bInSaveRecordingsToDisk;
+	}
+
+	/** 
+	 * Loads a trace file and starts analyzing it
+	 * @param InTraceFilename File Name including (Path Included) of the Trace file to load
+	 * @param PreStartAnalysisFunc Optional function called right before starting the trace analysis (i.e., TraceServices::IAnalysisService::StartAnalysis)
+	 */
+	UE_API FString LoadTraceFile(const FString& InTraceFilename, const TFunction<void()>& PreStartAnalysisFunc = nullptr);
+
+	UE_API FString LoadTraceFile(TUniquePtr<IFileHandle>&& InFileHandle, const FString& InTraceSessionName, const TFunction<void()>& PreStartAnalysisFunc = nullptr);
+
+	/**
+	 * Connects to a live Trace Session and starts analyzing it.
+	 * @param InSessionHost Trace Store Address for this session
+	 * @param SessionID Trace ID in the Trace Store provided as host
+	 * @param PreStartAnalysisFunc Optional function called right before starting the trace analysis (i.e., TraceServices::IAnalysisService::StartAnalysis)
+	 * @return Session Name
+	 */
+	UE_API FString ConnectToLiveSession(FStringView InSessionHost, uint32 SessionID, TFunction<void()> PreStartAnalysisFunc = nullptr);
+
+	/**
+	 * Opens up a direct trace stream and waits for data
+	 * @param RemoteSessionID Session ID to which we want to connect to
+	 * @param OutSessionPort Number of the port we open for this session
+	 * @param PreStartAnalysisFunc Optional function called right before starting the trace analysis (i.e., TraceServices::IAnalysisService::StartAnalysis)
+	 * @return Session Name
+	 */
+	UE_API FString ConnectToLiveSession_Direct(FGuid RemoteSessionID, uint16& OutSessionPort, TFunction<void()> PreStartAnalysisFunc = nullptr);
+
+	/**
+	 * Opens up a relay trace stream and waits for data
+	 * @param RemoteSessionID Session ID to which we want to connect to
+	 * @param PreStartAnalysisFunc Optional function called right before starting the trace analysis (i.e., TraceServices::IAnalysisService::StartAnalysis)
+	 * @return Session Name
+	 */
+	UE_API FString ConnectToLiveSession_Relay(FGuid RemoteSessionID, TFunction<void()> PreStartAnalysisFunc = nullptr);
+
+	/** Returns the path to the local trace store */
+	UE_API FString GetLocalTraceStoreDirPath();
+
+	/** Returns a ptr to the session registered with the provided session name. Null if no session is found */
+	UE_API TSharedPtr<const TraceServices::IAnalysisSession> GetSession(const FString& InSessionName);
+
+	/** Stops and de-registers a trace session registered with the provided session name */
+	UE_API void CloseSession(const FString& InSessionName);
+
+	/** Stops a trace session registered with the provided session name */
+	UE_API void StopSession(const FString& InSessionName);
+
+	UE_API static const Trace::FStoreClient::FSessionInfo* GetTraceSessionInfo(FStringView InSessionHost, FGuid TraceGuid);
+
+	template <typename TVisitor>
+	void EnumerateActiveSessions(TVisitor Callback);
+
+	/** Access the trace store at the provided host address, and returns the file name for the trace file
+	 * @param InSessionHost Address of the Trace Store
+	 * @param SessionID Session id of the trace data from which get the file name from
+	 * @return Full Filename with path to the trace file if it exists. Returns an empty string if a trace file cannot be found
+	 */
+	UE_API FString GetTraceFileNameFromStoreForSession(FStringView InSessionHost, uint32 SessionID);
+
+	/** Access the trace store at the provided host address, and returns a trace data stream ready to use
+	 * @param InSessionHost Address of the Trace Store
+	 * @param SessionID Session id of the trace data
+	 */
+	UE_API Trace::FStoreClient::FTraceData GetTraceDataStreamFromStore(FStringView InSessionHost, uint32 SessionID);
+
+	/**
+	 * Returns the actual path where recordings should be saved.
+	 * That path is either the override provided by the path override or
+	 * the default one.
+	 * @see SetSavePathOverride
+	 * @see GetDefaultSaveDirPath
+	 * */
+	UE_API FString GetSaveDirPath() const;
+
+	/** @return The default path inside the user directory where recordings should be saved. */
+	UE_API FString GetDefaultSaveDirPath() const;
+
+private:
+	TUniquePtr<FArchiveFileWriterGeneric> OpenTraceFileForWrite(const FString& InDirectoryToSavePath) const;
+
+	bool ConnectToLiveSession_Internal(uint32 SessionID, const TFunction<void()>& PreStartAnalysisFunc, const FString& InRequestedSessionName, Trace::FStoreClient::FTraceData&& InTraceDataStream);
+
+	static TUniquePtr<Trace::IInDataStream> CreateFileDataStream(TUniquePtr<IFileHandle>&& InFileHandle);
+	TUniquePtr<FRelayDataStream> CreateRelayDataStream(FGuid RemoteSessionID) const;
+
+	FRuntimeModule& RuntimeModule;
+	FEngineEditorBridge& EngineEditorBridge;
+
+	/** The trace analysis sessions. */
+	TMap<FString, TSharedPtr<const TraceServices::IAnalysisSession>> AnalysisSessionByName;
+
+	TSharedPtr<TraceServices::IModule> TraceModule;
+
+	const FString SaveDirectorySubPathInUserDir;
+	TAttribute<FFilePath> SavePathOverride;
+	TAttribute<bool> bSaveRecordingsToDisk = true;
+
+	static uint16 CurrentDirectModeStartingPort;
+	static constexpr uint16 DefaultStartingPort = 1993;
+};
+
+template <typename TVisitor>
+void FTraceSessionsManager::EnumerateActiveSessions(TVisitor Callback)
+{
+	for (const TPair<FString, TSharedPtr<const TraceServices::IAnalysisSession>>& SessionWithName : AnalysisSessionByName)
+	{
+		if (SessionWithName.Value)
+		{
+			if (!Callback(SessionWithName.Value.ToSharedRef()))
+			{
+				return;
+			}
+		}
+	}
+}
+
+} // UE::TraceBasedDebuggers
+#undef UE_API
+#endif // UE_WITH_TRACE_BASED_DEBUGGERS_ANALYSIS
\ No newline at end of file
diff --git a/Engine/Source/Runtime/Experimental/TraceBasedDebuggersAnalysis/TraceBasedDebuggersAnalysis.Build.cs b/Engine/Source/Runtime/Experimental/TraceBasedDebuggersAnalysis/TraceBasedDebuggersAnalysis.Build.cs
new file mode 100644
--- /dev/null
+++ b/Engine/Source/Runtime/Experimental/TraceBasedDebuggersAnalysis/TraceBasedDebuggersAnalysis.Build.cs
@@ -0,0 +1,59 @@
+// Copyright Epic Games, Inc. All Rights Reserved.
+
+namespace UnrealBuildTool.Rules
+{
+	/** 
+	 * This module is extending the sharable functionalities of the TraceBasedDebuggers module
+	 * for targets supporting trace analysis.
+	 */
+	public class TraceBasedDebuggersAnalysis : ModuleRules
+	{
+		public TraceBasedDebuggersAnalysis(ReadOnlyTargetRules Target) : base(Target)
+		{
+			PublicDependencyModuleNames.AddRange(
+				new string[]
+				{
+					"Core",
+					"CoreUObject",
+					"TraceBasedDebuggers"
+				}
+				);
+
+			PrivateDependencyModuleNames.AddRange(
+				new string[]
+				{
+					"Slate",
+				}
+				);
+
+			bAllowUETypesInNamespaces = true;
+			CppCompileWarningSettings.UnsafeTypeCastWarningLevel = WarningLevel.Error;
+
+			// Allow debugger trace analysis on desktop platforms
+			if (Target.Platform.IsInGroup(UnrealPlatformGroup.Desktop))
+			{
+				PublicDefinitions.Add("UE_WITH_TRACE_BASED_DEBUGGERS_ANALYSIS=1");
+				PublicDependencyModuleNames.AddRange(
+					new[]
+					{
+						"TraceServices",
+						"TraceAnalysis"
+					}
+				);
+			}
+			else
+			{
+				PublicDefinitions.Add("UE_WITH_TRACE_BASED_DEBUGGERS_ANALYSIS=0");
+			}
+
+			if (AreTraceBasedDebuggersSupported(Target))
+			{
+				PublicDefinitions.Add("WITH_TRACE_BASED_DEBUGGERS=1");
+			}
+			else
+			{
+				PublicDefinitions.Add("WITH_TRACE_BASED_DEBUGGERS=0");
+			}
+		}
+	}
+}
