Chaos - undo change which disabled friction on first-contact frames

#rnx
#Synced-CL 49124308

diff --git a/Engine/Source/Runtime/Experimental/Chaos/Private/Chaos/Collision/PBDCollisionContainerSolver.cpp b/Engine/Source/Runtime/Experimental/Chaos/Private/Chaos/Collision/PBDCollisionContainerSolver.cpp
--- a/Engine/Source/Runtime/Experimental/Chaos/Private/Chaos/Collision/PBDCollisionContainerSolver.cpp
+++ b/Engine/Source/Runtime/Experimental/Chaos/Private/Chaos/Collision/PBDCollisionContainerSolver.cpp
@@ -157,7 +157,9 @@ namespace Chaos
 			}
 		}
 
-		if (ManifoldPoint.Flags.bInitialContact)
+		// Optionally disable friction on the first contact frames
+		// @todo(chaos): this should be tied to its own setting rather than piggy-backing on bUseSplitImpulse
+		if (ManifoldPoint.Flags.bInitialContact && SolverSettings.bUseSplitImpulse)
 		{
 			WorldContactDeltaTangentU = 0;
 			WorldContactDeltaTangentV = 0;
@@ -692,7 +694,7 @@ namespace Chaos
 			UpdateCollisionConstraintFromSolver(Constraint, CollisionSolver, Dt);
 
 			// Reset the collision solver here as the pointers will be invalid on the next tick
-			// @todo(chaos): Pointers are always reinitalized before use next tick so this isn't strictly necessary
+			// @todo(chaos): Pointers are always reinitialized before use next tick so this isn't strictly necessary
 			CollisionSolver.Reset(nullptr, 0);
 		}
 	}
