Log error messages if we fail to parse encryption guid/key pairs from the ini file or cmdline.

#rnx

#Synced-CL 48862713

diff --git a/Engine/Source/Runtime/Experimental/IoStore/OnDemand/Private/OnDemandConfig.cpp b/Engine/Source/Runtime/Experimental/IoStore/OnDemand/Private/OnDemandConfig.cpp
--- a/Engine/Source/Runtime/Experimental/IoStore/OnDemand/Private/OnDemandConfig.cpp
+++ b/Engine/Source/Runtime/Experimental/IoStore/OnDemand/Private/OnDemandConfig.cpp
@@ -110,7 +110,10 @@ static bool TryParseConfigContent(const FString& ConfigContent, const FString& C
 	FString ContentKey;
 	if (ConfigFile.GetString(TEXT("Endpoint"), TEXT("ContentKey"), ContentKey))
 	{
-		ApplyEncryptionKeyFromString(ContentKey);
+		if (!ApplyEncryptionKeyFromString(ContentKey))
+		{
+			UE_LOG(LogIoStoreOnDemand, Error, TEXT("Failed to parse encryption key from '%s'"), *ConfigFileName);
+		}
 	}
 
 	return OutConfig.IsValid();
@@ -221,7 +224,10 @@ TIoStatusOr<FOnDemandIoStoreConfig> TryParseConfig(const TCHAR* CommandLine)
 		FString EncryptionKey;
 		if (FParse::Value(CommandLine, TEXT("Ias.EncryptionKey="), EncryptionKey))
 		{
-			ApplyEncryptionKeyFromString(EncryptionKey);
+			if (!ApplyEncryptionKeyFromString(EncryptionKey))
+			{
+				UE_LOG(LogIoStoreOnDemand, Error, TEXT("Failed to parse encryption key from cmdline 'Ias.EncryptionKey'"));
+			}
 		}
 	}
 
