Include host group name when serializing on-demand TOC to JSON.

#rnx
#Synced-CL 49075355

diff --git a/Engine/Source/Runtime/Experimental/IoStore/OnDemand/Private/Tool/ListTocs.cpp b/Engine/Source/Runtime/Experimental/IoStore/OnDemand/Private/Tool/ListTocs.cpp
--- a/Engine/Source/Runtime/Experimental/IoStore/OnDemand/Private/Tool/ListTocs.cpp
+++ b/Engine/Source/Runtime/Experimental/IoStore/OnDemand/Private/Tool/ListTocs.cpp
@@ -263,6 +263,7 @@ static void ToJson(FJsonWriter JsonWriter, const FOnDemandToc& Toc, EOnDemandToc
 		JsonWriter->WriteValue(TEXT("BlockSize"), uint64(Toc.Header.BlockSize));
 		JsonWriter->WriteValue(TEXT("CompressionFormat"), Toc.Header.CompressionFormat);
 		JsonWriter->WriteValue(TEXT("ChunksDirectory"), Toc.Header.ChunksDirectory);
+		JsonWriter->WriteValue(TEXT("HostGroupName"), Toc.Header.HostGroupName);
 	}
 	JsonWriter->WriteObjectEnd();
 	JsonWriter->WriteObjectStart(TEXT("Meta"));
