[Chaos Visual Debugger]Adding a IsValid check to prevent a crash when exiting PIE session while recording is still on in Chaos Visual Debugger

#Synced-CL 49127736

diff --git a/Engine/Source/Runtime/Experimental/ChaosSolverEngine/Private/Chaos/ChaosVDEngineEditorBridge.cpp b/Engine/Source/Runtime/Experimental/ChaosSolverEngine/Private/Chaos/ChaosVDEngineEditorBridge.cpp
--- a/Engine/Source/Runtime/Experimental/ChaosSolverEngine/Private/Chaos/ChaosVDEngineEditorBridge.cpp
+++ b/Engine/Source/Runtime/Experimental/ChaosSolverEngine/Private/Chaos/ChaosVDEngineEditorBridge.cpp
@@ -217,7 +217,8 @@ FChaosVDParticleMetadata FChaosVDEngineEditorBridge::GenerateParticleMetadata(co
 	// As we need to access their uobject owners to generate the name, we need to make sure these are not GC'd or wait until GC is done.
 	FGCScopeGuard GCLock;
 
-	if (UObject* ProxyOwner = ParticleProxy->GetOwner())
+	UObject* ProxyOwner = ParticleProxy->GetOwner();
+	if (IsValid(ProxyOwner))
 	{
 		if (UPrimitiveComponent* OwnerAsPrimitiveComponent = Cast<UPrimitiveComponent>(ProxyOwner))
 		{
