[HWRT/GeometryCollection] Modified how r.RayTracing.Geometry.GeometryCollection is applied to affect HasRayTracingRepresentation() instead of filtering during GetDynamicRayTracingInstances().

#Synced-CL 49186836

diff --git a/Engine/Source/Runtime/Experimental/GeometryCollectionEngine/Private/GeometryCollection/GeometryCollectionSceneProxy.cpp b/Engine/Source/Runtime/Experimental/GeometryCollectionEngine/Private/GeometryCollection/GeometryCollectionSceneProxy.cpp
--- a/Engine/Source/Runtime/Experimental/GeometryCollectionEngine/Private/GeometryCollection/GeometryCollectionSceneProxy.cpp
+++ b/Engine/Source/Runtime/Experimental/GeometryCollectionEngine/Private/GeometryCollection/GeometryCollectionSceneProxy.cpp
@@ -82,6 +82,10 @@ FAutoConsoleVariableRef CVarRayTracingGeometryCollection(
 	TEXT("r.RayTracing.Geometry.GeometryCollection"),
 	GRayTracingGeometryCollection,
 	TEXT("Include geometry collection proxy meshes in ray tracing effects (default = 0 (Geometry collection meshes disabled in ray tracing))"),
+	FConsoleVariableDelegate::CreateLambda([](IConsoleVariable* InVariable)
+		{
+			FGlobalComponentRecreateRenderStateContext Context;
+		}),
 	ECVF_RenderThreadSafe
 );
 
@@ -151,7 +155,7 @@ FGeometryCollectionSceneProxyBase::FGeometryCollectionSceneProxyBase(UGeometryCo
 	, bUseShaderBoneTransform(VertexFactory.UseShaderBoneTransform(Component->GetScene()->GetShaderPlatform()))
 	, bSupportsTripleBufferVertexUpload(GRHISupportsMapWriteNoOverwrite)
 #if RHI_RAYTRACING
-	, bHasRayTracingRepresentation(IsRayTracingEnabled() && Component->GetRestCollection()->RenderData->PartRayTracingGeometries.Num() > 0)
+	, bHasRayTracingRepresentation(GRayTracingGeometryCollection && IsRayTracingEnabled() && Component->GetRestCollection()->RenderData->PartRayTracingGeometries.Num() > 0)
 
 #endif
 {
@@ -552,11 +556,6 @@ void FGeometryCollectionSceneProxyBase::GetDynamicRayTracingInstances(FRayTracin
 {
 	checkf(bHasRayTracingRepresentation, TEXT("Shouldn't try to get ray tracing instances from proxy that doesn't have a ray tracing representation."));
 
-	if (GRayTracingGeometryCollection == 0)
-	{
-		return;
-	}
-
 	if (MeshDescription.Sections.IsEmpty())
 	{
 		return;
