[HWRT] Updated FNaniteGeometryCollectionSceneProxy to generate ray tracing instances using GPUScene instead of CPU transforms.
- FNaniteGeometryCollectionSceneProxy already stored collection transforms in GPU Scene as instances.

#Synced-CL 49186836

diff --git a/Engine/Source/Runtime/Experimental/GeometryCollectionEngine/Private/GeometryCollection/GeometryCollectionSceneProxy.cpp b/Engine/Source/Runtime/Experimental/GeometryCollectionEngine/Private/GeometryCollection/GeometryCollectionSceneProxy.cpp
--- a/Engine/Source/Runtime/Experimental/GeometryCollectionEngine/Private/GeometryCollection/GeometryCollectionSceneProxy.cpp
+++ b/Engine/Source/Runtime/Experimental/GeometryCollectionEngine/Private/GeometryCollection/GeometryCollectionSceneProxy.cpp
@@ -668,12 +668,19 @@ void FGeometryCollectionSceneProxyBase::GetDynamicRayTracingInstances(FRayTracin
 	{
 		check(bUseSubSections);
 
-		TArray<FMatrix, SceneRenderingAllocator>& PartInstanceTransforms = Collector.AllocateOneFrameResource<TArray<FMatrix, SceneRenderingAllocator>>();
-		PartInstanceTransforms.AddUninitialized(DynamicData->Transforms.Num());
+		TConstArrayView<FMatrix> PartInstanceTransforms; // only used when !bIsNanite
 
-		for (int32 TransformIndex = 0; TransformIndex < DynamicData->Transforms.Num(); ++TransformIndex)
+		if (!bIsNanite)
 		{
-			PartInstanceTransforms[TransformIndex] = FMatrix(DynamicData->Transforms[TransformIndex]) * LocalToWorld;
+			TArray<FMatrix, SceneRenderingAllocator>& Tmp = Collector.AllocateOneFrameResource<TArray<FMatrix, SceneRenderingAllocator>>();
+			Tmp.AddUninitialized(DynamicData->Transforms.Num());
+
+			for (int32 TransformIndex = 0; TransformIndex < DynamicData->Transforms.Num(); ++TransformIndex)
+			{
+				Tmp[TransformIndex] = FMatrix(DynamicData->Transforms[TransformIndex]) * LocalToWorld;
+			}
+
+			PartInstanceTransforms = Tmp;
 		}
 
 		bool bInstanceMaskAndFlagsDirty = false;
@@ -698,16 +705,46 @@ void FGeometryCollectionSceneProxyBase::GetDynamicRayTracingInstances(FRayTracin
 
 		checkf(PartRayTracingGeometries.Num() == SectionArray.Num(), TEXT("Expected to have one part per section (%d vs %d)."), PartRayTracingGeometries.Num(), SectionArray.Num());
 
+		TArrayView<uint32> PrimitiveInstanceIndices; // only used when bIsNanite
+
+		if (bIsNanite)
+		{
+			checkf(TransformToInstanceMapping.Num() == DynamicData->Transforms.Num(), TEXT("Invalid transform to instance mapping. There are %d transforms but mapping array only has %d entries."), DynamicData->Transforms.Num(), TransformToInstanceMapping.Num());
+
+			// pre-allocate array of primitive instance indices to avoid multiple allocations and copies later
+			TArray<uint32, SceneRenderingAllocator>& Tmp = Collector.AllocateOneFrameResource<TArray<uint32, SceneRenderingAllocator>>();
+			Tmp.AddUninitialized(PartRayTracingGeometries.Num());
+
+			PrimitiveInstanceIndices = Tmp;
+		}
+
 		for (int32 SectionIndex = 0; SectionIndex < PartRayTracingGeometries.Num(); ++SectionIndex)
 		{
 			const FGeometryCollectionMeshElement& Section = SectionArray[SectionIndex];
 
 			FRayTracingInstance RayTracingInstance;
 			RayTracingInstance.Geometry = &PartRayTracingGeometries[SectionIndex];
-			RayTracingInstance.InstanceTransformsView = MakeConstArrayView(&PartInstanceTransforms[Section.TransformIndex], 1);
 			RayTracingInstance.MaterialsView = MakeConstArrayView(&CachedPartRayTracingMaterials[SectionIndex], 1);
 			RayTracingInstance.bInstanceMaskAndFlagsDirty = bInstanceMaskAndFlagsDirty;
 
+			if (bIsNanite)
+			{
+				int32 InstanceIndex = TransformToInstanceMapping[Section.TransformIndex];
+
+				if (InstanceIndex < 0)
+				{
+					continue;
+				}
+
+				PrimitiveInstanceIndices[SectionIndex] = InstanceIndex;
+
+				RayTracingInstance.PrimitiveInstanceIndicesView = MakeConstArrayView(&PrimitiveInstanceIndices[SectionIndex], 1); // parts are stored in GPU Scene as instances
+			}
+			else
+			{
+				RayTracingInstance.InstanceTransformsView = MakeConstArrayView(&PartInstanceTransforms[Section.TransformIndex], 1);
+			}
+
 			for (int32 ViewIndex = 0; ViewIndex < Views.Num(); ViewIndex++)
 			{
 				if ((VisibilityMap & (1 << ViewIndex)) == 0)
@@ -1532,6 +1569,8 @@ void FNaniteGeometryCollectionSceneProxy::UpdateInstanceSceneDataBuffers(const F
 	check(TransformCount == TransformToGeometryIndices.Num());
 	check(TransformCount == TransformChildren.Num());
 
+	TransformToInstanceMapping.Reset(TransformCount);
+
 	// set the prev by copying the last current
 	ProxyData.PrevInstanceToPrimitiveRelative = ProxyData.InstanceToPrimitiveRelative;
 	bCanSkipRedundantTransformUpdates = false; // should we compare the transform to better decide about this ? 
@@ -1553,9 +1592,12 @@ void FNaniteGeometryCollectionSceneProxy::UpdateInstanceSceneDataBuffers(const F
 		const int32 TransformToGeometryIndex = TransformToGeometryIndices[TransformIndex];
 		if (SimulationType[TransformIndex] != FGeometryCollection::ESimulationTypes::FST_Rigid)
 		{
+			TransformToInstanceMapping.Add(-1);
 			continue;
 		}
 
+		TransformToInstanceMapping.Add(ProxyData.InstanceToPrimitiveRelative.Num());
+
 		const FGeometryNaniteData& NaniteData = GeometryNaniteData[TransformToGeometryIndex];
 		const FRenderTransform& InstanceToPrimitiveRelative = ProxyData.InstanceToPrimitiveRelative.Emplace_GetRef(InstanceSceneDataBuffersImpl.ComputeInstanceToPrimitiveRelative(DynamicData->Transforms[TransformIndex], AccessTag));
 
@@ -1588,20 +1630,6 @@ void FNaniteGeometryCollectionSceneProxy::SetDynamicData_RenderThread(FRHIComman
 	UpdateInstanceSceneDataBuffers(PrimitiveLocalToWorld);
 }
 
-void FNaniteGeometryCollectionSceneProxy::ResetPreviousTransforms_RenderThread()
-{
-	FInstanceSceneDataBuffers::FAccessTag AccessTag(PointerHash(this));
-	FInstanceSceneDataBuffers::FWriteView ProxyData = InstanceSceneDataBuffersImpl.BeginWriteAccess(AccessTag);
-	// Reset previous transforms to avoid locked motion vectors
-	// TODO: we should be able to just turn off & delete the prev transforms instead.
-	check(ProxyData.InstanceToPrimitiveRelative.Num() == ProxyData.PrevInstanceToPrimitiveRelative.Num()); // Sanity check, should always have matching associated arrays
-	for (int32 InstanceIndex = 0; InstanceIndex < ProxyData.InstanceToPrimitiveRelative.Num(); ++InstanceIndex)
-	{
-		ProxyData.PrevInstanceToPrimitiveRelative[InstanceIndex] = ProxyData.InstanceToPrimitiveRelative[InstanceIndex];
-	}
-	InstanceSceneDataBuffersImpl.EndWriteAccess(AccessTag);
-}
-
 void FNaniteGeometryCollectionSceneProxy::FlushGPUSceneUpdate_GameThread()
 {
 	ENQUEUE_RENDER_COMMAND(NaniteProxyUpdateGPUScene)(
diff --git a/Engine/Source/Runtime/Experimental/GeometryCollectionEngine/Private/GeometryCollection/GeometryCollectionSceneProxy.h b/Engine/Source/Runtime/Experimental/GeometryCollectionEngine/Private/GeometryCollection/GeometryCollectionSceneProxy.h
--- a/Engine/Source/Runtime/Experimental/GeometryCollectionEngine/Private/GeometryCollection/GeometryCollectionSceneProxy.h
+++ b/Engine/Source/Runtime/Experimental/GeometryCollectionEngine/Private/GeometryCollection/GeometryCollectionSceneProxy.h
@@ -220,6 +220,8 @@ protected:
 
 	TArray<UMaterialInterface*> Materials;
 
+	TArray<int32> TransformToInstanceMapping;
+
 	int32 NumTransforms = 0;
 
 	FGeometryCollectionDynamicData* DynamicData = nullptr;
@@ -360,8 +362,6 @@ public:
 	/** Called on render thread to setup dynamic geometry for rendering */
 	void SetDynamicData_RenderThread(FRHICommandListBase& RHICmdList, FGeometryCollectionDynamicData* NewDynamicData, const FMatrix &PrimitiveLocalToWorld);
 
-	void ResetPreviousTransforms_RenderThread();
-
 	void FlushGPUSceneUpdate_GameThread();
 
 	FORCEINLINE void SetRequiresGPUSceneUpdate_RenderThread(bool bRequireUpdate)
