Attempting to read data not in the IAS Cache will no longer count as a successful read in our statistics/analytics.

#rnx

- The call to FOnDemandIoBackendStats::OnCacheGet was being called when we read data from the memory cache and when the data was not found. This means that every request would end up being recorded as a cache read whether it came from the cache or http.
- Moved the call to the code paths were a successful read from the memory cache was made.

#Synced-CL 49174259

diff --git a/Engine/Source/Runtime/Experimental/IoStore/OnDemand/Private/JournaledCache.cpp b/Engine/Source/Runtime/Experimental/IoStore/OnDemand/Private/JournaledCache.cpp
--- a/Engine/Source/Runtime/Experimental/IoStore/OnDemand/Private/JournaledCache.cpp
+++ b/Engine/Source/Runtime/Experimental/IoStore/OnDemand/Private/JournaledCache.cpp
@@ -9,7 +9,7 @@
 #include "Containers/Map.h"
 #include "Containers/UnrealString.h"
 #include "HAL/FileManager.h"
-#include "HAL/LowLevelMemTracker.h"
+#include "HAL/LowLevelMemTracker.h" 
 #include "HAL/PlatformFile.h"
 #include "HAL/PlatformProcess.h"
 #include "HAL/Runnable.h"
@@ -1462,6 +1462,7 @@ FCache::FGetToken FCache::Get(uint64 Key, FIoBuffer& OutData) const
 			else if (const FIoBuffer* Data = MemCache.Get(Key); Data != nullptr)
 			{
 				OutData = FIoBuffer(FIoBuffer::Clone, Data->Data(), Data->GetSize());
+				FOnDemandIoBackendStats::Get()->OnCacheGet(OutData.GetSize(), ECacheType::Memory);
 			}
 		}
 		else
@@ -1469,6 +1470,7 @@ FCache::FGetToken FCache::Get(uint64 Key, FIoBuffer& OutData) const
 			if (const FIoBuffer* Data = MemCache.Get(Key); Data != nullptr)
 			{
 				OutData = *Data;
+				FOnDemandIoBackendStats::Get()->OnCacheGet(OutData.GetSize(), ECacheType::Memory);
 			}
 
 			if (Partial.Key == Key)
@@ -1478,8 +1480,6 @@ FCache::FGetToken FCache::Get(uint64 Key, FIoBuffer& OutData) const
 		}
 	}
 
-	FOnDemandIoBackendStats::Get()->OnCacheGet(OutData.GetSize(), ECacheType::Memory);
-
 	return 0;
 }
 
