Remove the on-demand host group name Default as the default mount argument and only use
the default name when the host group is missing from both the serialized TOC and the mount arguments.

#rnx

#Synced-CL 49075355

diff --git a/Engine/Source/Runtime/Experimental/IoStore/OnDemand/Private/OnDemandIoStore.cpp b/Engine/Source/Runtime/Experimental/IoStore/OnDemand/Private/OnDemandIoStore.cpp
--- a/Engine/Source/Runtime/Experimental/IoStore/OnDemand/Private/OnDemandIoStore.cpp
+++ b/Engine/Source/Runtime/Experimental/IoStore/OnDemand/Private/OnDemandIoStore.cpp
@@ -1746,11 +1746,19 @@ FIoStatus FOnDemandIoStore::SetupHostGroup(const FSharedOnDemandContainer& Conta
 		return OutHostNames;
 	};
 
-	UE_CLOG(Container->HostGroupName != MountArgs.HostGroupName, LogIoStoreOnDemand, Log,
-		TEXT("Overriding container host group '%s' with '%s' from mount option"),
-		*Container->HostGroupName.ToString(), *MountArgs.HostGroupName.ToString());
+	if (Container->HostGroupName.IsNone())
+	{
+		Container->HostGroupName = MountArgs.HostGroupName.IsNone()
+			? FOnDemandHostGroup::DefaultName
+			: MountArgs.HostGroupName;
+	}
+	else if (!MountArgs.HostGroupName.IsNone() && MountArgs.HostGroupName != Container->HostGroupName)
+	{
+		UE_LOG(LogIoStoreOnDemand, Log, TEXT("Overriding container host group '%s' with '%s' from mount option"),
+			*Container->HostGroupName.ToString(), *MountArgs.HostGroupName.ToString());
+		Container->HostGroupName = MountArgs.HostGroupName;
+	}
 
-	Container->HostGroupName = MountArgs.HostGroupName;
 	if (FHttpIoDispatcher::IsHostGroupRegistered(Container->HostGroupName))
 	{
 		return FIoStatus::Ok;
diff --git a/Engine/Source/Runtime/Experimental/IoStore/OnDemandCore/Public/IO/IoStoreOnDemand.h b/Engine/Source/Runtime/Experimental/IoStore/OnDemandCore/Public/IO/IoStoreOnDemand.h
--- a/Engine/Source/Runtime/Experimental/IoStore/OnDemandCore/Public/IO/IoStoreOnDemand.h
+++ b/Engine/Source/Runtime/Experimental/IoStore/OnDemandCore/Public/IO/IoStoreOnDemand.h
@@ -169,6 +169,9 @@ ENUM_CLASS_FLAGS(EOnDemandMountOptions);
  * Example:
  * http(s)://<Host>/<TocRelativePath>/chunks/<1-Byte Hex>/<hash>.iochunk
  * http(s)://<Host>/<TOC.ChunksDirectory>/chunks/<1-Byte Hex>/<hash>.iochunk
+ *
+ * Host group parameters does not need to be specific when mounting on-demand containers.
+ * The containers will be deferred until the specified host group is registered.
  */
 struct FOnDemandMountArgs
 {
@@ -181,7 +184,7 @@ struct FOnDemandMountArgs
 	/** Serialize the TOC from the specified file path. */
 	FString FilePath;
 	/** Name of a new or existing host group. */
-	FName HostGroupName = FOnDemandHostGroup::DefaultName;
+	FName HostGroupName = NAME_None;
 	/** List of URLs from where to download the chunks. */
 	FOnDemandHostGroup HostGroup;
 	/** Mount options. */
