Remove the IAX analytic 'IasHttpDistributedEndpointResolved' as we no longer resolve the distributed endpoint ourselves so we have no useful information here.

#rnx
#Synced-CL 49043027

diff --git a/Engine/Source/Runtime/Experimental/IoStore/OnDemand/Private/Statistics.cpp b/Engine/Source/Runtime/Experimental/IoStore/OnDemand/Private/Statistics.cpp
--- a/Engine/Source/Runtime/Experimental/IoStore/OnDemand/Private/Statistics.cpp
+++ b/Engine/Source/Runtime/Experimental/IoStore/OnDemand/Private/Statistics.cpp
@@ -403,7 +403,6 @@ FCounterAtomicInt		GCacheRejectBytes(TEXT("Ias/CachePutRejectBytes"), TraceCount
 double					GAnalyticsStartTiming = 0.0;
 // http stats
 int64					GAppResumeCount = 0;
-bool					GHttpDistributedEndpointResolved = false;
 FCounterInt				GHttpConnectCount(TEXT("Ias/HttpConnectCount"), TraceCounterDisplayHint_None);
 FCounterInt				GHttpDisconnectCount(TEXT("Ias/HttpDisconnectCount"), TraceCounterDisplayHint_None);
 UE_IAX_COUNTER(HttpGetCount, FCounterInt);
@@ -1011,7 +1010,6 @@ void FOnDemandIoBackendStats::ReportGeneralAnalytics(TArray<FAnalyticsEventAttri
 	AppendAnalyticsEventAttributeArray(OutAnalyticsArray
 		, TEXT("IaxTimeSpent"), TimeSpent
 		, TEXT("IasAppResumeCount"), GAppResumeCount
-		, TEXT("IasHttpDistributedEndpointResolved"), GHttpDistributedEndpointResolved
 		, TEXT("IasHttpHasEverConnected"), GHttpConnectCount.Get() > 0 // Report if the system has ever actually managed to make a connection
 	);
 
@@ -1141,7 +1139,6 @@ private:
 	{
 		// Report if the system has ever actually managed to make a connection
 		AppendAnalyticsEventAttributeArray(OutAnalyticsArray
-			,TEXT("IasHttpDistributedEndpointResolved"), GHttpDistributedEndpointResolved
 			,TEXT("IasHttpHasEverConnected"), GHttpConnectCount.Get() > 0
 		);
 
@@ -1392,11 +1389,6 @@ void FOnDemandIoBackendStats::OnCacheSuspended(double Seconds)
 	GCacheSuspendedSeconds.Set(Seconds);
 }
 
-void FOnDemandIoBackendStats::OnHttpDistributedEndpointResolved()
-{
-	GHttpDistributedEndpointResolved = true;
-}
-
 void FOnDemandIoBackendStats::OnHttpConnected()
 {
 	GHttpConnectCount.Increment();
diff --git a/Engine/Source/Runtime/Experimental/IoStore/OnDemand/Private/Statistics.h b/Engine/Source/Runtime/Experimental/IoStore/OnDemand/Private/Statistics.h
--- a/Engine/Source/Runtime/Experimental/IoStore/OnDemand/Private/Statistics.h
+++ b/Engine/Source/Runtime/Experimental/IoStore/OnDemand/Private/Statistics.h
@@ -97,9 +97,6 @@ public:
 	void OnCacheSetMaxBytes(uint64 TotalSize) IAS_STATISTICS_IMPL()
 	void OnCacheSuspended(double Seconds) IAS_STATISTICS_IMPL()
 
-	// Per session
-	void OnHttpDistributedEndpointResolved() IAS_STATISTICS_IMPL()
-
 	// Per hostgroup
 	void OnHttpConnected() IAS_STATISTICS_IMPL()
 	void OnHttpDisconnected() IAS_STATISTICS_IMPL()
