FMultiEndpointHttpClient can now handle requests for empty/broken hostgroups.

#rnx

- FMultiEndpointHttpClient::Get will no longer crash when called with an empty or broken FOnDemandHostGroup.
- The request will be stored until the next call to ::Tick at which point it's callback will be invoked as having failed.
-- This means that the caller only needs to handle a single error path (on callback) and don't have to worry about ::Get returning an error in someway.
- In this scenario FMultiEndpointHttpClient::Get will return an invalid FHttpTicketId (0) which will be ignored if used to try and cancel a request.
- We will log the error before the callback is invoked with a very similar format to the existing log message.
- With our current code base it should not be possible for FMultiEndpointHttpClient::Get to be called with a broken host group as we check for that before issuing requests. This hardening is to prevent problems with future use cases.

#Synced-CL 49043027

diff --git a/Engine/Source/Runtime/Experimental/IoStore/OnDemand/Private/OnDemandHttpClient.cpp b/Engine/Source/Runtime/Experimental/IoStore/OnDemand/Private/OnDemandHttpClient.cpp
--- a/Engine/Source/Runtime/Experimental/IoStore/OnDemand/Private/OnDemandHttpClient.cpp
+++ b/Engine/Source/Runtime/Experimental/IoStore/OnDemand/Private/OnDemandHttpClient.cpp
@@ -294,6 +294,7 @@ bool FMultiEndpointHttpClient::Tick(int32 WaitTimeMs, uint32 MaxKiBPerSecond)
 
 	const uint32 TicketCount = EventLoop.Tick(WaitTimeMs);
 
+	ProcessFailedRequests();
 	ProcessRetryAttempts(TicketCount);
 
 	const bool bIsIdle = EventLoop.IsIdle();
@@ -390,6 +391,15 @@ FMultiEndpointHttpClient::FHttpTicketId FMultiEndpointHttpClient::IssueRequest(F
 {
 	using namespace UE::IoStore::HTTP;
 
+	if (Request.Connection.CurrentHost == INDEX_NONE)
+	{
+		// Store the request response so we can invoke it in the next tick.
+		// This way the calling code does not need to handle multiple paths.
+		FailedRequests.Emplace(MoveTemp(Request));
+
+		return 0;
+	}
+
 	check(Request.Connection.HostGroup.IsEmpty() == false);
 	check(Request.Connection.Pools[Request.Connection.CurrentHost].IsValid());
 
@@ -556,7 +566,8 @@ void FMultiEndpointHttpClient::CompleteRequest(FRequest&& Request, const UE::IoS
 
 #if DO_CHECK
 	check(Request.SinkCounter == Request.RetryCount);
-#endif
+#endif // DO_CHECK
+
 	Log(Request, TicketStatus);
 
 	FMultiEndpointHttpClientResponse Response
@@ -657,9 +668,37 @@ void FMultiEndpointHttpClient::RetryRequest(FRequest&& Request)
 	Retries.Emplace(MoveTemp(Request));
 }
 
+void FMultiEndpointHttpClient::ProcessFailedRequests()
+{
+	if (FailedRequests.IsEmpty())
+	{
+		return;
+	}
+
+	TRACE_CPUPROFILER_EVENT_SCOPE(FMultiEndpointHttpClient::ProcessFailedRequests);
+
+	for (FRequest& Request :FailedRequests)
+	{
+#if DO_CHECK
+		check(Request.SinkCounter == 0);
+#endif // DO_CHECK
+
+		FMultiEndpointHttpClientResponse Response
+		{
+			.Reason = TEXT("Hostgroup had no valid urls for the request")
+		};
+
+		LogError(Request, Response.Reason);
+
+		Request.OnResponse(MoveTemp(Response));
+	}
+
+	FailedRequests.Empty();
+}
+
 void FMultiEndpointHttpClient::ProcessRetryAttempts(uint32 TicketCount)
 {
-	if(Retries.IsEmpty())
+	if (Retries.IsEmpty())
 	{
 		return;
 	}
@@ -771,7 +810,11 @@ FMultiEndpointHttpClient::FConnection& FMultiEndpointHttpClient::GetConnection(c
 	Conn.Pools.SetNum(HostGroup.Hosts().Num());
 	Conn.CurrentHost = HostGroup.PrimaryHostIndex();
 
-	Conn.Pools[Conn.CurrentHost] = CreateConnection(HostGroup.PrimaryHost());
+	if (Conn.CurrentHost != INDEX_NONE)
+	{
+		Conn.Pools[Conn.CurrentHost] = CreateConnection(HostGroup.PrimaryHost());
+	}
+
 	return Conn;
 }
 
@@ -844,4 +887,10 @@ void FMultiEndpointHttpClient::Log(const FRequest& Request, const HTTP::FTicketS
 #endif //!NO_LOGGING
 }
 
+void FMultiEndpointHttpClient::LogError(const FRequest& Request, const FString& ErrorReason) const
+{
+	// Attempt to keep roughly the same format as ::Log where is makes sense.
+	UE_LOGF(LogHttpIoDispatcher, Error, "http-  0:     0ms      0KiB [ ???] %hs (Attempt 0/0) %ls", *Request.RelativeUrl, *ErrorReason);
+}
+
 } // namespace UE::IoStore
diff --git a/Engine/Source/Runtime/Experimental/IoStore/OnDemand/Private/OnDemandHttpClient.h b/Engine/Source/Runtime/Experimental/IoStore/OnDemand/Private/OnDemandHttpClient.h
--- a/Engine/Source/Runtime/Experimental/IoStore/OnDemand/Private/OnDemandHttpClient.h
+++ b/Engine/Source/Runtime/Experimental/IoStore/OnDemand/Private/OnDemandHttpClient.h
@@ -192,7 +192,7 @@ private:
 
 #if DO_CHECK
 		/** Used to make sure that the sink is only called with a completed state one per request attempt */
-		int32				SinkCounter = 0;
+		int32						SinkCounter = 0;
 #endif // DO_CHECK
 	};
 
@@ -202,6 +202,7 @@ private:
 	void								CompleteRequest(FRequest&& Request, const UE::IoStore::HTTP::FTicketStatus& TicketStatus);
 	void								RetryRequest(FRequest&& Request);
 
+	void								ProcessFailedRequests();
 	void								ProcessRetryAttempts(uint32 TicketCount);
 
 	TUniquePtr<HTTP::FConnectionPool>	CreateConnection(FAnsiStringView HostUrl) const;
@@ -211,11 +212,15 @@ private:
 	uint32								GetRetryLimitForRequest(const FRequest& Request) const;
 
 	void								Log(const FRequest& Request, const HTTP::FTicketStatus& TicketStatus) const;
+	void								LogError(const FRequest& Request, const FString& ErrorReason) const;
 
 	FMultiEndpointHttpClientConfig		Config;
 	TArray<TUniquePtr<FConnection>>		Connections;
 	HTTP::FEventLoop					EventLoop;
 	TArray<FRequest>					Retries;
+	/** Stores the response for requests that failed to be issued, to be flushed in the next Tick */
+	TArray<FRequest>					FailedRequests;
+
 	struct FTicketInfo
 	{
 		HTTP::FTicket HttpTicket;
