[Chaos] Run on the same worker FlushAsyncAcceleration in single worker physics
Push Data should be single threaded is cvar p.Chaos.SingleWorkerPhysics is enabled

#Synced-CL 49195397

#Synced-CL 49108854

diff --git a/Engine/Source/Runtime/Experimental/Chaos/Private/Chaos/PBDRigidsEvolution.cpp b/Engine/Source/Runtime/Experimental/Chaos/Private/Chaos/PBDRigidsEvolution.cpp
--- a/Engine/Source/Runtime/Experimental/Chaos/Private/Chaos/PBDRigidsEvolution.cpp
+++ b/Engine/Source/Runtime/Experimental/Chaos/Private/Chaos/PBDRigidsEvolution.cpp
@@ -698,7 +698,14 @@ namespace Chaos
 			AsyncInternalAcceleration->SetSyncTimestamp(LatestExternalTimestampConsumed_Internal);
 			AsyncExternalAcceleration->SetSyncTimestamp(LatestExternalTimestampConsumed_Internal);
 		};
-		TaskDispatcher.FlushAccelerationQueue(FlushAsyncAccelerationQueueLambda);
+		if (bSingleWorkerPhysics)
+		{
+			FlushAsyncAccelerationQueueLambda();
+		}
+		else
+		{
+			TaskDispatcher.FlushAccelerationQueue(FlushAsyncAccelerationQueueLambda);
+		}
 	}
 
 	//TODO: make static and _External suffix
diff --git a/Engine/Source/Runtime/Experimental/Chaos/Private/PBDRigidsSolver.cpp b/Engine/Source/Runtime/Experimental/Chaos/Private/PBDRigidsSolver.cpp
--- a/Engine/Source/Runtime/Experimental/Chaos/Private/PBDRigidsSolver.cpp
+++ b/Engine/Source/Runtime/Experimental/Chaos/Private/PBDRigidsSolver.cpp
@@ -2662,7 +2662,7 @@ namespace Chaos
 	{
 		FPullPhysicsData* PullData = MarshallingManager.GetCurrentPullData_Internal();
 
-		if (!(GSingleThreadedPhysics || CVars::bSingleThreadPullData))
+		if (!(bSingleWorkerPhysics || CVars::bSingleThreadPullData))
 		{
 			using TSOA = typename TParticleView<FPBDRigidParticles>::TSOA;
 			using THandle = typename TSOA::THandleType;
@@ -2779,12 +2779,11 @@ namespace Chaos
 			RewindData->BufferPhysicsResults(PullData->DirtyRigidErrors);
 		}
 
-		if (!(GSingleThreadedPhysics || CVars::bSingleThreadPullData))
+		if (!(bSingleWorkerPhysics || CVars::bSingleThreadPullData))
 		{
 			TaskDispatcher->WaitTaskEndBufferResults();
 		}
 
-
 		// Now that results have been buffered we have completed a solve step so we can broadcast that event
 		EventPostSolve.Broadcast(MLastDt);
 
