Dataflow node : add evelaute node method that evaluate the node all active outputs

#Synced-CL 49237557

diff --git a/Engine/Source/Runtime/Experimental/Dataflow/Core/Private/Dataflow/DataflowNode.cpp b/Engine/Source/Runtime/Experimental/Dataflow/Core/Private/Dataflow/DataflowNode.cpp
--- a/Engine/Source/Runtime/Experimental/Dataflow/Core/Private/Dataflow/DataflowNode.cpp
+++ b/Engine/Source/Runtime/Experimental/Dataflow/Core/Private/Dataflow/DataflowNode.cpp
@@ -713,6 +713,45 @@ void FDataflowNode::Invalidate(const UE::Dataflow::FTimestamp& InModifiedTimesta
 	}
 }
 
+void FDataflowNode::EvaluateNode(UE::Dataflow::FContext& Context) const
+{
+	// TODO(dataflow) : in the future we should also have the possibility to evaluate the watched output if needed by the context 
+	if (IsActive())
+	{
+		if (ExpandedOutputs.IsEmpty())
+		{
+			Evaluate(Context, nullptr);
+		}
+		else if (!IsFrozen())
+		{
+			for (const TPair<UE::Dataflow::FConnectionKey, FDataflowOutput*>& OutputByKey : ExpandedOutputs)
+			{
+				if (const FDataflowOutput* const Output = OutputByKey.Value)
+				{
+					// todo(ccaillaud) : shoudl we have an option to skip the non connected outputs ? 
+					if (Context.HasData(Output->CacheKey(), GetTimestamp()))
+					{
+						Evaluate(Context, Output);
+					}
+				}
+			}
+		}
+	}
+	else
+	{
+		for (const TPair<UE::Dataflow::FConnectionKey, FDataflowOutput*>& OutputByKey : ExpandedOutputs)
+		{
+			if (const FDataflowOutput* const Output = OutputByKey.Value)
+			{
+				if (const FDataflowInput* const PassthroughInput = Output->GetPassthroughInput())
+				{
+					Output->ForwardInput(PassthroughInput, Context);
+				}
+			}
+		}
+	}
+}
+
 const FProperty* FDataflowNode::FindProperty(const UStruct* Struct, const void* InProperty, const FName& PropertyName, TArray<const FProperty*>* OutPropertyChain) const
 {
 	return UE::Dataflow::Private::FindProperty(Struct, this, InProperty, PropertyName, OutPropertyChain);
diff --git a/Engine/Source/Runtime/Experimental/Dataflow/Core/Public/Dataflow/DataflowNode.h b/Engine/Source/Runtime/Experimental/Dataflow/Core/Public/Dataflow/DataflowNode.h
--- a/Engine/Source/Runtime/Experimental/Dataflow/Core/Public/Dataflow/DataflowNode.h
+++ b/Engine/Source/Runtime/Experimental/Dataflow/Core/Public/Dataflow/DataflowNode.h
@@ -355,6 +355,14 @@ PRAGMA_ENABLE_DEPRECATION_WARNINGS
 	//
 	// Evaluation
 	//
+
+	/** 
+	* This evaluate the node based on the existing state of the outputs
+	* if the node has no output then evaluate will be called with a null output 
+	* otherwise Evaluate will be called for each connected output if it has an outdated cached value in teh context
+	*/
+	void EvaluateNode(UE::Dataflow::FContext& Context) const;
+
 	virtual void Evaluate(UE::Dataflow::FContext& Context, const FDataflowOutput*) const { ensure(false); }
 
 	virtual bool EvaluateOnGameThreadOnly() const { return false; }
