Update the hostgroup status to be called on the IAX http thread as it is not thread safe.

#rnx

- After the refactor to the new backend the hostgroup was being updated with the http response in FHttpDispatcher::CompleteHttpRequest which is called on the task graph. This is unsafe because the hostgroup connection status is not thread safe and previously assumed that it would be called on the IAX http thread.
- Even if we made the hostgroup thread safe we would still want to update it's status on the IAX http thread so if we detect a bad connection we would react to it immediately and not at some future point where the code is executed on the task graph system.
- Moved the actual updating code to  it's own method and call it immediately when the http response is issued (which will be during the http client tick) so that it is thread safe.
- The remaining statistics processing in CompleteHttpRequest can remain on the taskgraph as that code is thread safe and updating it immediately does not matter.

#Synced-CL 49174259

diff --git a/Engine/Source/Runtime/Experimental/IoStore/OnDemand/Private/OnDemandHttpIoDispacher.cpp b/Engine/Source/Runtime/Experimental/IoStore/OnDemand/Private/OnDemandHttpIoDispacher.cpp
--- a/Engine/Source/Runtime/Experimental/IoStore/OnDemand/Private/OnDemandHttpIoDispacher.cpp
+++ b/Engine/Source/Runtime/Experimental/IoStore/OnDemand/Private/OnDemandHttpIoDispacher.cpp
@@ -1038,7 +1038,7 @@ private:
 	bool							ProcessHttpRequests();
 	void							CancelRemainingHttpRequests();
 	FHttpRequest*					DequeueHttpRequests();
-	void							CompleteHttpRequest(FHttpRequest& Request, FMultiEndpointHttpClientResponse&& HttpResponse, FIASHostGroup& HostGroup, const FIoHttpRange& TotalRange);
+	void							CompleteHttpRequest(FHttpRequest& Request, FMultiEndpointHttpClientResponse&& HttpResponse, const FIoHttpRange& TotalRange);
 	void							CompleteRequest(FHttpRequest& Request, TArray<FAnsiString>&& Headers, FIoBuffer Body, const FIoHttpRange& TotalRange, uint32 StatusCode, bool bCached);
 	void							CompleteRequest(FHttpRequest& Request, FIoBuffer& Body, const FIoHttpRange& TotalRange, uint32 StatusCode, bool bCached)
 									{
@@ -1048,6 +1048,8 @@ private:
 	void							ApplyCancellationToken(FHttpRequest* RequestChain, FMultiEndpointHttpClient::FHttpTicketId TicketId);
 	void							CompleteCancellationToken(FHttpRequest* RequestChain);
 
+	void							UpdateHostGroup(FIASHostGroup& Hostgroup, const FMultiEndpointHttpClientResponse& Response) const;
+
 #if UE_ALLOW_HTTP_PAUSE
 	void OnPauseCommand(bool bShouldPause, const TArray<FString>& Args, FOutputDevice& Ar);
 #endif //UE_ALLOW_HTTP_PAUSE
@@ -1670,12 +1672,13 @@ bool FHttpDispatcher::ProcessHttpRequests()
 					NumConcurrentRequests--;
 
 					CompleteCancellationToken(HttpRequest);
+					UpdateHostGroup(HostGroup, HttpResponse);
 
 					UE::Tasks::Launch(
 						UE_SOURCE_LOCATION,
-						[this, HttpRequest, HttpResponse = MoveTemp(HttpResponse), TotalRange, HostGroup = MoveTemp(HostGroup)]() mutable
+						[this, HttpRequest, HttpResponse = MoveTemp(HttpResponse), TotalRange]() mutable
 						{
-							CompleteHttpRequest(*HttpRequest, MoveTemp(HttpResponse), HostGroup, TotalRange);
+							CompleteHttpRequest(*HttpRequest, MoveTemp(HttpResponse), TotalRange);
 						}, UE::Tasks::ETaskPriority::BackgroundLow);
 				});
 
@@ -1745,7 +1748,7 @@ FHttpRequest* FHttpDispatcher::DequeueHttpRequests()
 	return nullptr;
 }
 
-void FHttpDispatcher::CompleteHttpRequest(FHttpRequest& Request, FMultiEndpointHttpClientResponse&& HttpResponse, FIASHostGroup& HostGroup, const FIoHttpRange& TotalRange)
+void FHttpDispatcher::CompleteHttpRequest(FHttpRequest& Request, FMultiEndpointHttpClientResponse&& HttpResponse, const FIoHttpRange& TotalRange)
 {
 	if (HttpResponse.RetryCount > 0)
 	{
@@ -1754,7 +1757,6 @@ void FHttpDispatcher::CompleteHttpRequest(FHttpRequest& Request, FMultiEndpointH
 
 	if (HttpResponse.IsOk())
 	{
-		HostGroup.OnSuccessfulResponse();
 		FOnDemandIoBackendStats::Get()->OnHttpCdnCacheReply(Request.GetHttpRequestType(), HttpResponse.CDNCacheStatus);
 		FOnDemandIoBackendStats::Get()->OnHttpGet(Request.GetHttpRequestType(), HttpResponse.Body.DataSize(), HttpResponse.Sample.GetTotalMs());
 	}
@@ -1765,11 +1767,6 @@ void FHttpDispatcher::CompleteHttpRequest(FHttpRequest& Request, FMultiEndpointH
 	else
 	{
 		FOnDemandIoBackendStats::Get()->OnHttpError(Request.GetHttpRequestType());
-		if (HostGroup.OnFailedResponse())
-		{
-			// A disconnect was triggered
-			FOnDemandIoBackendStats::Get()->OnHttpDisconnected();
-		}
 	}
 
 	CompleteRequest(Request, MoveTemp(HttpResponse.Headers), HttpResponse.Body, TotalRange, HttpResponse.StatusCode, false);
@@ -1921,6 +1918,21 @@ void FHttpDispatcher::CompleteCancellationToken(FHttpRequest* RequestChain)
 	}
 }
 
+void FHttpDispatcher::UpdateHostGroup(FIASHostGroup& Hostgroup, const FMultiEndpointHttpClientResponse& Response) const
+{
+	if (Response.IsOk())
+	{
+		Hostgroup.OnSuccessfulResponse();
+	}
+	else if (!Response.IsCanceled())
+	{
+		if (Hostgroup.OnFailedResponse())
+		{
+			FOnDemandIoBackendStats::Get()->OnHttpDisconnected(); // A disconnect was triggered
+		}
+	}
+}
+
 } // namespace UE::IoStore::HttpIoDispatcher
 
 ///////////////////////////////////////////////////////////////////////////////
