Minor reworking of FBitWindow so that it is always initialized to a proper value.

#rnx

- Originally FBitWindow was written to be more generic but we have no reason to believe that it will be used outside of it's current use case and may well be deleted in reworks next year.
- At the moment it is intiialized to a non working state and relies on ::Reset being called with a valid number of samples.
- This rework just sets the samples during the constructor (from the same cvar) so that it is valid all of the time.

#Synced-CL 49198420

diff --git a/Engine/Source/Runtime/Experimental/IoStore/OnDemand/Private/IasHostGroup.cpp b/Engine/Source/Runtime/Experimental/IoStore/OnDemand/Private/IasHostGroup.cpp
--- a/Engine/Source/Runtime/Experimental/IoStore/OnDemand/Private/IasHostGroup.cpp
+++ b/Engine/Source/Runtime/Experimental/IoStore/OnDemand/Private/IasHostGroup.cpp
@@ -30,24 +30,16 @@ static FAutoConsoleVariableRef CVar_IasHttpErrorHighWater(
 
 struct FBitWindow
 {
-	void Reset()
+	FBitWindow()
 	{
-		Reset(Bits.Num());
+		Init(GIasHttpErrorSampleCount);
 	}
 
-	void Reset(uint32 Count)
+	~FBitWindow() = default;
+
+	void Reset()
 	{
-		if (Count > 0)
-		{
-			Count = FMath::RoundUpToPowerOfTwo(Count);
-			Bits.Init(false, Count);
-			Counter = 0;
-			Mask = Count - 1 ;
-		}
-		else
-		{
-			*this = FBitWindow();
-		}
+		Init(Bits.Num());
 	}
 
 	void Add(bool bValue)
@@ -70,6 +62,18 @@ struct FBitWindow
 	}
 
 private:
+
+	void Init(uint32 Count)
+	{
+		if (Count > 0)
+		{
+			Count = FMath::RoundUpToPowerOfTwo(Count);
+			Bits.Init(false, Count);
+			Counter = 0;
+			Mask = Count - 1;
+		}
+	}
+
 	TBitArray<> Bits;
 	uint32 Counter = 0;
 	uint32 Mask = 0;
@@ -93,7 +97,7 @@ struct FIASHostGroup::FImpl
 	void Reset(FOnDemandHostGroup&& InHostGroup)
 	{
 		HostGroup = MoveTemp(InHostGroup);
-		HttpErrorHistory.Reset(GIasHttpErrorSampleCount);
+		HttpErrorHistory.Reset();
 		bHttpEnabled = true;
 	}
 
@@ -114,7 +118,7 @@ TIoStatusOr<FIASHostGroup> FIASHostGroup::Create(FName Name, TConstArrayView<FAn
 	}
 
 	FIASHostGroup HostGroup(Name, Result.ConsumeValueOrDie());
-	HostGroup.Impl->HttpErrorHistory.Reset(GIasHttpErrorSampleCount);
+	HostGroup.Impl->HttpErrorHistory.Reset();
 
 	return HostGroup;
 }
@@ -129,7 +133,7 @@ TIoStatusOr<FIASHostGroup> FIASHostGroup::Create(FName Name, TConstArrayView<FSt
 	}
 
 	FIASHostGroup HostGroup(Name, Result.ConsumeValueOrDie());
-	HostGroup.Impl->HttpErrorHistory.Reset(GIasHttpErrorSampleCount);
+	HostGroup.Impl->HttpErrorHistory.Reset();
 
 	return HostGroup;
 }
