[TraceBasedDebuggers] adding missing defines for targets with WITH_TRACE_BASED_DEBUGGERS = 0

#Synced-CL 49268467

diff --git a/Engine/Plugins/ChaosVD/Source/ChaosVD/Private/Widgets/ChaosVDRecordingControls.cpp b/Engine/Plugins/ChaosVD/Source/ChaosVD/Private/Widgets/ChaosVDRecordingControls.cpp
--- a/Engine/Plugins/ChaosVD/Source/ChaosVD/Private/Widgets/ChaosVDRecordingControls.cpp
+++ b/Engine/Plugins/ChaosVD/Source/ChaosVD/Private/Widgets/ChaosVDRecordingControls.cpp
@@ -21,7 +21,9 @@ FChaosVDRecordingControls::FChaosVDRecordingControls(const TSharedRef<SChaosVDMa
 	, const TSharedRef<UE::TraceBasedDebuggers::FRemoteSessionsManager>& InRemoteSessionManager)
 	: FRecordingControls(InMainTab->GetStatusBarName(), InRemoteSessionManager, Chaos::VD::DebuggerGuid, LogChaosVDEditor)
 {
+#if WITH_TRACE_BASED_DEBUGGERS
 	TraceRelayTransportInstanceWeakPtr = FChaosVDEngineEditorBridge::Get().GetTraceRelayTransportInstance();
+#endif
 
 	MainTabWeakPtr = InMainTab;
 	CurrentLoadingMode = EChaosVDLoadRecordedDataMode::SingleSource;
diff --git a/Engine/Source/Runtime/Experimental/TraceBasedDebuggersAnalysis/Private/Widgets/RecordingControls.cpp b/Engine/Source/Runtime/Experimental/TraceBasedDebuggersAnalysis/Private/Widgets/RecordingControls.cpp
--- a/Engine/Source/Runtime/Experimental/TraceBasedDebuggersAnalysis/Private/Widgets/RecordingControls.cpp
+++ b/Engine/Source/Runtime/Experimental/TraceBasedDebuggersAnalysis/Private/Widgets/RecordingControls.cpp
@@ -1,4 +1,4 @@
-// Copyright Epic Games, Inc. All Rights Reserved.
+﻿// Copyright Epic Games, Inc. All Rights Reserved.
 
 #if WITH_EDITOR
 
@@ -548,6 +548,7 @@ void FRecordingControls::ToggleSingleSessionRecordingState(const TSharedRef<FSes
 
 void FRecordingControls::SetSessionRecordingState(const bool bIsRecording, const TSharedRef<FSessionInfo>& SessionInfoRef)
 {
+#if WITH_TRACE_BASED_DEBUGGERS
 	using namespace UE::TraceBasedDebuggers;
 	const TSharedPtr<FRemoteSessionsManager> RemoteSessionManager = RemoteSessionManagerWeakPtr.Pin();
 	if (!RemoteSessionManager)
@@ -696,6 +697,7 @@ void FRecordingControls::SetSessionRecordingState(const bool bIsRecording, const
 		CancelInFlightConnectionAttempt(SessionInfoRef.ToSharedPtr());
 		SendStopRecordingCommandInternal(SessionInfoRef->Address);
 	}
+#endif // WITH_TRACE_BASED_DEBUGGERS
 }
 
 FReply FRecordingControls::ToggleRecordingState()
diff --git a/Engine/Source/Runtime/Experimental/TraceBasedDebuggersAnalysis/Public/Widgets/RecordingControls.h b/Engine/Source/Runtime/Experimental/TraceBasedDebuggersAnalysis/Public/Widgets/RecordingControls.h
--- a/Engine/Source/Runtime/Experimental/TraceBasedDebuggersAnalysis/Public/Widgets/RecordingControls.h
+++ b/Engine/Source/Runtime/Experimental/TraceBasedDebuggersAnalysis/Public/Widgets/RecordingControls.h
@@ -105,7 +105,9 @@ protected:
 	UE_API FSlateIcon GetIconForSession(FGuid SessionId) const;
 
 	const TWeakPtr<FRemoteSessionsManager> RemoteSessionManagerWeakPtr;
+#if WITH_TRACE_BASED_DEBUGGERS
 	TWeakPtr<IDataRelayTransport> TraceRelayTransportInstanceWeakPtr;
+#endif
 	const FSlateBrush* StartRecordingBrush = nullptr;
 	const FSlateBrush* StopRecordingBrush = nullptr;
 	FSlateIcon RecordingSessionIcon;
