[TraceBasedDebuggers] keep track of the session discovery activations to handle multiple debuggers

#Synced-CL 49280888

diff --git a/Engine/Plugins/ChaosVD/Source/ChaosVD/Private/Widgets/SChaosVDMainTab.cpp b/Engine/Plugins/ChaosVD/Source/ChaosVD/Private/Widgets/SChaosVDMainTab.cpp
--- a/Engine/Plugins/ChaosVD/Source/ChaosVD/Private/Widgets/SChaosVDMainTab.cpp
+++ b/Engine/Plugins/ChaosVD/Source/ChaosVD/Private/Widgets/SChaosVDMainTab.cpp
@@ -220,7 +220,7 @@ void SChaosVDMainTab::Construct(const FArguments& InArgs, TSharedPtr<FChaosVDEng
 
 	if (TSharedPtr<UE::TraceBasedDebuggers::FRemoteSessionsManager> RemoteSessionManager = FChaosVDEngineEditorBridge::Get().GetSessionsManager())
 	{
-		RemoteSessionManager->StartSessionDiscovery();
+		RemoteSessionManager->StartSessionDiscovery(Chaos::VD::DebuggerGuid);
 	}
 }
 
@@ -228,7 +228,7 @@ SChaosVDMainTab::~SChaosVDMainTab()
 {
 	if (TSharedPtr<UE::TraceBasedDebuggers::FRemoteSessionsManager> RemoteSessionManager = FChaosVDEngineEditorBridge::Get().GetSessionsManager())
 	{
-		RemoteSessionManager->StopSessionDiscovery();
+		RemoteSessionManager->StopSessionDiscovery(Chaos::VD::DebuggerGuid);
 	}
 
 	FChaosVDExtensionsManager::Get().OnExtensionRegistered().RemoveAll(this);
diff --git a/Engine/Source/Runtime/Experimental/TraceBasedDebuggers/Private/RemoteSessionsManager.cpp b/Engine/Source/Runtime/Experimental/TraceBasedDebuggers/Private/RemoteSessionsManager.cpp
--- a/Engine/Source/Runtime/Experimental/TraceBasedDebuggers/Private/RemoteSessionsManager.cpp
+++ b/Engine/Source/Runtime/Experimental/TraceBasedDebuggers/Private/RemoteSessionsManager.cpp
@@ -317,25 +317,37 @@ void FRemoteSessionsManager::EnumerateMessageTypes(const FVisitorFunction& InVis
 	}
 }
 
-void FRemoteSessionsManager::StartSessionDiscovery()
+void FRemoteSessionsManager::StartSessionDiscovery(const FGuid& InDebuggerGuid)
 {
-	if (TickHandle.IsValid())
+	if (ActiveDebuggersUsingSessionDiscovery.IsEmpty())
 	{
-		UE_LOG(LogTraceBasedDebuggers, Warning, TEXT("[%hs] Session discovery already started"), __func__);
-		return;
+		ensureMsgf(!TickHandle.IsValid(), TEXT("Tick handle should never be valid when a debugger starts sessions discovery"));
+		constexpr float TickInterval = 1.0f;
+		TickHandle = FTSTicker::GetCoreTicker().AddTicker(FTickerDelegate::CreateRaw(this, &FRemoteSessionsManager::Tick), TickInterval);
 	}
 
-	constexpr float TickInterval = 1.0f;
-	TickHandle = FTSTicker::GetCoreTicker().AddTicker(FTickerDelegate::CreateRaw(this, &FRemoteSessionsManager::Tick), TickInterval);
+	const int32 PreviousCount = ActiveDebuggersUsingSessionDiscovery.Num();
+	ActiveDebuggersUsingSessionDiscovery.AddUnique(InDebuggerGuid);
+	if (PreviousCount == ActiveDebuggersUsingSessionDiscovery.Num())
+	{
+		UE_LOG(LogTraceBasedDebuggers, Warning, TEXT("[%hs] Session discovery already started by debugger %s"), __func__, *InDebuggerGuid.ToString());
+	}
 }
 
-void FRemoteSessionsManager::StopSessionDiscovery()
+void FRemoteSessionsManager::StopSessionDiscovery(const FGuid& InDebuggerGuid)
 {
-	if (TickHandle.IsValid())
+	const int32 NumRemoved = ActiveDebuggersUsingSessionDiscovery.Remove(InDebuggerGuid);
+
+	UE_CLOG(NumRemoved == 0, LogTraceBasedDebuggers, Warning, TEXT("[%hs] Session discovery was not started by debugger %s"), __func__, *InDebuggerGuid.ToString());
+
+	if (ActiveDebuggersUsingSessionDiscovery.IsEmpty())
 	{
-		RemoveExpiredSessions(ERemoveSessionOptions::ForceRemoveAll);
-		FTSTicker::RemoveTicker(TickHandle);
-		TickHandle.Reset();
+		if (ensureMsgf(TickHandle.IsValid(), TEXT("Tick handle should be valid when removing the last debugger requesting sessions discovery")))
+		{
+			RemoveExpiredSessions(ERemoveSessionOptions::ForceRemoveAll);
+			FTSTicker::RemoveTicker(TickHandle);
+			TickHandle.Reset();
+		}
 	}
 }
 
diff --git a/Engine/Source/Runtime/Experimental/TraceBasedDebuggers/Public/RemoteSessionsManager.h b/Engine/Source/Runtime/Experimental/TraceBasedDebuggers/Public/RemoteSessionsManager.h
--- a/Engine/Source/Runtime/Experimental/TraceBasedDebuggers/Public/RemoteSessionsManager.h
+++ b/Engine/Source/Runtime/Experimental/TraceBasedDebuggers/Public/RemoteSessionsManager.h
@@ -174,12 +174,22 @@ struct FRemoteSessionsManager : TSharedFromThis<FRemoteSessionsManager>
 	/**
 	 * Starts sending ping requests that potential targets (client, servers or other editors) can answer to and report themselves as available
 	 */
-	UE_API void StartSessionDiscovery();
+	UE_API void StartSessionDiscovery(const FGuid& InDebuggerGuid);
+	UE_DEPRECATED(5.8, "Use overload taking a debugger guid")
+	void StartSessionDiscovery()
+	{
+		StartSessionDiscovery(FGuid{});
+	}
 
 	/**
 	 * Stop sending ping requests to potential targets (client, servers or other editors).
 	 */
-	UE_API void StopSessionDiscovery();
+	UE_API void StopSessionDiscovery(const FGuid& InDebuggerGuid);
+	UE_DEPRECATED(5.8, "Use overload taking a debugger guid")
+	void StopSessionDiscovery()
+	{
+		StopSessionDiscovery(FGuid{});
+	}
 
 	/**
 	 * Issues a command to the provided address to change the state of a data channel
@@ -421,6 +431,7 @@ private:
 
 	FMessagingInitialized MessagingInitializedDelegate;
 
+	TArray<FGuid> ActiveDebuggersUsingSessionDiscovery;
 	TArray<const UScriptStruct*> SupportedMessageTypes;
 
 	bool bInitialized = false;
