Print error details when FOnDemandInstallCache::Shutdown fails to flush chunks.

#Synced-CL 49267577

diff --git a/Engine/Source/Runtime/Experimental/IoStore/OnDemand/Private/OnDemandInstallCache.cpp b/Engine/Source/Runtime/Experimental/IoStore/OnDemand/Private/OnDemandInstallCache.cpp
--- a/Engine/Source/Runtime/Experimental/IoStore/OnDemand/Private/OnDemandInstallCache.cpp
+++ b/Engine/Source/Runtime/Experimental/IoStore/OnDemand/Private/OnDemandInstallCache.cpp
@@ -2186,7 +2186,7 @@ void FOnDemandInstallCache::Shutdown()
 {
 	if (FResult Result = FlushPendingChunksImpl(PendingChunks); Result.HasError())
 	{
-		UE_LOG(LogIoStoreOnDemand, Error, TEXT("Failed to flush pending chunks on shutdown, reason '%s'"), *LexToString(Result.GetError()));
+		UE_LOG(LogIoStoreOnDemand, Error, TEXT("Failed to flush pending chunks on shutdown, reason '%s'"), *Result.GetError().GetErrorMessage(true).ToString());
 	}
 
 	FCas::FLastAccess LastAccess = Cas.ConsumeLastAcccess();
