[HWRT] Fix incorrect normals on geometry collection when using inline HWRT on PC.

#Synced-CL 49278935

diff --git a/Engine/Source/Runtime/Experimental/GeometryCollectionEngine/Private/GeometryCollection/GeometryCollectionRenderData.cpp b/Engine/Source/Runtime/Experimental/GeometryCollectionEngine/Private/GeometryCollection/GeometryCollectionRenderData.cpp
--- a/Engine/Source/Runtime/Experimental/GeometryCollectionEngine/Private/GeometryCollection/GeometryCollectionRenderData.cpp
+++ b/Engine/Source/Runtime/Experimental/GeometryCollectionEngine/Private/GeometryCollection/GeometryCollectionRenderData.cpp
@@ -820,7 +820,7 @@ void FGeometryCollectionRenderData::InitResources(UGeometryCollection const& Own
 							Segment.FirstPrimitive = Section.TriangleStart;
 							Segment.VertexBuffer = MeshResource.PositionVertexBuffer.VertexBufferRHI;
 							Segment.NumPrimitives = Section.TriangleCount;
-							Segment.MaxVertices = Section.VertexEnd;
+							Segment.MaxVertices = MeshResource.PositionVertexBuffer.GetNumVertices(); // TODO: investigate why Section.VertexEnd doesn't work here when using SubSections
 
 							Initializer.Segments.Add(Segment);
 
