[Chaos] Add a cVar to disable pruning in the acceleration structure update.
Pruning is an optimization technique which use multithreading, to update a single threaded process.
In single thread pruning should be disabled.

#Synced-CL 49195397

diff --git a/Engine/Source/Runtime/Experimental/Chaos/Private/Chaos/PBDRigidsEvolution.cpp b/Engine/Source/Runtime/Experimental/Chaos/Private/Chaos/PBDRigidsEvolution.cpp
--- a/Engine/Source/Runtime/Experimental/Chaos/Private/Chaos/PBDRigidsEvolution.cpp
+++ b/Engine/Source/Runtime/Experimental/Chaos/Private/Chaos/PBDRigidsEvolution.cpp
@@ -56,6 +56,11 @@ namespace Chaos
 	CHAOS_API int32 AccelerationStructureTimeSlicingMaxBytesCopy = 100000;
 	FAutoConsoleVariableRef CVarAccelerationStructureTimeSlicingMaxBytesCopy(TEXT("p.Chaos.AccelerationStructureTimeSlicingMaxBytesCopy"), AccelerationStructureTimeSlicingMaxBytesCopy, TEXT("The Maximum number of bytes to copy to the external acceleration structure during Copy Time Slicing"));
 
+	bool bPruningInSingleWorkerPhysics = false;
+	FAutoConsoleVariableRef CVarPruningInSingleWorkerPhysics(TEXT("p.Chaos.PruningInSingleWorkerPhysics"), bPruningInSingleWorkerPhysics, TEXT("Enable Pruning to update acceleration structure in Single Worker mode.\
+																Pruning the acceleration structure is an optimization technique that speed up the spatial acceleration update which is a normally a single threaded process by using several threads to prune data to avoid some unecessary update. \
+																In single worker physics it would be better to disable it."));
+
 	CHAOS_API FBroadPhaseConfig BroadPhaseConfig;
 
 	FAutoConsoleVariableRef CVarBroadphaseIsTree(TEXT("p.BroadphaseType"), BroadPhaseConfig.BroadphaseType, TEXT(""));
@@ -644,18 +649,21 @@ namespace Chaos
 		else
 		{
 			InternalAccelerationQueue.CleanUpDuplicated();
-			for (TArray<FPendingSpatialData>& PendingDataArray : InternalAccelerationQueue.PendingDataArrays)
+			if (bPruningInSingleWorkerPhysics)
 			{
-				TArray<FPendingSpatialData> NewPendingData;
-				NewPendingData.Reserve(PendingDataArray.Num());
-				for (const FPendingSpatialData& PendingData : PendingDataArray)
+				for (TArray<FPendingSpatialData>& PendingDataArray : InternalAccelerationQueue.PendingDataArrays)
 				{
-					if (HasToBeUpdated(PendingData))
+					TArray<FPendingSpatialData> NewPendingData;
+					NewPendingData.Reserve(PendingDataArray.Num());
+					for (const FPendingSpatialData& PendingData : PendingDataArray)
 					{
-						NewPendingData.Add(PendingData);
+						if (HasToBeUpdated(PendingData))
+						{
+							NewPendingData.Add(PendingData);
+						}
 					}
+					PendingDataArray = MoveTemp(NewPendingData);
 				}
-				PendingDataArray = MoveTemp(NewPendingData);
 			}
 		}
 	}
