[HWRT] Add support for cached static ray tracing instances in FNaniteGeometryCollectionSceneProxy
- saves >3ms critical path time on render thread in test level with 300 GCs with 22 parts each.

#Synced-CL 49353664

diff --git a/Engine/Source/Runtime/Experimental/GeometryCollectionEngine/Private/GeometryCollection/GeometryCollectionSceneProxy.cpp b/Engine/Source/Runtime/Experimental/GeometryCollectionEngine/Private/GeometryCollection/GeometryCollectionSceneProxy.cpp
--- a/Engine/Source/Runtime/Experimental/GeometryCollectionEngine/Private/GeometryCollection/GeometryCollectionSceneProxy.cpp
+++ b/Engine/Source/Runtime/Experimental/GeometryCollectionEngine/Private/GeometryCollection/GeometryCollectionSceneProxy.cpp
@@ -94,6 +94,10 @@ static FAutoConsoleVariableRef CVarRayTracingGeometryCollectionWPO(
 	TEXT("r.RayTracing.Geometry.GeometryCollection.WPO"),
 	GRayTracingGeometryCollectionWPO,
 	TEXT("Whether to update geometry collection ray tracing representation based on material World Position Offset."),
+	FConsoleVariableDelegate::CreateLambda([](IConsoleVariable* InVariable)
+		{
+			FGlobalComponentRecreateRenderStateContext Context;
+		}),
 	ECVF_RenderThreadSafe
 );
 
@@ -110,6 +114,10 @@ static TAutoConsoleVariable<int32> CVarRayTracingGeometryCollectionCombinedBLAS(
 	TEXT("Whether to always use a combined BLAS instead of one instance per collection part.\n")
 	TEXT("A combined BLAS needs to be fully rebuilt whenever any transform changes.\n")
 	TEXT("This is automatically enabled for geometry collections using WPO since BLAS must be updated anyway."),
+	FConsoleVariableDelegate::CreateLambda([](IConsoleVariable* InVariable)
+		{
+			FGlobalComponentRecreateRenderStateContext Context;
+		}),
 	ECVF_RenderThreadSafe
 );
 
@@ -667,20 +675,14 @@ void FGeometryCollectionSceneProxyBase::GetDynamicRayTracingInstances(FRayTracin
 	else
 	{
 		check(bUseSubSections);
+		checkf(!bIsNanite, TEXT("Nanite geometry collections are expected to use static ray tracing instances."));
 
-		TConstArrayView<FMatrix> PartInstanceTransforms; // only used when !bIsNanite
+		TArray<FMatrix, SceneRenderingAllocator>& PartInstanceTransforms = Collector.AllocateOneFrameResource<TArray<FMatrix, SceneRenderingAllocator>>();
+		PartInstanceTransforms.AddUninitialized(DynamicData->Transforms.Num());
 
-		if (!bIsNanite)
+		for (int32 TransformIndex = 0; TransformIndex < DynamicData->Transforms.Num(); ++TransformIndex)
 		{
-			TArray<FMatrix, SceneRenderingAllocator>& Tmp = Collector.AllocateOneFrameResource<TArray<FMatrix, SceneRenderingAllocator>>();
-			Tmp.AddUninitialized(DynamicData->Transforms.Num());
-
-			for (int32 TransformIndex = 0; TransformIndex < DynamicData->Transforms.Num(); ++TransformIndex)
-			{
-				Tmp[TransformIndex] = FMatrix(DynamicData->Transforms[TransformIndex]) * LocalToWorld;
-			}
-
-			PartInstanceTransforms = Tmp;
+			PartInstanceTransforms[TransformIndex] = FMatrix(DynamicData->Transforms[TransformIndex]) * LocalToWorld;
 		}
 
 		bool bInstanceMaskAndFlagsDirty = false;
@@ -705,46 +707,16 @@ void FGeometryCollectionSceneProxyBase::GetDynamicRayTracingInstances(FRayTracin
 
 		checkf(PartRayTracingGeometries.Num() == SectionArray.Num(), TEXT("Expected to have one part per section (%d vs %d)."), PartRayTracingGeometries.Num(), SectionArray.Num());
 
-		TArrayView<uint32> PrimitiveInstanceIndices; // only used when bIsNanite
-
-		if (bIsNanite)
-		{
-			checkf(TransformToInstanceMapping.Num() == DynamicData->Transforms.Num(), TEXT("Invalid transform to instance mapping. There are %d transforms but mapping array only has %d entries."), DynamicData->Transforms.Num(), TransformToInstanceMapping.Num());
-
-			// pre-allocate array of primitive instance indices to avoid multiple allocations and copies later
-			TArray<uint32, SceneRenderingAllocator>& Tmp = Collector.AllocateOneFrameResource<TArray<uint32, SceneRenderingAllocator>>();
-			Tmp.AddUninitialized(PartRayTracingGeometries.Num());
-
-			PrimitiveInstanceIndices = Tmp;
-		}
-
 		for (int32 SectionIndex = 0; SectionIndex < PartRayTracingGeometries.Num(); ++SectionIndex)
 		{
 			const FGeometryCollectionMeshElement& Section = SectionArray[SectionIndex];
 
 			FRayTracingInstance RayTracingInstance;
 			RayTracingInstance.Geometry = &PartRayTracingGeometries[SectionIndex];
+			RayTracingInstance.InstanceTransformsView = MakeConstArrayView(&PartInstanceTransforms[Section.TransformIndex], 1);
 			RayTracingInstance.MaterialsView = MakeConstArrayView(&CachedPartRayTracingMaterials[SectionIndex], 1);
 			RayTracingInstance.bInstanceMaskAndFlagsDirty = bInstanceMaskAndFlagsDirty;
 
-			if (bIsNanite)
-			{
-				int32 InstanceIndex = TransformToInstanceMapping[Section.TransformIndex];
-
-				if (InstanceIndex < 0)
-				{
-					continue;
-				}
-
-				PrimitiveInstanceIndices[SectionIndex] = InstanceIndex;
-
-				RayTracingInstance.PrimitiveInstanceIndicesView = MakeConstArrayView(&PrimitiveInstanceIndices[SectionIndex], 1); // parts are stored in GPU Scene as instances
-			}
-			else
-			{
-				RayTracingInstance.InstanceTransformsView = MakeConstArrayView(&PartInstanceTransforms[Section.TransformIndex], 1);
-			}
-
 			for (int32 ViewIndex = 0; ViewIndex < Views.Num(); ViewIndex++)
 			{
 				if ((VisibilityMap & (1 << ViewIndex)) == 0)
@@ -757,6 +729,60 @@ void FGeometryCollectionSceneProxyBase::GetDynamicRayTracingInstances(FRayTracin
 		}
 	}
 }
+
+void FGeometryCollectionSceneProxyBase::GetStaticRayTracingInstances(FStaticRayTracingInstances& OutRayTracingInstances, FRHIUniformBuffer* UniformBuffer)
+{
+	for (const FRayTracingGeometry& PartRayTracingGeometry : PartRayTracingGeometries)
+	{
+		check(PartRayTracingGeometry.IsValid() && !PartRayTracingGeometry.IsEvicted());
+	}
+
+	checkf(bHasRayTracingRepresentation, TEXT("Shouldn't try to get ray tracing instances from proxy that doesn't have a ray tracing representation."));
+
+	const uint32 LODIndex = 0;
+	const bool bWireframe = false;
+	const bool bReverseCulling = false; // DXR doesn't need the flipping implied by the transform
+	const bool bUseSubSections = true;
+
+	TConstArrayView<FGeometryCollectionMeshElement> SectionArray = GetSectionArray(bUseSubSections, false);
+
+	// Grab the material proxies we'll be using for each section
+	// TODO: Add BoneColor support in Path/Ray tracing?
+	FMaterialArrayType MaterialProxies;
+	GetSectionMaterialProxies(SectionArray, MaterialProxies);
+
+	checkf(PartRayTracingGeometries.Num() == SectionArray.Num(), TEXT("Expected to have one part per section (%d vs %d)."), PartRayTracingGeometries.Num(), SectionArray.Num());
+
+	PrimitiveInstanceIndices.Reset();
+	PrimitiveInstanceIndices.SetNumUninitialized(PartRayTracingGeometries.Num());
+
+	// One ray tracing LOD with N instances (one per part)
+
+	FStaticRayTracingInstancesLOD& StaticRayTracingInstancesLOD = OutRayTracingInstances.LODs.AddDefaulted_GetRef();
+
+	for (int32 SectionIndex = 0; SectionIndex < PartRayTracingGeometries.Num(); ++SectionIndex)
+	{
+		const FGeometryCollectionMeshElement& Section = SectionArray[SectionIndex];
+
+		const int32 InstanceIndex = TransformToInstanceMapping[Section.TransformIndex];
+
+		if (InstanceIndex < 0)
+		{
+			continue;
+		}
+
+		PrimitiveInstanceIndices[SectionIndex] = InstanceIndex;
+
+		FRayTracingInstance& RayTracingInstance = StaticRayTracingInstancesLOD.Instances.AddDefaulted_GetRef();
+		RayTracingInstance.Geometry = &PartRayTracingGeometries[SectionIndex];
+		RayTracingInstance.NumTransforms = 1;
+		RayTracingInstance.PrimitiveInstanceIndicesView = MakeConstArrayView(&PrimitiveInstanceIndices[SectionIndex], 1); // parts are stored in GPU Scene as instances
+
+		FMeshBatch& Mesh = RayTracingInstance.Materials.AddDefaulted_GetRef();
+		GetMeshElement(LODIndex, 0, Section, &VertexFactory, MaterialProxies[SectionIndex], SDPG_World, UniformBuffer, bWireframe, bReverseCulling, Mesh);
+	}
+}
+
 #endif
 
 uint32 FGeometryCollectionSceneProxyBase::GetAllocatedSize() const
@@ -1430,6 +1456,13 @@ FNaniteGeometryCollectionSceneProxy::FNaniteGeometryCollectionSceneProxy(UGeomet
 	CollisionResponse = Component->GetCollisionResponseToChannels();
 
 	UpdateInstanceSceneDataBuffers(Component->GetRenderMatrix());
+
+#if RHI_RAYTRACING
+	if (HasRayTracingRepresentation())
+	{
+		bUseStaticRayTracingInstances = !MeshDescription.SubSections.IsEmpty() && (!bAnyMaterialHasWorldPositionOffset || GRayTracingGeometryCollectionWPO) && CVarRayTracingGeometryCollectionCombinedBLAS.GetValueOnGameThread() == 0;
+	}
+#endif
 }
 
 FNaniteGeometryCollectionSceneProxy::~FNaniteGeometryCollectionSceneProxy()
@@ -1702,6 +1735,33 @@ void FNaniteGeometryCollectionSceneProxy::GetDynamicMeshElements(const TArray<co
 	}
 }
 
+#if RHI_RAYTRACING
+ERayTracingPrimitiveFlags FNaniteGeometryCollectionSceneProxy::GetRayTracingPrimitiveFlags()
+{
+	if (!bUseStaticRayTracingInstances)
+	{
+		return Super::GetRayTracingPrimitiveFlags();
+	}
+
+	if (!(IsVisibleInRayTracing() && ShouldRenderInMainPass() && (IsDrawnInGame() || AffectsIndirectLightingWhileHidden() || CastsHiddenShadow())) && !IsRayTracingFarField())
+	{
+		return ERayTracingPrimitiveFlags::Exclude;
+	}
+
+	if (!HasRayTracingRepresentation())
+	{
+		return ERayTracingPrimitiveFlags::Exclude;
+	}
+
+	if (IsFirstPerson())
+	{
+		// First person primitives are currently not supported in raytracing as this kind of geometry only makes sense from the camera's point of view.
+		return ERayTracingPrimitiveFlags::Exclude;
+	}
+
+	return ERayTracingPrimitiveFlags::CacheInstances;
+}
+#endif // RHI_RAYTRACING
 
 FGeometryCollectionDynamicDataPool::FGeometryCollectionDynamicDataPool()
 {
diff --git a/Engine/Source/Runtime/Experimental/GeometryCollectionEngine/Private/GeometryCollection/GeometryCollectionSceneProxy.h b/Engine/Source/Runtime/Experimental/GeometryCollectionEngine/Private/GeometryCollection/GeometryCollectionSceneProxy.h
--- a/Engine/Source/Runtime/Experimental/GeometryCollectionEngine/Private/GeometryCollection/GeometryCollectionSceneProxy.h
+++ b/Engine/Source/Runtime/Experimental/GeometryCollectionEngine/Private/GeometryCollection/GeometryCollectionSceneProxy.h
@@ -196,9 +196,10 @@ protected:
 	void UpdatingRayTracingGeometry_RenderingThread(FRHICommandListBase& RHICmdList, TConstArrayView<FGeometryCollectionMeshElement> InSectionArray, bool bRayTracingGeometryPerSection);
 
 	bool IsRayTracingRelevant() const { return true; }
-	bool IsRayTracingStaticRelevant() const { return false; }
+	bool IsRayTracingStaticRelevant() const { return bUseStaticRayTracingInstances; }
 	bool HasRayTracingRepresentation() const { return bHasRayTracingRepresentation; }
 	void GetDynamicRayTracingInstances(FRayTracingInstanceCollector& Collector, const FMatrix& LocalToWorld, FRHIUniformBuffer* UniformBuffer, bool bAnyMaterialHasWorldPositionOffset);
+	void GetStaticRayTracingInstances(FStaticRayTracingInstances& OutRayTracingInstances, FRHIUniformBuffer* UniformBuffer);
 #endif
 
 	uint32 GetAllocatedSize() const;
@@ -226,6 +227,10 @@ protected:
 
 	FGeometryCollectionDynamicData* DynamicData = nullptr;
 
+#if RHI_RAYTRACING
+	bool bUseStaticRayTracingInstances = false;
+#endif
+
 private:
 
 	/** Update skinned position buffer used by mobile CPU skinning path. */
@@ -244,6 +249,7 @@ private:
 
 	TConstArrayView<FRayTracingGeometry> PartRayTracingGeometries;
 	TArray<FMeshBatch> CachedPartRayTracingMaterials;
+	TArray<uint32> PrimitiveInstanceIndices;
 #endif
 
 	bool bRenderResourcesCreated = false;
@@ -314,10 +320,10 @@ public:
 	virtual void DrawStaticElements(FStaticPrimitiveDrawInterface* PDI) override;
 
 #if RHI_RAYTRACING
-	bool IsRayTracingRelevant() const override { return FGeometryCollectionSceneProxyBase::IsRayTracingRelevant(); }
-	bool IsRayTracingStaticRelevant() const override { return FGeometryCollectionSceneProxyBase::IsRayTracingStaticRelevant(); }
-	bool HasRayTracingRepresentation() const override { return FGeometryCollectionSceneProxyBase::HasRayTracingRepresentation(); }
-	void GetDynamicRayTracingInstances(FRayTracingInstanceCollector& Collector) override { return FGeometryCollectionSceneProxyBase::GetDynamicRayTracingInstances(Collector, GetLocalToWorld(), GetUniformBuffer(), bAnyMaterialHasWorldPositionOffset); }
+	virtual bool IsRayTracingRelevant() const override { return FGeometryCollectionSceneProxyBase::IsRayTracingRelevant(); }
+	virtual bool IsRayTracingStaticRelevant() const override { return FGeometryCollectionSceneProxyBase::IsRayTracingStaticRelevant(); }
+	virtual bool HasRayTracingRepresentation() const override { return FGeometryCollectionSceneProxyBase::HasRayTracingRepresentation(); }
+	virtual void GetDynamicRayTracingInstances(FRayTracingInstanceCollector& Collector) override { return FGeometryCollectionSceneProxyBase::GetDynamicRayTracingInstances(Collector, GetLocalToWorld(), GetUniformBuffer(), bAnyMaterialHasWorldPositionOffset); }
 #endif
 
 protected:
@@ -381,10 +387,12 @@ public:
 	}
 
 #if RHI_RAYTRACING
-	bool IsRayTracingRelevant() const override { return FGeometryCollectionSceneProxyBase::IsRayTracingRelevant(); }
-	bool IsRayTracingStaticRelevant() const override { return FGeometryCollectionSceneProxyBase::IsRayTracingStaticRelevant(); }
-	bool HasRayTracingRepresentation() const override { return FGeometryCollectionSceneProxyBase::HasRayTracingRepresentation(); }
-	void GetDynamicRayTracingInstances(FRayTracingInstanceCollector& Collector) override { return FGeometryCollectionSceneProxyBase::GetDynamicRayTracingInstances(Collector, GetLocalToWorld(), GetUniformBuffer(), bAnyMaterialHasWorldPositionOffset); }
+	virtual bool IsRayTracingRelevant() const override { return FGeometryCollectionSceneProxyBase::IsRayTracingRelevant(); }
+	virtual bool IsRayTracingStaticRelevant() const override { return FGeometryCollectionSceneProxyBase::IsRayTracingStaticRelevant(); }
+	virtual bool HasRayTracingRepresentation() const override { return FGeometryCollectionSceneProxyBase::HasRayTracingRepresentation(); }
+	virtual void GetDynamicRayTracingInstances(FRayTracingInstanceCollector& Collector) override { return FGeometryCollectionSceneProxyBase::GetDynamicRayTracingInstances(Collector, GetLocalToWorld(), GetUniformBuffer(), bAnyMaterialHasWorldPositionOffset); }
+	virtual ERayTracingPrimitiveFlags GetRayTracingPrimitiveFlags() override;
+	virtual void GetStaticRayTracingInstances(FStaticRayTracingInstances& OutRayTracingInstances) override { return FGeometryCollectionSceneProxyBase::GetStaticRayTracingInstances(OutRayTracingInstances, GetUniformBuffer()); }
 #endif
 
 protected:
diff --git a/Engine/Source/Runtime/Renderer/Private/PrimitiveSceneInfo.cpp b/Engine/Source/Runtime/Renderer/Private/PrimitiveSceneInfo.cpp
--- a/Engine/Source/Runtime/Renderer/Private/PrimitiveSceneInfo.cpp
+++ b/Engine/Source/Runtime/Renderer/Private/PrimitiveSceneInfo.cpp
@@ -1377,6 +1377,16 @@ void FPrimitiveSceneInfo::SetupCachedRayTracingInstance(FPrimitiveSceneInfo* Sce
 	OutCachedRayTracingInstance.NumTransforms = RayTracingInstance.NumTransforms;
 	OutCachedRayTracingInstance.BaseInstanceSceneDataOffset = SceneInfo->GetInstanceSceneDataOffset();
 
+	{
+		TConstArrayView<uint32> PrimitiveInstanceIndices = RayTracingInstance.GetPrimitiveInstanceIndices();
+		if (!PrimitiveInstanceIndices.IsEmpty())
+		{
+			checkf(RayTracingInstance.NumTransforms <= 1 && PrimitiveInstanceIndices.Num() == 1,
+				TEXT("Cached static ray tracing instances currently only supports PrimitiveInstanceIndices for instances with a single transform."));
+			OutCachedRayTracingInstance.BaseInstanceSceneDataOffset += PrimitiveInstanceIndices[0];
+		}
+	}
+
 	if (Nanite::GetRayTracingMode() != Nanite::ERayTracingMode::Fallback && SceneProxy->IsNaniteMesh())
 	{
 		// nanite ray tracing geometry might not be ready yet
