Geometry collection : Add support for GyroscopicTorque option to Geometry Collections

#Synced-CL 49325867

diff --git a/Engine/Source/Runtime/Experimental/Chaos/Private/PhysicsProxy/GeometryCollectionPhysicsProxy.cpp b/Engine/Source/Runtime/Experimental/Chaos/Private/PhysicsProxy/GeometryCollectionPhysicsProxy.cpp
--- a/Engine/Source/Runtime/Experimental/Chaos/Private/PhysicsProxy/GeometryCollectionPhysicsProxy.cpp
+++ b/Engine/Source/Runtime/Experimental/Chaos/Private/PhysicsProxy/GeometryCollectionPhysicsProxy.cpp
@@ -1751,6 +1751,7 @@ void FGeometryCollectionPhysicsProxy::InitializeBodiesPT(Chaos::FPBDRigidsSolver
 				Handle->SetAngularEtherDrag(Parameters.AngularDamping);
 				Handle->SetInitialOverlapDepenetrationVelocity(Parameters.InitialOverlapDepenetrationVelocity);
 				Handle->SetSleepThresholdMultiplier(Parameters.SleepThresholdMultiplier);
+				Handle->SetGyroscopicTorqueEnabled(Parameters.bGyroscopicTorqueEnabled);
 			}
 
 #if !(UE_BUILD_SHIPPING || UE_BUILD_TEST)
diff --git a/Engine/Source/Runtime/Experimental/Chaos/Public/GeometryCollection/GeometryCollectionSimulationCoreTypes.h b/Engine/Source/Runtime/Experimental/Chaos/Public/GeometryCollection/GeometryCollectionSimulationCoreTypes.h
--- a/Engine/Source/Runtime/Experimental/Chaos/Public/GeometryCollection/GeometryCollectionSimulationCoreTypes.h
+++ b/Engine/Source/Runtime/Experimental/Chaos/Public/GeometryCollection/GeometryCollectionSimulationCoreTypes.h
@@ -219,6 +219,7 @@ PRAGMA_ENABLE_DEPRECATION_WARNINGS
 		, bUseSimplicialsWhenAvailable(false)
 		, bUseDamagePropagation(false)
 		, bOptimizeForRuntimeMemory(false)
+		, bGyroscopicTorqueEnabled(false)
 		, PositionSolverIterations(8)
 		, VelocitySolverIterations(1)
 		, ProjectionSolverIterations(1)
@@ -302,6 +303,7 @@ PRAGMA_ENABLE_DEPRECATION_WARNINGS
 		, bUseSimplicialsWhenAvailable(Other.bUseSimplicialsWhenAvailable)
 		, bUseDamagePropagation(Other.bUseDamagePropagation)
 		, bOptimizeForRuntimeMemory(Other.bOptimizeForRuntimeMemory)
+		, bGyroscopicTorqueEnabled(Other.bGyroscopicTorqueEnabled)
 		, PositionSolverIterations(Other.PositionSolverIterations)
 		, VelocitySolverIterations(Other.VelocitySolverIterations)
 		, ProjectionSolverIterations(Other.ProjectionSolverIterations)
@@ -420,6 +422,7 @@ PRAGMA_ENABLE_DEPRECATION_WARNINGS
 
 	bool bUseDamagePropagation : 1;
 	bool bOptimizeForRuntimeMemory : 1;
+	bool bGyroscopicTorqueEnabled : 1;
 
 	int32 PositionSolverIterations;
 	int32 VelocitySolverIterations;
diff --git a/Engine/Source/Runtime/Experimental/GeometryCollectionEngine/Private/GeometryCollection/GeometryCollectionComponent.cpp b/Engine/Source/Runtime/Experimental/GeometryCollectionEngine/Private/GeometryCollection/GeometryCollectionComponent.cpp
--- a/Engine/Source/Runtime/Experimental/GeometryCollectionEngine/Private/GeometryCollection/GeometryCollectionComponent.cpp
+++ b/Engine/Source/Runtime/Experimental/GeometryCollectionEngine/Private/GeometryCollection/GeometryCollectionComponent.cpp
@@ -3966,6 +3966,7 @@ PRAGMA_ENABLE_DEPRECATION_WARNINGS
 		SimulationParameters.InitialOverlapDepenetrationVelocity = BodyInstance.GetMaxDepenetrationVelocity();
 		SimulationParameters.SleepThresholdMultiplier = BodyInstance.GetSleepThresholdMultiplier();
 		SimulationParameters.bUseDamagePropagation = DamagePropagationData.bEnabled;
+		SimulationParameters.bGyroscopicTorqueEnabled = BodyInstance.bGyroscopicTorqueEnabled;
 		SimulationParameters.BreakDamagePropagationFactor = DamagePropagationData.BreakDamagePropagationFactor;
 		SimulationParameters.ShockDamagePropagationFactor = DamagePropagationData.ShockDamagePropagationFactor;
 		SimulationParameters.WorldTransform = GetComponentToWorld();
