[NetPhysics] CVar to choose which physics step (if multiple are scheduled at the same time) to trigger resimulation on.

#Synced-CL 49326374

diff --git a/Engine/Source/Programs/HeadlessChaos/Private/HeadlessChaosTestRewind.cpp b/Engine/Source/Programs/HeadlessChaos/Private/HeadlessChaosTestRewind.cpp
--- a/Engine/Source/Programs/HeadlessChaos/Private/HeadlessChaosTestRewind.cpp
+++ b/Engine/Source/Programs/HeadlessChaos/Private/HeadlessChaosTestRewind.cpp
@@ -577,7 +577,8 @@ namespace ChaosTest {
 								if (!bHasResimmed || Time < 4.5)
 								{
 									//simply movement without hitting floor because it's spawned too late
-									EXPECT_NEAR(PT0->X()[2], 14.5 - Time, 1e-2);
+									const float ExpectedLocation = FMath::Max(1.5, 14.5 - Time);
+									EXPECT_NEAR(PT0->X()[2], ExpectedLocation, 1e-2);
 								}
 								else
 								{
@@ -2347,9 +2348,6 @@ namespace ChaosTest {
 				}
 			}
 
-			EXPECT_EQ(FirstStepResim, 0);
-			EXPECT_EQ(NumResimSteps, NumPhysSteps);
-
 			//no desync so should be empty
 #if REWIND_DESYNC
 			const TArray<FDesyncedParticleInfo> DesyncedParticles = RewindData->ComputeDesyncInfo();
diff --git a/Engine/Source/Runtime/Experimental/Chaos/Private/Chaos/Framework/PhysicsSolverBase.cpp b/Engine/Source/Runtime/Experimental/Chaos/Private/Chaos/Framework/PhysicsSolverBase.cpp
--- a/Engine/Source/Runtime/Experimental/Chaos/Private/Chaos/Framework/PhysicsSolverBase.cpp
+++ b/Engine/Source/Runtime/Experimental/Chaos/Private/Chaos/Framework/PhysicsSolverBase.cpp
@@ -31,6 +31,10 @@ namespace Chaos
 	CHAOS_API int32 RewindBeforeAdvance = 1;
 	FAutoConsoleVariableRef CVarReorderRewindResimulation(TEXT("p.Resim.RewindBeforeAdvance"), RewindBeforeAdvance, TEXT("Whether to apply the rewind before or after the advance"));
 
+	CHAOS_API int32 ResimulateOnSpecifiedStep = 1;
+	FAutoConsoleVariableRef CVarResimulateOnSpecifiedStep(TEXT("p.Resim.ResimulateOnSpecifiedStep"), ResimulateOnSpecifiedStep
+		, TEXT("0 = Resimulation can trigger on any physics step. \n -1 or other negative values = Resimulation only happens on the last physics step (if multiple are queued up at the same time). \n N = Resimulation only happens on the specified step or last, whichever comes first (if multiple are queued up at the same time)"));
+
 
 	extern int GSingleThreadedPhysics;
 	void FPhysicsSolverBase::ChangeBufferMode(EMultiBufferMode InBufferMode)
@@ -157,13 +161,24 @@ namespace Chaos
 #if PHYSICS_THREAD_CONTEXT
 		FPhysicsThreadContextScope Scope(/*IsPhysicsThreadContext=*/true);
 #endif
-		if(PushData && !Solver.IsResimming())
-		{ 
-			if(RewindBeforeAdvance != 0)
+
+		if (PushData && !Solver.IsResimming())
+		{
+			if (RewindBeforeAdvance != 0)
 			{
 				Solver.AddResimulationRequest_Internal(*PushData);
 			}
 
+			// Check if a resimulation is allowed based on which step to resim on when multiple physics steps are scheduled in a row from the game thread
+			const int32 SpecifiedStep = ResimulateOnSpecifiedStep < 0
+				? PushData->IntervalNumSteps // Only resimulate on the last frame
+				: FMath::Min(ResimulateOnSpecifiedStep, PushData->IntervalNumSteps); // Resimulate on specified step or last if not enough steps
+
+			const bool bAllowResimulation = ResimulateOnSpecifiedStep == 0 
+				? true // Always allow resimulation
+				: (PushData->IntervalStep + 1) == SpecifiedStep; // Allow resimulation on specific step
+
+			if (bAllowResimulation)
 			{
 				SCOPE_CYCLE_COUNTER(STAT_ConditionalApplyRewind);
 				Solver.ConditionalApplyRewind_Internal();
