Add new cvar 'iax.HttpSortRequestsByType' that allow the system to favor different http requests when sorting the priority order.

#rnx

- When set to -1 or any negative value (the default) we favor IAS requests.
- When set to 1 or any positive value we favor IAD requests.
- When set to 0 we ignore request type and only sort on request priority.
- This was added as IAD is often able to flood the request queue with many small requests often with a high priority which can leave IAS quite starved and unable to process for many minutes. Since IAs tends to issue less (but larger) requests which often have a very visible quality effect on the screen it makes sense to get them done sooner than later. We can always change the weighting in the future if these scenarios change.
-- Might change this to 0 by default in the near future and move the decision to opt to IAS first as a project config option.

#Synced-CL 49321981

diff --git a/Engine/Source/Runtime/Experimental/IoStore/OnDemand/Private/OnDemandHttpIoDispacher.cpp b/Engine/Source/Runtime/Experimental/IoStore/OnDemand/Private/OnDemandHttpIoDispacher.cpp
--- a/Engine/Source/Runtime/Experimental/IoStore/OnDemand/Private/OnDemandHttpIoDispacher.cpp
+++ b/Engine/Source/Runtime/Experimental/IoStore/OnDemand/Private/OnDemandHttpIoDispacher.cpp
@@ -110,6 +110,13 @@ static FAutoConsoleVariableRef CVar_HttpIoDispatcherMaxRangeKiB(
 	ECVF_Default
 );
 
+int32 GIaxHttpSortRequestsByType = -1;
+static FAutoConsoleVariableRef CVar_IaxHttpSortRequestsByType(
+	TEXT("iax.HttpSortRequestsByType"),
+	GIaxHttpSortRequestsByType,
+	TEXT("-1 = Favor IAS Request, 0 = Treat all requests equally, 1 = Favor IAD requests")
+);
+
 float GHttpIoDispatcherShutdownTimeoutSeconds = 5.0;
 static FAutoConsoleVariableRef CVar_HttpIoDispatcherShutdownTimeoutSeconds(
 	TEXT("iax.HttpIoDispatcherShutdownTimeoutSeconds"),
@@ -484,6 +491,21 @@ FReturnedRefCountValue FHttpRequest::Release() const
 ////////////////////////////////////////////////////////////////////////////////
 static bool HttpRequestSortPredicate(const FHttpRequest& LHS, const FHttpRequest& RHS)
 {
+	// 'iax.HttpSortRequestsByType' must be enabled and the requests need to be of different types
+	// Note this assumes that all requests in the chain have the same type which is currently true.
+	if (GIaxHttpSortRequestsByType != 0 && LHS.GetHttpRequestType() != RHS.GetHttpRequestType())
+	{
+		if (GIaxHttpSortRequestsByType < 0) // Negative value favors streaming (IAS)
+		{
+			return LHS.GetHttpRequestType() == EHttpRequestType::Streaming;
+		}
+
+		if (GIaxHttpSortRequestsByType > 0) // Positive value favors installed (IAD)
+		{
+			return LHS.GetHttpRequestType() == EHttpRequestType::Installed;
+		}
+	}
+
 	if (LHS.Options.GetPriority() == RHS.Options.GetPriority())
 	{
 		return LHS.SeqNo < RHS.SeqNo;
