Change enum to enum class and deprecate old ESpatialAccelerationCollectionBucketInnerIdx. Necessary to avoid naming conflict with clang.

#Synced-CL 49358024

diff --git a/Engine/Source/Runtime/Experimental/Chaos/Private/Chaos/PBDRigidsEvolution.cpp b/Engine/Source/Runtime/Experimental/Chaos/Private/Chaos/PBDRigidsEvolution.cpp
--- a/Engine/Source/Runtime/Experimental/Chaos/Private/Chaos/PBDRigidsEvolution.cpp
+++ b/Engine/Source/Runtime/Experimental/Chaos/Private/Chaos/PBDRigidsEvolution.cpp
@@ -82,7 +82,9 @@ namespace Chaos
 
 		static bool IsDynamicTree(FSpatialAccelerationIdx SpatialAccelerationIdx)
 		{
-			return AccelerationStructureUseDynamicTree && (SpatialAccelerationIdx.InnerIdx == ESpatialAccelerationCollectionBucketInnerIdx::Dynamic || SpatialAccelerationIdx.InnerIdx == ESpatialAccelerationCollectionBucketInnerIdx::DynamicQueryOnly);
+			constexpr auto DynamicAsUInt16 = static_cast<std::underlying_type_t<ESpatialAccelerationCollectionBucketInnerIndex>>(ESpatialAccelerationCollectionBucketInnerIndex::Dynamic);
+			constexpr auto DynamicQueryOnlyAsUInt16 = static_cast<std::underlying_type_t<ESpatialAccelerationCollectionBucketInnerIndex>>(ESpatialAccelerationCollectionBucketInnerIndex::DynamicQueryOnly);			
+			return AccelerationStructureUseDynamicTree && (SpatialAccelerationIdx.InnerIdx == DynamicAsUInt16 || SpatialAccelerationIdx.InnerIdx == DynamicQueryOnlyAsUInt16);
 		}
 
 		TUniquePtr<ISpatialAccelerationCollection<FAccelerationStructureHandle, FReal, 3>> CreateEmptyCollection() override
@@ -95,24 +97,29 @@ namespace Chaos
 			for (uint16 BucketIdx = 0; BucketIdx < NumBuckets; ++BucketIdx)
 			{
 				// For static bodies
-				Collection->AddSubstructure(CreateAccelerationPerBucket_Threaded(Empty, BucketIdx, true, false, false), BucketIdx, ESpatialAccelerationCollectionBucketInnerIdx::Default);
+				constexpr auto DefaultAsUInt16 = static_cast<std::underlying_type_t<ESpatialAccelerationCollectionBucketInnerIndex>>(ESpatialAccelerationCollectionBucketInnerIndex::Default);
+				Collection->AddSubstructure(CreateAccelerationPerBucket_Threaded(Empty, BucketIdx, true, false, false), BucketIdx, DefaultAsUInt16);
 
 				// Always create this if QueryOnlyObjects are isolated to ensure all consecutive indices are created (within buckets)
 				if (AccelerationStructureSplitStaticAndDynamic == 1 || AccelerationStructureIsolateQueryOnlyObjects == 1) 
 				{
+					constexpr auto DynamicAsUInt16 = static_cast<std::underlying_type_t<ESpatialAccelerationCollectionBucketInnerIndex>>(ESpatialAccelerationCollectionBucketInnerIndex::Dynamic);
 					// Non static bodies
-					Collection->AddSubstructure(CreateAccelerationPerBucket_Threaded(Empty, BucketIdx, true, IsDynamicTree(FSpatialAccelerationIdx{BucketIdx, ESpatialAccelerationCollectionBucketInnerIdx::Dynamic }), true), BucketIdx, ESpatialAccelerationCollectionBucketInnerIdx::Dynamic);
+					Collection->AddSubstructure(CreateAccelerationPerBucket_Threaded(Empty, BucketIdx, true, IsDynamicTree(FSpatialAccelerationIdx{BucketIdx, DynamicAsUInt16 }), true), BucketIdx, DynamicAsUInt16);
 				}
 			}
 
 			if (AccelerationStructureIsolateQueryOnlyObjects && NumBuckets == 1)
 			{
+				const auto DefaultQueryOnlyAsUInt16 = static_cast<std::underlying_type_t<ESpatialAccelerationCollectionBucketInnerIndex>>(ESpatialAccelerationCollectionBucketInnerIndex::DefaultQueryOnly);
 				constexpr uint16 BucketIdx = 0;
-				Collection->AddSubstructure(CreateAccelerationPerBucket_Threaded(Empty, BucketIdx, true, false, false), BucketIdx, ESpatialAccelerationCollectionBucketInnerIdx::DefaultQueryOnly);
+				Collection->AddSubstructure(CreateAccelerationPerBucket_Threaded(Empty, BucketIdx, true, false, false), BucketIdx, DefaultQueryOnlyAsUInt16);
 				if (AccelerationStructureSplitStaticAndDynamic == 1)
 				{
+					constexpr auto DynamicQueryOnlyAsUInt16 = static_cast<std::underlying_type_t<ESpatialAccelerationCollectionBucketInnerIndex>>(ESpatialAccelerationCollectionBucketInnerIndex::DynamicQueryOnly);
+
 					// Non static bodies
-					Collection->AddSubstructure(CreateAccelerationPerBucket_Threaded(Empty, BucketIdx, true, IsDynamicTree(FSpatialAccelerationIdx{ BucketIdx,  ESpatialAccelerationCollectionBucketInnerIdx::DynamicQueryOnly }), false), BucketIdx, ESpatialAccelerationCollectionBucketInnerIdx::DynamicQueryOnly);
+					Collection->AddSubstructure(CreateAccelerationPerBucket_Threaded(Empty, BucketIdx, true, IsDynamicTree(FSpatialAccelerationIdx{ BucketIdx,  DynamicQueryOnlyAsUInt16 }), false), BucketIdx, DynamicQueryOnlyAsUInt16);
 				}
 			}
 
diff --git a/Engine/Source/Runtime/Experimental/Chaos/Public/Chaos/SpatialAccelerationCollection.h b/Engine/Source/Runtime/Experimental/Chaos/Public/Chaos/SpatialAccelerationCollection.h
--- a/Engine/Source/Runtime/Experimental/Chaos/Public/Chaos/SpatialAccelerationCollection.h
+++ b/Engine/Source/Runtime/Experimental/Chaos/Public/Chaos/SpatialAccelerationCollection.h
@@ -16,7 +16,7 @@ namespace Chaos
 {
 
 // Bucket inner indices define for specific use cases		
-enum ESpatialAccelerationCollectionBucketInnerIdx
+enum class ESpatialAccelerationCollectionBucketInnerIndex
 {
 	Default = 0, // If acceleration structures are not split up by body type, they end up in this bucket, otherwise static bodies will go here
 	Dynamic = 1, // For dynamic bodies if they use a separate acceleration structure
@@ -24,6 +24,16 @@ enum ESpatialAccelerationCollectionBucketInnerIdx
 	DynamicQueryOnly = 3 // Dynamic Query only bodies if they are separated
 };
 
+// Bucket inner indices define for specific use cases		
+enum UE_DEPRECATED(5.8, "Use ESpatialAccelerationCollectionBucketInnerIndex instead")
+ESpatialAccelerationCollectionBucketInnerIdx
+{
+	Default_Renamed_To_Avoid_Compile_Errors = 0, // If acceleration structures are not split up by body type, they end up in this bucket, otherwise static bodies will go here
+	Dynamic = 1, // For dynamic bodies if they use a separate acceleration structure
+	DefaultQueryOnly = 2, // Query only bodies if they are separated
+	DynamicQueryOnly = 3 // Dynamic Query only bodies if they are separated
+};
+
 template <typename T>
 void FreeObjHelper(T*& RawPtr)
 {
@@ -538,7 +548,10 @@ public:
 			int BucketInnerIdx = 0;
 			for (const BucketEntryType& BucketEntry : Bucket.Objects)
 			{
-				if (BucketInnerIdx != ESpatialAccelerationCollectionBucketInnerIdx::DefaultQueryOnly && BucketInnerIdx != ESpatialAccelerationCollectionBucketInnerIdx::DynamicQueryOnly)
+				constexpr auto DefaultQueryOnlyAsUInt16 = static_cast<std::underlying_type_t<ESpatialAccelerationCollectionBucketInnerIndex>>(ESpatialAccelerationCollectionBucketInnerIndex::DefaultQueryOnly);
+				constexpr auto DynamicQueryOnlyAsUInt16 = static_cast<std::underlying_type_t<ESpatialAccelerationCollectionBucketInnerIndex>>(ESpatialAccelerationCollectionBucketInnerIndex::DynamicQueryOnly);			
+				
+				if (BucketInnerIdx != DefaultQueryOnlyAsUInt16 && BucketInnerIdx != DynamicQueryOnlyAsUInt16)
 				{
 					auto AccelStructure = static_cast<Chaos::TAABBTree<TPayloadType, Chaos::TAABBTreeLeafArray<Chaos::FAccelerationStructureHandle, true, T>, true, T>*>(BucketEntry.Acceleration.Get());
 					AccelStructure->OverlapFast(QueryBounds, Visitor);
diff --git a/Engine/Source/Runtime/PhysicsCore/Private/ChaosEngineInterface.cpp b/Engine/Source/Runtime/PhysicsCore/Private/ChaosEngineInterface.cpp
--- a/Engine/Source/Runtime/PhysicsCore/Private/ChaosEngineInterface.cpp
+++ b/Engine/Source/Runtime/PhysicsCore/Private/ChaosEngineInterface.cpp
@@ -2447,29 +2447,29 @@ void FChaosEngineInterface::CreateActor(const FActorCreationParams& InParams,FPh
 
 	// Set the particle acceleration structure spatial index here
 	{
-		FSpatialAccelerationIdx SpatialIndex{0, ESpatialAccelerationCollectionBucketInnerIdx::Default };
+		FSpatialAccelerationIdx SpatialIndex{0, static_cast<std::underlying_type_t<ESpatialAccelerationCollectionBucketInnerIndex>>(ESpatialAccelerationCollectionBucketInnerIndex::Default) };
 		if (AccelerationStructureSplitStaticAndDynamic == 1)
 		{
 			if (AccelerationStructureIsolateQueryOnlyObjects == 1)
 			{
 				if (InParams.bStatic && InParams.bQueryOnly)
 				{
-					SpatialIndex = FSpatialAccelerationIdx{0, ESpatialAccelerationCollectionBucketInnerIdx::DefaultQueryOnly};
+					SpatialIndex = FSpatialAccelerationIdx{0, static_cast<std::underlying_type_t<ESpatialAccelerationCollectionBucketInnerIndex>>(ESpatialAccelerationCollectionBucketInnerIndex::DefaultQueryOnly)};
 				}
 				else if (!InParams.bStatic && InParams.bQueryOnly)
 				{
-					SpatialIndex = FSpatialAccelerationIdx{ 0, ESpatialAccelerationCollectionBucketInnerIdx::DynamicQueryOnly};
+					SpatialIndex = FSpatialAccelerationIdx{ 0, static_cast<std::underlying_type_t<ESpatialAccelerationCollectionBucketInnerIndex>>(ESpatialAccelerationCollectionBucketInnerIndex::DynamicQueryOnly)};
 				}
 				else if (!InParams.bStatic && !InParams.bQueryOnly)
 				{
-					SpatialIndex = FSpatialAccelerationIdx{ 0, ESpatialAccelerationCollectionBucketInnerIdx::Dynamic };
+					SpatialIndex = FSpatialAccelerationIdx{ 0, static_cast<std::underlying_type_t<ESpatialAccelerationCollectionBucketInnerIndex>>(ESpatialAccelerationCollectionBucketInnerIndex::Dynamic) };
 				}
 			}
 			else
 			{
 				if (!InParams.bStatic)
 				{
-					SpatialIndex = FSpatialAccelerationIdx{ 0, ESpatialAccelerationCollectionBucketInnerIdx::Dynamic };
+					SpatialIndex = FSpatialAccelerationIdx{ 0, static_cast<std::underlying_type_t<ESpatialAccelerationCollectionBucketInnerIndex>>(ESpatialAccelerationCollectionBucketInnerIndex::Dynamic) };
 				}				
 			}
 		}
@@ -2479,7 +2479,7 @@ void FChaosEngineInterface::CreateActor(const FActorCreationParams& InParams,FPh
 			{
 				if (InParams.bQueryOnly)
 				{
-					SpatialIndex = FSpatialAccelerationIdx{ 0, ESpatialAccelerationCollectionBucketInnerIdx::DefaultQueryOnly };
+					SpatialIndex = FSpatialAccelerationIdx{ 0, static_cast<std::underlying_type_t<ESpatialAccelerationCollectionBucketInnerIndex>>(ESpatialAccelerationCollectionBucketInnerIndex::DefaultQueryOnly) };
 				}
 			}		
 		}
