remove auto and use uint16

#Synced-CL 49374523

diff --git a/Engine/Source/Runtime/Experimental/Chaos/Private/Chaos/PBDRigidsEvolution.cpp b/Engine/Source/Runtime/Experimental/Chaos/Private/Chaos/PBDRigidsEvolution.cpp
--- a/Engine/Source/Runtime/Experimental/Chaos/Private/Chaos/PBDRigidsEvolution.cpp
+++ b/Engine/Source/Runtime/Experimental/Chaos/Private/Chaos/PBDRigidsEvolution.cpp
@@ -82,8 +82,8 @@ namespace Chaos
 
 		static bool IsDynamicTree(FSpatialAccelerationIdx SpatialAccelerationIdx)
 		{
-			constexpr auto DynamicAsUInt16 = static_cast<std::underlying_type_t<ESpatialAccelerationCollectionBucketInnerIndex>>(ESpatialAccelerationCollectionBucketInnerIndex::Dynamic);
-			constexpr auto DynamicQueryOnlyAsUInt16 = static_cast<std::underlying_type_t<ESpatialAccelerationCollectionBucketInnerIndex>>(ESpatialAccelerationCollectionBucketInnerIndex::DynamicQueryOnly);			
+			constexpr uint16 DynamicAsUInt16 = static_cast<std::underlying_type_t<ESpatialAccelerationCollectionBucketInnerIndex>>(ESpatialAccelerationCollectionBucketInnerIndex::Dynamic);
+			constexpr uint16 DynamicQueryOnlyAsUInt16 = static_cast<std::underlying_type_t<ESpatialAccelerationCollectionBucketInnerIndex>>(ESpatialAccelerationCollectionBucketInnerIndex::DynamicQueryOnly);			
 			return AccelerationStructureUseDynamicTree && (SpatialAccelerationIdx.InnerIdx == DynamicAsUInt16 || SpatialAccelerationIdx.InnerIdx == DynamicQueryOnlyAsUInt16);
 		}
 
@@ -97,13 +97,13 @@ namespace Chaos
 			for (uint16 BucketIdx = 0; BucketIdx < NumBuckets; ++BucketIdx)
 			{
 				// For static bodies
-				constexpr auto DefaultAsUInt16 = static_cast<std::underlying_type_t<ESpatialAccelerationCollectionBucketInnerIndex>>(ESpatialAccelerationCollectionBucketInnerIndex::Default);
+				constexpr uint16 DefaultAsUInt16 = static_cast<std::underlying_type_t<ESpatialAccelerationCollectionBucketInnerIndex>>(ESpatialAccelerationCollectionBucketInnerIndex::Default);
 				Collection->AddSubstructure(CreateAccelerationPerBucket_Threaded(Empty, BucketIdx, true, false, false), BucketIdx, DefaultAsUInt16);
 
 				// Always create this if QueryOnlyObjects are isolated to ensure all consecutive indices are created (within buckets)
 				if (AccelerationStructureSplitStaticAndDynamic == 1 || AccelerationStructureIsolateQueryOnlyObjects == 1) 
 				{
-					constexpr auto DynamicAsUInt16 = static_cast<std::underlying_type_t<ESpatialAccelerationCollectionBucketInnerIndex>>(ESpatialAccelerationCollectionBucketInnerIndex::Dynamic);
+					constexpr uint16 DynamicAsUInt16 = static_cast<std::underlying_type_t<ESpatialAccelerationCollectionBucketInnerIndex>>(ESpatialAccelerationCollectionBucketInnerIndex::Dynamic);
 					// Non static bodies
 					Collection->AddSubstructure(CreateAccelerationPerBucket_Threaded(Empty, BucketIdx, true, IsDynamicTree(FSpatialAccelerationIdx{BucketIdx, DynamicAsUInt16 }), true), BucketIdx, DynamicAsUInt16);
 				}
@@ -116,7 +116,7 @@ namespace Chaos
 				Collection->AddSubstructure(CreateAccelerationPerBucket_Threaded(Empty, BucketIdx, true, false, false), BucketIdx, DefaultQueryOnlyAsUInt16);
 				if (AccelerationStructureSplitStaticAndDynamic == 1)
 				{
-					constexpr auto DynamicQueryOnlyAsUInt16 = static_cast<std::underlying_type_t<ESpatialAccelerationCollectionBucketInnerIndex>>(ESpatialAccelerationCollectionBucketInnerIndex::DynamicQueryOnly);
+					constexpr uint16 DynamicQueryOnlyAsUInt16 = static_cast<std::underlying_type_t<ESpatialAccelerationCollectionBucketInnerIndex>>(ESpatialAccelerationCollectionBucketInnerIndex::DynamicQueryOnly);
 
 					// Non static bodies
 					Collection->AddSubstructure(CreateAccelerationPerBucket_Threaded(Empty, BucketIdx, true, IsDynamicTree(FSpatialAccelerationIdx{ BucketIdx,  DynamicQueryOnlyAsUInt16 }), false), BucketIdx, DynamicQueryOnlyAsUInt16);
diff --git a/Engine/Source/Runtime/Experimental/Chaos/Public/Chaos/SpatialAccelerationCollection.h b/Engine/Source/Runtime/Experimental/Chaos/Public/Chaos/SpatialAccelerationCollection.h
--- a/Engine/Source/Runtime/Experimental/Chaos/Public/Chaos/SpatialAccelerationCollection.h
+++ b/Engine/Source/Runtime/Experimental/Chaos/Public/Chaos/SpatialAccelerationCollection.h
@@ -548,8 +548,8 @@ public:
 			int BucketInnerIdx = 0;
 			for (const BucketEntryType& BucketEntry : Bucket.Objects)
 			{
-				constexpr auto DefaultQueryOnlyAsUInt16 = static_cast<std::underlying_type_t<ESpatialAccelerationCollectionBucketInnerIndex>>(ESpatialAccelerationCollectionBucketInnerIndex::DefaultQueryOnly);
-				constexpr auto DynamicQueryOnlyAsUInt16 = static_cast<std::underlying_type_t<ESpatialAccelerationCollectionBucketInnerIndex>>(ESpatialAccelerationCollectionBucketInnerIndex::DynamicQueryOnly);			
+				constexpr uint16 DefaultQueryOnlyAsUInt16 = static_cast<std::underlying_type_t<ESpatialAccelerationCollectionBucketInnerIndex>>(ESpatialAccelerationCollectionBucketInnerIndex::DefaultQueryOnly);
+				constexpr uint16 DynamicQueryOnlyAsUInt16 = static_cast<std::underlying_type_t<ESpatialAccelerationCollectionBucketInnerIndex>>(ESpatialAccelerationCollectionBucketInnerIndex::DynamicQueryOnly);			
 				
 				if (BucketInnerIdx != DefaultQueryOnlyAsUInt16 && BucketInnerIdx != DynamicQueryOnlyAsUInt16)
 				{
