Fix Reset() not getting called when IAD cache fails to verify at startup.
Add IadCasVerification analytic.

#Synced-CL 49267577

diff --git a/Engine/Source/Runtime/Experimental/IoStore/OnDemand/Private/OnDemandInstallCache.cpp b/Engine/Source/Runtime/Experimental/IoStore/OnDemand/Private/OnDemandInstallCache.cpp
--- a/Engine/Source/Runtime/Experimental/IoStore/OnDemand/Private/OnDemandInstallCache.cpp
+++ b/Engine/Source/Runtime/Experimental/IoStore/OnDemand/Private/OnDemandInstallCache.cpp
@@ -714,7 +714,7 @@ struct FCas
 	uint64						GetTotalBlockSize() const;
 	FBlockFileSizes				FindAllBlockFiles() const;
 	void						Compact();
-	FResult						Verify(TArray<FCasAddr>& OutAddrs);
+	FResult						Verify(TArray<FCasAddr>& OutRemovedChunksAddrs, TMap<FCasBlockId, uint64>& OutOverBudgetBlocks);
 	void						LoadSnapshot(FCasSnapshot&& Snapshot);
 	void						ApplyJournalEntry(const FCasJournal::FEntry& JournalEntry);
 	FLastAccess					ConsumeLastAcccess()
@@ -1300,7 +1300,7 @@ void FCas::Compact()
 	LastAccess.Compact();
 }
 
-FResult FCas::Verify(TArray<FCasAddr>& OutAddrs)
+FResult FCas::Verify(TArray<FCasAddr>& OutRemovedChunksAddrs, TMap<FCasBlockId, uint64>& OutOverBudgetBlocks)
 {
 	using namespace UE::UnifiedError::IoStoreOnDemand;
 
@@ -1361,6 +1361,11 @@ FResult FCas::Verify(TArray<FCasAddr>& OutAddrs)
 				}
 			}
 
+			if (*FileSize > FCas::GetMaxBlockSize())
+			{
+				OutOverBudgetBlocks.Add(BlockId, *FileSize);
+			}
+
 			if (*FileSize <= std::numeric_limits<uint32>::max())
 			{
 				BlockState.Size = uint32(*FileSize);
@@ -1411,7 +1416,7 @@ FResult FCas::Verify(TArray<FCasAddr>& OutAddrs)
 			FString ErrorMessage	= FString::Printf(TEXT("Missing CAS block '%s'"), *Filename);
 			MissingReferencedBlocks.Add(MoveTemp(Filename));
 
-			OutAddrs.Add(It->Key);
+			OutRemovedChunksAddrs.Add(It->Key);
 			It.RemoveCurrent();
 
 			Result = MakeCasError<void>(ECasErrorCode::VerifyFailed, EIoErrorCode::NotFound, MoveTemp(ErrorMessage));
@@ -2158,7 +2163,7 @@ FOnDemandInstallCache::FOnDemandInstallCache(const FOnDemandInstallCacheConfig&
 		UE_LOGFMT(LogIoStoreOnDemand, Error, "Resetting install cache, reason: {Error}", InitResult.GetError());
 
 		FOnDemandInstallCacheStats::OnStartupError(InitResult);
-		InitResult = InitialVerify();
+		InitResult = Reset();
 	}
 
 	if (InitResult.HasValue())
@@ -3039,10 +3044,22 @@ FResult FOnDemandInstallCache::InitialVerify()
 	// Verify the blocks on disk with the current state of the CAS
 	{
 		TArray<FCasAddr> RemovedChunks;
-		if (FResult Verify = Cas.Verify(RemovedChunks); Verify.HasError())
+		TMap<FCasBlockId, uint64> OverBudgetBlocks;
+		FResult Verify = Cas.Verify(RemovedChunks, OverBudgetBlocks);
+
+		uint64 TotalBlockBytesOverBudget = 0;
+		for (const TPair<FCasBlockId, uint64>& Kv : OverBudgetBlocks)
 		{
-			FOnDemandInstallCacheStats::OnCasVerificationError(RemovedChunks.Num());
+			if (ensure(Kv.Value > FCas::GetMaxBlockSize()))
+			{
+				TotalBlockBytesOverBudget += (Kv.Value - FCas::GetMaxBlockSize());
+			}
+		}
 
+		FOnDemandInstallCacheStats::OnCasVerification(Verify, RemovedChunks.Num(), TotalBlockBytesOverBudget);
+
+		if (Verify.HasError())
+		{
 			// Remove all entries that doesn't have a valid cache block 
 			FCasJournal::FTransaction Transaction = FCasJournal::Begin(GetJournalFilename());
 			for (const FCasAddr& Addr : RemovedChunks)
diff --git a/Engine/Source/Runtime/Experimental/IoStore/OnDemand/Private/Statistics.cpp b/Engine/Source/Runtime/Experimental/IoStore/OnDemand/Private/Statistics.cpp
--- a/Engine/Source/Runtime/Experimental/IoStore/OnDemand/Private/Statistics.cpp
+++ b/Engine/Source/Runtime/Experimental/IoStore/OnDemand/Private/Statistics.cpp
@@ -1655,10 +1655,16 @@ void FOnDemandInstallCacheStats::OnJournalCommit(const FResult& Result, int64 By
 	InstallCacheAnalytics.JournalCommitCount++;
 }
 
-void FOnDemandInstallCacheStats::OnCasVerificationError(int32 RemoveChunks)
+void FOnDemandInstallCacheStats::OnCasVerification(const FResult& Result, int32 RemovedChunkCount, uint64 BlockBytesOverBudget)
 {
+	FOnDemandImmediateAnalytic Event = MakeAnalyticsEventFromResult(TEXT("IadCasVerification"), Result);
+	AppendAnalyticsEventAttributeArray(Event.AnalyticsArray,
+		TEXT("RemovedChunkCount"), RemovedChunkCount,
+		TEXT("BlockBytesOverBudget"), BlockBytesOverBudget);
+	SendImmediateAnalytic(MoveTemp(Event));
+
 	UE::TUniqueLock Lock(InstallCacheAnalytics);
-	InstallCacheAnalytics.VerificationRemovedBlockCount += RemoveChunks;
+	InstallCacheAnalytics.VerificationRemovedBlockCount += RemovedChunkCount;
 }
 
 void FOnDemandInstallCacheStats::OnPurge(
diff --git a/Engine/Source/Runtime/Experimental/IoStore/OnDemand/Private/Statistics.h b/Engine/Source/Runtime/Experimental/IoStore/OnDemand/Private/Statistics.h
--- a/Engine/Source/Runtime/Experimental/IoStore/OnDemand/Private/Statistics.h
+++ b/Engine/Source/Runtime/Experimental/IoStore/OnDemand/Private/Statistics.h
@@ -168,7 +168,7 @@ public:
 	static void OnFlush(const FResult& Result, int64 ByteCount) IAS_STATISTICS_IMPL()
 	static void OnJournalCommit(EIoErrorCode ErrorCode, int64 ByteCount) IAS_STATISTICS_IMPL()
 	static void OnJournalCommit(const FResult& Result, int64 ByteCount) IAS_STATISTICS_IMPL()
-	static void OnCasVerificationError(int32 RemoveChunks) IAS_STATISTICS_IMPL()
+	static void OnCasVerification(const FResult& Result, int32 RemovedChunkCount, uint64 BlockBytesOverBudget) IAS_STATISTICS_IMPL()
 	static void OnPurge(
 		const FResult& Result,	
 		uint64 MaxCacheSize,
