[TraceBasedDebuggers] used IsRunning* helpers to handle cases where the editor is launched with -server, -client, -game

#Synced-CL 49372183

diff --git a/Engine/Source/Runtime/Experimental/TraceBasedDebuggers/Private/RemoteSessionsManager.cpp b/Engine/Source/Runtime/Experimental/TraceBasedDebuggers/Private/RemoteSessionsManager.cpp
--- a/Engine/Source/Runtime/Experimental/TraceBasedDebuggers/Private/RemoteSessionsManager.cpp
+++ b/Engine/Source/Runtime/Experimental/TraceBasedDebuggers/Private/RemoteSessionsManager.cpp
@@ -1,4 +1,4 @@
-// Copyright Epic Games, Inc. All Rights Reserved.
+﻿// Copyright Epic Games, Inc. All Rights Reserved.
 
 #include "RemoteSessionsManager.h"
 
@@ -403,7 +403,9 @@ void FRemoteSessionsManager::SendPong(const FSessionPing& InMessage)
 	{
 		PongMessage->InstanceId = FApp::GetInstanceId();
 		PongMessage->SessionId = FApp::GetSessionId();
-		PongMessage->BuildTargetType = static_cast<uint8>(FApp::GetBuildTargetType());
+
+		const EBuildTargetType TargetType = GetBuildTargetType();
+		PongMessage->BuildTargetType = static_cast<uint8>(TargetType);
 
 		if (InMessage.ControllerInstanceId == PongMessage->InstanceId)
 		{
@@ -415,7 +417,7 @@ void FRemoteSessionsManager::SendPong(const FSessionPing& InMessage)
 			PongMessage->SessionName = AppSessionName == TEXT("None") || AppSessionName.IsEmpty()
 				? FString::Format(TEXT("{0} {1} {2}"),
 						{FApp::GetProjectName()
-						, FString(LexToString(FApp::GetBuildTargetType()))
+						, FString(LexToString(TargetType))
 						, FString::FromInt(FPlatformProcess::GetCurrentProcessId())})
 				: AppSessionName;
 		}
diff --git a/Engine/Source/Runtime/Experimental/TraceBasedDebuggers/Private/TraceBasedDebuggerRuntime.cpp b/Engine/Source/Runtime/Experimental/TraceBasedDebuggers/Private/TraceBasedDebuggerRuntime.cpp
--- a/Engine/Source/Runtime/Experimental/TraceBasedDebuggers/Private/TraceBasedDebuggerRuntime.cpp
+++ b/Engine/Source/Runtime/Experimental/TraceBasedDebuggers/Private/TraceBasedDebuggerRuntime.cpp
@@ -385,7 +385,7 @@ void FRuntimeModule::GenerateRecordingFileName(FString& OutFileName) const
 	{
 		DebuggerName,
 		FString(FApp::GetProjectName())
-		, LexToString(FApp::GetBuildTargetType())
+		, LexToString(GetBuildTargetType())
 		, FDateTime::Now().ToString(TEXT("%Y%m%d_%H%M%S%s"))
 	};
 
diff --git a/Engine/Source/Runtime/Experimental/TraceBasedDebuggers/Private/TraceBasedDebuggerTypes.cpp b/Engine/Source/Runtime/Experimental/TraceBasedDebuggers/Private/TraceBasedDebuggerTypes.cpp
--- a/Engine/Source/Runtime/Experimental/TraceBasedDebuggers/Private/TraceBasedDebuggerTypes.cpp
+++ b/Engine/Source/Runtime/Experimental/TraceBasedDebuggers/Private/TraceBasedDebuggerTypes.cpp
@@ -2,4 +2,29 @@
 
 #include "TraceBasedDebuggerTypes.h"
 
+#include "Misc/App.h"
+
 DEFINE_LOG_CATEGORY(LogTraceBasedDebuggers)
+
+EBuildTargetType UE::TraceBasedDebuggers::GetBuildTargetType()
+{
+	// Using IsRunning* helpers to handle cases where the editor is launched with -server, -client, -game
+	EBuildTargetType TargetType;
+	if (IsRunningDedicatedServer())
+	{
+		TargetType = EBuildTargetType::Server;
+	}
+	else if (IsRunningClientOnly())
+	{
+		TargetType = EBuildTargetType::Client;
+	}
+	else if (IsRunningGame())
+	{
+		TargetType = EBuildTargetType::Game;
+	}
+	else
+	{
+		TargetType = FApp::GetBuildTargetType();
+	}
+	return TargetType;
+}
diff --git a/Engine/Source/Runtime/Experimental/TraceBasedDebuggers/Public/SessionInfo.h b/Engine/Source/Runtime/Experimental/TraceBasedDebuggers/Public/SessionInfo.h
--- a/Engine/Source/Runtime/Experimental/TraceBasedDebuggers/Public/SessionInfo.h
+++ b/Engine/Source/Runtime/Experimental/TraceBasedDebuggers/Public/SessionInfo.h
@@ -108,9 +108,9 @@ public:
 	template <typename DataType>
 	void SetDebuggerData(DataType&& NewSessionData)
 	{
-		if (DataType* ChaosData = GetDebuggerData<DataType>())
+		if (DataType* DebuggerData = GetDebuggerData<DataType>())
 		{
-			*ChaosData = MoveTemp(NewSessionData);
+			*DebuggerData = MoveTemp(NewSessionData);
 		}
 		else
 		{
diff --git a/Engine/Source/Runtime/Experimental/TraceBasedDebuggers/Public/TraceBasedDebuggerTypes.h b/Engine/Source/Runtime/Experimental/TraceBasedDebuggers/Public/TraceBasedDebuggerTypes.h
--- a/Engine/Source/Runtime/Experimental/TraceBasedDebuggers/Public/TraceBasedDebuggerTypes.h
+++ b/Engine/Source/Runtime/Experimental/TraceBasedDebuggers/Public/TraceBasedDebuggerTypes.h
@@ -6,6 +6,13 @@
 
 #define UE_API TRACEBASEDDEBUGGERS_API
 
+enum class EBuildTargetType : uint8;
+
 UE_API DECLARE_LOG_CATEGORY_EXTERN(LogTraceBasedDebuggers, Log, All);
 
+namespace UE::TraceBasedDebuggers
+{
+extern UE_API EBuildTargetType GetBuildTargetType();
+}
+
 #undef UE_API
\ No newline at end of file
