Possible fix for IAD requests not flushing out in time at shutdown - Fix IAD http requests not being canceled at shutdown.

#Synced-CL 49267577

diff --git a/Engine/Source/Runtime/Experimental/IoStore/OnDemand/Private/OnDemandContentInstaller.cpp b/Engine/Source/Runtime/Experimental/IoStore/OnDemand/Private/OnDemandContentInstaller.cpp
--- a/Engine/Source/Runtime/Experimental/IoStore/OnDemand/Private/OnDemandContentInstaller.cpp
+++ b/Engine/Source/Runtime/Experimental/IoStore/OnDemand/Private/OnDemandContentInstaller.cpp
@@ -1786,6 +1786,30 @@ void FOnDemandContentInstaller::Shutdown()
 		}
 	}
 
+	// Need special handling to cancel install requests
+	InstallerPipe.Launch(TEXT("OnDemandContentInstallerShutdown"),
+		[this]
+		{
+			TUniqueLock Lock(Mutex);
+			for (FRequest* Request : RunningRequests)
+			{
+				if (Request->IsInstall() == false)
+				{
+					continue;
+				}
+
+				FRequest::FInstall& Install = Request->AsInstall();
+
+				if (Install.bHttpRequestsIssued)
+				{
+					int32 NumCancelled = 0;
+					TryCancelHttpRequestsForInstallRequest(Install, NumCancelled);
+					UE_LOG(LogIoStoreOnDemand, Verbose, TEXT("Cancelled %d HTTP request(s) due to install request cancellation"), NumCancelled);
+				}
+			}
+		}
+	);
+
 	// Wait for the current request(s) to finish
 	for (;;)
 	{
@@ -1801,6 +1825,7 @@ void FOnDemandContentInstaller::Shutdown()
 			TUniqueLock Lock(Mutex);
 			if (RunningRequests.IsEmpty())
 			{
+				UE_LOG(LogIoStoreOnDemand, Display, TEXT("Content installer shutdown flushed requests after %.2lf"), WaitTimeSeconds);
 				break;
 			}
 		}
diff --git a/Engine/Source/Runtime/Experimental/IoStore/OnDemand/Private/OnDemandContentInstaller.h b/Engine/Source/Runtime/Experimental/IoStore/OnDemand/Private/OnDemandContentInstaller.h
--- a/Engine/Source/Runtime/Experimental/IoStore/OnDemand/Private/OnDemandContentInstaller.h
+++ b/Engine/Source/Runtime/Experimental/IoStore/OnDemand/Private/OnDemandContentInstaller.h
@@ -180,12 +180,11 @@ class FOnDemandContentInstaller
 			Variant.Emplace<FFlushLastAccess>(MoveTemp(Args), MoveTemp(OnCompleted));
 		}
 
-		bool				IsCancelled() const			{ return bCancelled; }
+		bool				IsCancelled() const			{ return bCancelled.load(std::memory_order_relaxed); }
 		bool				TryCancel()
 		{ 
-			bool bPrev = bCancelled;
-			bCancelled = true;
-			return bPrev == false;
+			bool bExpected = false; 
+			return bCancelled.compare_exchange_strong(bExpected, true);
 		}
 		bool				IsOk() const				{ return !IsCancelled() && !Result.HasError(); }
 		FString				GetErrorString() const		{ return IsOk() ? TEXT("Ok") : IsCancelled() ? TEXT("Cancelled") : LexToString(Result.GetError()); }
@@ -216,7 +215,7 @@ class FOnDemandContentInstaller
 		int32						Priority = 0;
 		uint64						StartTimeCycles = FPlatformTime::Cycles64();
 		FResult						Result = MakeValue();
-		bool						bCancelled = false;
+		std::atomic_bool			bCancelled{ false };
 		FRequestVariant				Variant;
 	};
 
