Chaos Deformable:
-supported cross asset constraints
-reorder constraint processing in deformable evolution: now constraint rules-cross asset rules-poststep(velocity update)-postsubstep
-added FChaosDeformableCreateSpringConstraintNode for building spring constraints beween assets
-supported spring constraints between two sets of particles (required for cross asset cases)

#Synced-CL 49269668

diff --git a/Engine/Source/Runtime/Experimental/Chaos/Public/Chaos/XPBDWeakConstraints.h b/Engine/Source/Runtime/Experimental/Chaos/Public/Chaos/XPBDWeakConstraints.h
--- a/Engine/Source/Runtime/Experimental/Chaos/Public/Chaos/XPBDWeakConstraints.h
+++ b/Engine/Source/Runtime/Experimental/Chaos/Public/Chaos/XPBDWeakConstraints.h
@@ -73,6 +73,24 @@ namespace Chaos::Softs
 			LambdaArray.Init((FSolverReal)0., Indices.Num() * 3);
 		}
 
+		FXPBDWeightedSpringConstraints(
+			const FSolverParticles& InParticles,
+			const FSolverParticles& InSecondParticles,
+			const TArray<TArray<int32>>& InIndices,
+			const TArray<TArray<FSolverReal>>& InWeights,
+			const TArray<FSolverReal>& InStiffness,
+			const TArray<TArray<int32>>& InSecondIndices,
+			const TArray<TArray<FSolverReal>>& InSecondWeights,
+			const FDeformableXPBDWeightedSpringConstraintParams& InParams
+		)
+			: Indices(InIndices), Weights(InWeights), SecondIndices(InSecondIndices), SecondWeights(InSecondWeights), Stiffness(InStiffness), DebugDrawParams(InParams)
+		{
+			ensureMsgf(Indices.Num() == SecondIndices.Num(), TEXT("Input Double Bindings have wrong size"));
+
+			//TODO: write InitColor for both sides of particles
+			//InitColor(InParticles);
+			LambdaArray.Init((FSolverReal)0., Indices.Num() * 3);
+		}
 
 		virtual ~FXPBDWeightedSpringConstraints() = default;
 
@@ -118,6 +136,16 @@ namespace Chaos::Softs
 			ApplyInParallel(Particles, Dt);
 		}
 
+		void ApplyInSerial(FSolverParticlesRange& Particles, FSolverParticlesRange& SecondParticles, const FSolverReal Dt) const
+		{
+			TRACE_CPUPROFILER_EVENT_SCOPE(STAT_ChaosXPBDWeakConstraintApply);
+
+			for (int32 ConstraintIdx = 0; ConstraintIdx < Indices.Num(); ++ConstraintIdx)
+			{
+				ApplySingleConstraintWithSelfTarget(Particles, SecondParticles, Dt, ConstraintIdx);
+			}
+		}
+
 		const TArray<TArray<int32>>& GetIndices()
 		{
 			return Indices;
@@ -248,6 +276,59 @@ namespace Chaos::Softs
 			}
 		}
 
+		void ApplySingleConstraintWithSelfTarget(FSolverParticlesRange& Particles, FSolverParticlesRange& SecondParticles, const FSolverReal Dt, const int32 ConstraintIndex, const FSolverReal Tol = UE_KINDA_SMALL_NUMBER) const
+		{
+			if (Stiffness[ConstraintIndex] < Tol)
+			{
+				return;
+			}
+
+			ensure(SecondIndices.Num() > 0);
+			FSolverReal AlphaTilde = FSolverReal(2) / (Stiffness[ConstraintIndex] * Dt * Dt);
+			if (Stiffness[ConstraintIndex] > 1e14)
+			{
+				AlphaTilde = FSolverReal(0);
+			}
+
+			TVec3<FSolverReal> SpringEdge((FSolverReal)0.);
+			for (int32 j = 0; j < Indices[ConstraintIndex].Num(); j++)
+			{
+				for (int32 alpha = 0; alpha < 3; alpha++)
+				{
+					SpringEdge[alpha] += Weights[ConstraintIndex][j] * Particles.P(Indices[ConstraintIndex][j])[alpha];
+				}
+			}
+			for (int32 j = 0; j < SecondIndices[ConstraintIndex].Num(); j++)
+			{
+				SpringEdge -= SecondWeights[ConstraintIndex][j] * SecondParticles.P(SecondIndices[ConstraintIndex][j]);
+			}
+
+			FSolverReal Denom = AlphaTilde;
+			for (int32 j = 0; j < Indices[ConstraintIndex].Num(); j++)
+			{
+				Denom += Weights[ConstraintIndex][j] * Weights[ConstraintIndex][j] * Particles.InvM(Indices[ConstraintIndex][j]);
+			}
+			for (int32 j = 0; j < SecondIndices[ConstraintIndex].Num(); j++)
+			{
+				Denom += SecondWeights[ConstraintIndex][j] * SecondWeights[ConstraintIndex][j] * SecondParticles.InvM(SecondIndices[ConstraintIndex][j]);
+			}
+			for (int32 Beta = 0; Beta < 3; Beta++)
+			{
+				FSolverReal Cj = SpringEdge[Beta];
+				FSolverReal DLambda = -Cj - AlphaTilde * LambdaArray[ConstraintIndex * 3 + Beta];
+				DLambda /= Denom;
+				LambdaArray[ConstraintIndex * 3 + Beta] += DLambda;
+				for (int32 j = 0; j < Indices[ConstraintIndex].Num(); j++)
+				{
+					Particles.P(Indices[ConstraintIndex][j])[Beta] += DLambda * Weights[ConstraintIndex][j] * Particles.InvM(Indices[ConstraintIndex][j]);
+				}
+				for (int32 j = 0; j < SecondIndices[ConstraintIndex].Num(); j++)
+				{
+					SecondParticles.P(SecondIndices[ConstraintIndex][j])[Beta] -= DLambda * SecondWeights[ConstraintIndex][j] * SecondParticles.InvM(SecondIndices[ConstraintIndex][j]);
+				}
+			}
+		}
+
 		void VisualizeAllBindings(const FSolverParticlesRange& InParticles, const FSolverReal Dt) const
 		{
 #if WITH_EDITOR
