[NetPhysics] Option to disable physics replication per-actor at runtime via EPhysicsReplicationMode::None

#Synced-CL 49558020

diff --git a/Engine/Source/Runtime/Engine/Classes/Engine/EngineTypes.h b/Engine/Source/Runtime/Engine/Classes/Engine/EngineTypes.h
--- a/Engine/Source/Runtime/Engine/Classes/Engine/EngineTypes.h
+++ b/Engine/Source/Runtime/Engine/Classes/Engine/EngineTypes.h
@@ -3386,18 +3386,23 @@ UENUM(BlueprintType)
 enum class EPhysicsReplicationMode : uint8
 {
 	/** Default physics replication.*/
-	Default UMETA(DisplayName = "Default"),
+	Default = 0 UMETA(DisplayName = "Default"),
 
-	/** Work In Progress. Physics replication performing velocity interpolation.
-	* Recommendation: Set on actors with a local role of ENetRole::ROLE_SimulatedProxy.
-	* Designed to handle local predictive interactions with other actors, especially actors of the local role ENetRole::ROLE_AutonomousProxy. */
-	PredictiveInterpolation UMETA(DisplayName = "Predictive Interpolation (WIP)"),
+	/** Physics replication performing velocity interpolation.
+	* Note: Not forward predicted which makes this not viable for ENetRole::ROLE_AutonomousProxy if input latency is of concern.
+	* Recommendation: All physics objects that interact should use the same replication mode. Use alongside a kinematic pawn (not physics simulated).
+	* Designed to handle local predictive interactions with other actors, especially actors of the local role ENetRole::ROLE_AutonomousProxy.
+	* Enable Project Settings > Physics Prediction for best result */
+	PredictiveInterpolation = 1 UMETA(DisplayName = "Predictive Interpolation"),
 
 	/** Work In Progress. Forward predicted replication by simulating physics and correcting errors through resimulating physics from a correct state in the past.
-	* Recommendation: Set on actors with a local role of ENetRole::ROLE_AutonomousProxy.
-	* Can be used for both ROLE_AutonomousProxy and ROLE_SimulatedProxy, though not recommended for ROLE_SimulatedProxy due to CPU performance.
+	* Recommendation: All physics objects that interact should use the same replication mode. Use on physics simulated pawns along with any physics objects that will interact with such pawns.
 	* Available when Project Settings > Physics Prediction is enabled. */
-	Resimulation UMETA(DisplayName = "Resimulation (WIP)"),
+	Resimulation = 2 UMETA(DisplayName = "Resimulation (WIP)"),
+
+	/** Don't apply physics replication
+	* Note: Can be used to temporarily disable physics replication for specific actors. */
+	None = 255 UMETA(DisplayName = "None"),
 };
 
 /** Specifies which player index will pass input to this actor/component */
diff --git a/Engine/Source/Runtime/Engine/Private/ActorReplication.cpp b/Engine/Source/Runtime/Engine/Private/ActorReplication.cpp
--- a/Engine/Source/Runtime/Engine/Private/ActorReplication.cpp
+++ b/Engine/Source/Runtime/Engine/Private/ActorReplication.cpp
@@ -203,9 +203,12 @@ void AActor::OnRep_ReplicatedMovement()
 					DebugColor = FColor::Red;
 					break;
 				case EPhysicsReplicationMode::Default:
-				default:
 					DebugColor = FColor::Cyan;
 					break;
+				case EPhysicsReplicationMode::None:
+				default:
+					DebugColor = FColor::Black;
+					break;
 				}
 			}
 			DrawDebugCapsule(GetWorld(), LocalRepMovement.Location, FMath::Max(GetSimpleCollisionHalfHeight(), 25.0f), FMath::Max(GetSimpleCollisionRadius(), 25.0f), LocalRepMovement.Rotation.Quaternion(), DebugColor, false, 1.f);
diff --git a/Engine/Source/Runtime/Engine/Private/PhysicsEngine/PhysicsReplication.cpp b/Engine/Source/Runtime/Engine/Private/PhysicsEngine/PhysicsReplication.cpp
--- a/Engine/Source/Runtime/Engine/Private/PhysicsEngine/PhysicsReplication.cpp
+++ b/Engine/Source/Runtime/Engine/Private/PhysicsEngine/PhysicsReplication.cpp
@@ -908,6 +908,20 @@ void FPhysicsReplicationAsync::OnPhysicsObjectUnregistered_Internal(Chaos::FCons
 
 	// Only clear Settings when PhysicsObject unregister (not when it stops replicating, hence why it's not baked into RemoveObjectFromReplication())
 	ObjectToSettings.Remove(PhysicsObject);
+
+
+	Chaos::FReadPhysicsObjectInterface_Internal Interface = Chaos::FPhysicsObjectInternalInterface::GetRead();
+	if (Chaos::FGeometryParticleHandle* Handle = Interface.GetParticle(PhysicsObject))
+	{
+		if (Chaos::FPhysicsSolverBase* SolverBase = GetSolver())
+		{
+			if (Chaos::FRewindData* RewindData = SolverBase->GetRewindData())
+			{
+				// Stop forcing resim as follower 
+				RewindData->ForceResimAsFollower(Handle, false);
+			}
+		}
+	}
 }
 
 void FPhysicsReplicationAsync::RegisterSettings(Chaos::FConstPhysicsObjectHandle PhysicsObject, TWeakPtr<const FNetworkPhysicsSettingsData> InSettings)
@@ -1095,7 +1109,7 @@ void FPhysicsReplicationAsync::OnPostSolve_Internal()
 		for (auto Itr = ObjectToTarget.CreateIterator(); Itr; ++Itr)
 		{
 			FReplicatedPhysicsTargetAsync& Target = Itr.Value();
-			if (Target.RepMode != EPhysicsReplicationMode::Resimulation)
+			if (Target.RepMode == EPhysicsReplicationMode::Default || Target.RepMode == EPhysicsReplicationMode::PredictiveInterpolation)
 			{
 				Chaos::FConstPhysicsObjectHandle& PhysicsObject = Itr.Key();
 				if (Chaos::FGeometryParticleHandle* GeometryParticle = Interface.GetParticle(PhysicsObject))
@@ -1198,6 +1212,8 @@ void FPhysicsReplicationAsync::UpdateRewindDataTarget(const FPhysicsRepAsyncInpu
 		RewindData->SetTargetStateAtFrame(*Handle, LocalFrame, Chaos::FFrameAndPhase::EParticleHistoryPhase::PrePushData,
 			Input.TargetState.Position, Input.TargetState.Quaternion,
 			Input.TargetState.LinVel, FMath::DegreesToRadians(Input.TargetState.AngVel), (Input.TargetState.Flags & ERigidBodyFlags::Sleeping));
+
+		RewindData->ForceResimAsFollower(Handle, /*bForceResim*/ Input.RepMode == EPhysicsReplicationMode::None);
 	}
 }
 
@@ -1483,6 +1499,10 @@ void FPhysicsReplicationAsync::ApplyTargetStatesAsync(const float DeltaSeconds)
 			case EPhysicsReplicationMode::Resimulation:
 				bRemoveItr = ResimulationReplication(RigidHandle, Target, DeltaSeconds);
 				break;
+
+			case EPhysicsReplicationMode::None:
+				bRemoveItr = true;
+				break;
 		}
 		Target.TickCount++;
 			
@@ -2643,9 +2663,12 @@ void FPhysicsReplicationAsync::DebugDrawReplicationMode(const FPhysicsRepAsyncIn
 				DebugColor = FColor::Red;
 				break;
 			case EPhysicsReplicationMode::Default:
-			default:
 				DebugColor = FColor::Cyan;
 				break;
+			case EPhysicsReplicationMode::None:
+			default:
+				DebugColor = FColor::Black;
+				break;
 			}
 		}
 	}
diff --git a/Engine/Source/Runtime/Experimental/Chaos/Private/PBDRigidsSolver.cpp b/Engine/Source/Runtime/Experimental/Chaos/Private/PBDRigidsSolver.cpp
--- a/Engine/Source/Runtime/Experimental/Chaos/Private/PBDRigidsSolver.cpp
+++ b/Engine/Source/Runtime/Experimental/Chaos/Private/PBDRigidsSolver.cpp
@@ -937,7 +937,6 @@ namespace Chaos
 				MRewindData->RemoveObject(Handle);
 			}
 			
-			Proxy->SetHandle(nullptr);
 			const int32 OffsetForRewind = (bResimDelayDestroyObject && MRewindData) ? MRewindData->Capacity() : 0;
 			PendingDestroyPhysicsProxy.Add(FPendingDestroyInfo{Proxy, GetCurrentFrame() + OffsetForRewind, Handle, UniqueIdx});
 			
@@ -1351,8 +1350,8 @@ namespace Chaos
 					// Use the handle to destroy the particle data
 					GetEvolution()->DestroyParticle(Info.Handle);
 				}
+				Info.Proxy->SetHandle(nullptr);
 
-				ensure(Info.Proxy->GetHandle_LowLevel() == nullptr);	//should have already cleared this out
 				delete Info.Proxy;
 				PendingDestroyPhysicsProxy.RemoveAtSwap(Idx, EAllowShrinking::No);
 			}
diff --git a/Engine/Source/Runtime/Experimental/Chaos/Private/RewindData.cpp b/Engine/Source/Runtime/Experimental/Chaos/Private/RewindData.cpp
--- a/Engine/Source/Runtime/Experimental/Chaos/Private/RewindData.cpp
+++ b/Engine/Source/Runtime/Experimental/Chaos/Private/RewindData.cpp
@@ -878,15 +878,15 @@ void FRewindData::RewindToFrame(int32 RewindFrame)
 	// Determine resim type
 	if (ResimConsoleVars::bIsBubbleResimEnabled)
 	{
-		// Bubble resim: Everything in the resim bubble fully resimulates, nothing resims as follower (as least for now)
+		// Bubble resim: Everything in the resim bubble fully resimulates unless marked for forced ResimAsFollower
 		QUICK_SCOPE_CYCLE_COUNTER(ChaosRewindAndResim_MarkAllParticlesInBubbleFullResim);
 		for (FDirtyParticleInfo& DirtyParticleInfo : DirtyParticles)
 		{
-			DirtyParticleInfo.bResimAsFollower = false;
+			DirtyParticleInfo.bResimAsFollower = ParticlesToForceResimAsFollower.Contains(DirtyParticleInfo.GetObjectPtr());
 		}
 		for (FGeometryParticleHandle* GeometryParticle : ParticlesOverlappingBubble)
 		{
-			GeometryParticle->SetResimType(EResimType::FullResim);
+			GeometryParticle->SetResimType(ParticlesToForceResimAsFollower.Contains(GeometryParticle) ? EResimType::ResimAsFollower : EResimType::FullResim);
 		}
 	}
 	else
@@ -895,6 +895,12 @@ void FRewindData::RewindToFrame(int32 RewindFrame)
 		{
 			FGeometryParticleHandle* GeometryParticle = DirtyParticleInfo.GetObjectPtr();
 			DirtyParticleInfo.bResimAsFollower = GeometryParticle->ObjectState() != EObjectStateType::Dynamic && GeometryParticle->ObjectState() != EObjectStateType::Sleeping;
+
+			if (!DirtyParticleInfo.bResimAsFollower)
+			{
+				// Force resim as follower
+				DirtyParticleInfo.bResimAsFollower = ParticlesToForceResimAsFollower.Contains(GeometryParticle);
+			}
 		}
 	}
 
@@ -2560,6 +2566,18 @@ void FRewindData::SetTargetStateAtFrame(FGeometryParticleHandle& Handle, const i
 		Position, Quaternion, LinVelocity, AngVelocity, bShouldSleep);
 }
 
+void FRewindData::ForceResimAsFollower(FGeometryParticleHandle* Handle, bool bForceResim)
+{
+	if (bForceResim)
+	{
+		ParticlesToForceResimAsFollower.Add(Handle);
+	}
+	else
+	{
+		ParticlesToForceResimAsFollower.Remove(Handle);
+	}
+}
+
 int32 FRewindData::GetLatestReceivedTargetStateFrame(Chaos::FGeometryParticleHandle& Handle) const
 {
 	if (const FDirtyParticleInfo* DirtyParticleInfo = FindDirtyObj(Handle))
diff --git a/Engine/Source/Runtime/Experimental/Chaos/Public/RewindData.h b/Engine/Source/Runtime/Experimental/Chaos/Public/RewindData.h
--- a/Engine/Source/Runtime/Experimental/Chaos/Public/RewindData.h
+++ b/Engine/Source/Runtime/Experimental/Chaos/Public/RewindData.h
@@ -1762,6 +1762,8 @@ public:
 	void CHAOS_API SetTargetStateAtFrame(FGeometryParticleHandle& Handle, const int32 Frame, FFrameAndPhase::EParticleHistoryPhase Phase,
 		const FVector& Position, const FQuat& Quaternion, const FVector& LinVelocity, const FVector& AngVelocity, const bool bShouldSleep);
 
+	void CHAOS_API ForceResimAsFollower(FGeometryParticleHandle* Handle, bool bForceResim);
+
 	int32 CHAOS_API GetLatestReceivedTargetStateFrame(FGeometryParticleHandle& Handle) const;
 
 	/** Extract some history information before cleaning/pushing state*/
@@ -2237,6 +2239,8 @@ private:
 	TArray<const Private::FPBDIsland*> IslandValidationIslands;
 	TArray<const FGeometryParticleHandle*> IslandValidationIslandParticles;
 
+	TSet<const FGeometryParticleHandle*> ParticlesToForceResimAsFollower;
+
 	// BUBBLE RESIM
 	// Bubble resim needs to be enabled, see Chaos::ResimConsoleVars::bIsBubbleResimEnabled
 	// Bubble resim mitigates the cost of resimulations by restricting the amount of particles being resimulated
@@ -2291,12 +2295,12 @@ class IRewindCallback
 {
 public:
 	virtual ~IRewindCallback() = default;
-	/** Called before any sim callbacks are triggered but after physics data has marshalled over
+	/** Called before any sim callbacks are triggered but after physics data has marshaled over
 	*   This means brand new physics particles are already created for example, and any pending game thread modifications have happened
 	*   See ISimCallbackObject for recording inputs to callbacks associated with this PhysicsStep */
-	virtual void ProcessInputs_Internal(int32 PhysicsStep, const TArray<FSimCallbackInputAndObject>& SimCallbackInputs){}
+	virtual void ProcessInputs_Internal(int32 PhysicsStep, const TArray<FSimCallbackInputAndObject>& SimCallbackInputs) {}
 
-	/** Called after any presim callbacks are triggered and after physics data has marshalled over in order to modify the sim callback outputs */
+	/** Called after any presim callbacks are triggered and after physics data has marshaled over in order to modify the sim callback outputs */
 	virtual void ApplyCallbacks_Internal(int32 PhysicsStep, const TArray<ISimCallbackObject*>& SimCallbackObjects) {}
 
 	/** Called before any inputs are marshalled over to the physics thread.
@@ -2305,12 +2309,18 @@ public:
 	*/
 	virtual void ProcessInputs_External(int32 PhysicsStep, const TArray<FSimCallbackInputAndObject>& SimCallbackInputs) {}
 
-	/** Called before inputs are split into potential sub-steps and marshalled over to the physics thread.
+	/** Called before inputs are split into potential sub-steps and marshaled over to the physics thread.
 	*	The physics state has not been applied yet, and cannot be inspected anyway because this is triggered from the external thread (game thread)
 	*	Gives user the ability to call GetProducerInputData_External one last time.
 	*	Input data is shared amongst sub-steps. If NumSteps > 1 it means any input data injected will be shared for all sub-steps generated
 	*/
-	virtual void InjectInputs_External(int32 PhysicsStep, int32 NumSteps){}
+	virtual void InjectInputs_External(int32 PhysicsStep, int32 NumSteps) {}
+
+	/** Called on game thread for each physics thread task that gets queued up, before async inputs are marshaled.
+	*	Gives user the ability to call GetProducerInputData_External for each physics step.
+	*	Input data is shared amongst sub-steps. If NumSteps > 1 it means any input data injected will be shared for all sub-steps generated
+	*/
+	virtual void ProduceInputs_External(int32 PhysicsStep) {}
 
 	/** Called after sim step to give the option to rewind. Any pending inputs for the next frame will remain in the queue
 	*   Return the PhysicsStep to start resimulating from. Resim will run up until latest step passed into RecordInputs (i.e. latest physics sim simulated so far)
@@ -2324,7 +2334,7 @@ public:
 	/** Called before each rewind step. This is to give user code the opportunity to trigger other code before each rewind step
 	*   Usually to simulate external systems that ran in lock step with the physics sim
 	*/
-	virtual void PreResimStep_Internal(int32 PhysicsStep, bool bFirstStep){}
+	virtual void PreResimStep_Internal(int32 PhysicsStep, bool bFirstStep) {}
 
 	/** Called after each rewind step. This is to give user code the opportunity to trigger other code after each rewind step
 	*   Usually to simulate external systems that ran in lock step with the physics sim
@@ -2338,7 +2348,7 @@ public:
 	virtual void UnregisterRewindableSimCallback_Internal(ISimCallbackObject* Callback) {}
 
 	/** Called When resim is finished with debug information about the resim */
-	virtual void SetResimDebugInfo_Internal(const FResimDebugInfo& ResimDebugInfo){}
+	virtual void SetResimDebugInfo_Internal(const FResimDebugInfo& ResimDebugInfo) {}
 
 	/** Rewind Data holding the callback */
 	Chaos::FRewindData* RewindData = nullptr;
