IAD Cache prep to add modtime
- Move last access to FCasBlockState and remove LastAccess map

#Synced-CL 49554261

diff --git a/Engine/Source/Runtime/Experimental/IoStore/OnDemand/Private/OnDemandInstallCache.cpp b/Engine/Source/Runtime/Experimental/IoStore/OnDemand/Private/OnDemandInstallCache.cpp
--- a/Engine/Source/Runtime/Experimental/IoStore/OnDemand/Private/OnDemandInstallCache.cpp
+++ b/Engine/Source/Runtime/Experimental/IoStore/OnDemand/Private/OnDemandInstallCache.cpp
@@ -8,6 +8,7 @@
 
 #include "Algo/Accumulate.h"
 #include "Algo/Find.h"
+#include "Algo/Transform.h"
 #include "Async/Mutex.h"
 #include "Async/SharedMutex.h"
 #include "Async/UniqueLock.h"
@@ -303,7 +304,7 @@ struct FCasJournal
 		TArray<FEntry>	Entries;
 	};
 
-	using FEntryHandler = TFunction<void(const FEntry&)>;
+	using FEntryHandler = TFunction<FResult(const FEntry&)>;
 
 	static FResult			Replay(const FString& JournalFile, FEntryHandler&& Handler);
 	static FResult			Create(const FString& JournalFile);
@@ -435,7 +436,11 @@ FResult FCasJournal::Replay(const FString& JournalFile, FEntryHandler&& Handler)
 			return MakeJournalError(ECasErrorCode::ReplayJournalFailed, EIoErrorCode::InvalidCode, FString::Printf(TEXT("Found critical error journal entry in '%s'"), *JournalFile));
 		}
 
-		Handler(Entry);
+		FResult HandlerResult = Handler(Entry);
+		if (HandlerResult.HasError())
+		{
+			return HandlerResult;
+		}
 	}
 
 	return MakeValue();
@@ -669,6 +674,7 @@ struct FCasBlockState
 {
 	uint32 Size = 0;
 	uint32 ReadHandleCount = 0;
+	int64 LastAccessUtcTicks = 0;
 };
 
 struct FCasSnapshot;
@@ -716,17 +722,8 @@ struct FCas
 	void						Compact();
 	FResult						Verify(TArray<FCasAddr>& OutRemovedChunksAddrs, TMap<FCasBlockId, uint64>& OutOverBudgetBlocks);
 	void						LoadSnapshot(FCasSnapshot&& Snapshot);
-	void						ApplyJournalEntry(const FCasJournal::FEntry& JournalEntry);
-	FLastAccess					ConsumeLastAcccess()
-	{
-		FLastAccess Temp;
-		{
-			TUniqueLock Lock(Mutex);
-			Temp = MoveTemp(LastAccess);
-			LastAccess = FLastAccess();
-		}
-		return Temp;
-	}
+	FResult						ApplyJournalEntry(const FCasJournal::FEntry& JournalEntry);
+
 	// Returns the timestamps that are "dirty" and need to be flushed to disk and clears dirty flag
 	FLastAccess					GetAndClearDirtyLastAccess();
 
@@ -752,7 +749,6 @@ public: // TODO: FIXME: encapsulate these members
 	FLookup					Lookup;
 private:
 	FBlockStates			BlockIds;
-	FLastAccess				LastAccess;
 	FReadHandles			ReadHandles;
 	FEventRef				BlockReadsDoneEvent;
 	const int64				LastAccessGranularityTicks;
@@ -774,7 +770,6 @@ FResult FCas::Initialize(FStringView Directory, bool bDeleteExisting)
 
 	Lookup.Empty();
 	BlockIds.Empty();
-	LastAccess.Empty();
 
 	TStringBuilder<256> Path;
 	FPathViews::Append(Path, RootDirectory, TEXT("blocks"));
@@ -896,7 +891,6 @@ FResult FCas::DeleteBlock(FCasBlockId BlockId, TFunctionRef<void(const FCasAddr&
 
 	BlockIds.Remove(BlockId);
 	ReadHandles.Remove(BlockId);
-	LastAccess.Remove(BlockId);
 	for (auto It = Lookup.CreateIterator(); It; ++It)
 	{
 		if (It->Value.BlockId == BlockId)
@@ -1030,20 +1024,26 @@ bool FCas::UnlockedTrackAccessIf(const ECasTrackAccessType Type, const uint32 Bl
 {
 	check(BlockId.IsValid());
 
-	int64* MaybeFoundTicks = LastAccess.FindByHash(BlockIdHash, BlockId);
-	if (MaybeFoundTicks == nullptr)
+	FCasBlockState* BlockState = BlockIds.FindByHash(BlockIdHash, BlockId);
+	if (ensure(BlockState) == false)
+	{
+		return false;
+	}
+
+	int64& FoundTicks = BlockState->LastAccessUtcTicks;
+
+	if (FoundTicks == 0)
 	{
 		if (bDirty)
 		{
 			UtcTicks |= DirtyTimestampMask;
 		}
 
-		LastAccess.AddByHash(BlockIdHash, BlockId, UtcTicks);
+		FoundTicks = UtcTicks;
 
 		return true;
 	}
 
-	int64& FoundTicks = *MaybeFoundTicks;
 	const int64 PrevTicks = (FoundTicks & ~DirtyTimestampMask);
 	bDirty = bDirty || (FoundTicks & DirtyTimestampMask); // Don't clear dirty flag if already set
 
@@ -1149,7 +1149,7 @@ uint64 FCas::GetBlockInfoOld(FCasBlockInfoMap& OutBlockInfo) const
 			BlockInfo.Add(BlockId, FCasBlockInfo
 			{
 				.FileSize = uint64(FileSize),
-				.LastAccess = UtcTicks ? (*UtcTicks & ~FCas::DirtyTimestampMask) : TimestampForMissingLastAccess
+				.LastAccess = (UtcTicks && *UtcTicks) ? (*UtcTicks & ~FCas::DirtyTimestampMask) : TimestampForMissingLastAccess
 			});
 			TotalSize += uint64(FileSize);
 
@@ -1166,7 +1166,11 @@ uint64 FCas::GetBlockInfoOld(FCasBlockInfoMap& OutBlockInfo) const
 	FLastAccess Access;
 	{
 		TUniqueLock Lock(Mutex);
-		Access = LastAccess;
+		Access.Reserve(BlockIds.Num());
+		Algo::Transform(BlockIds, Access, [](const TPair<FCasBlockId, FCasBlockState>& Kv) 
+		{
+			return FLastAccess::ElementType(Kv.Key, Kv.Value.LastAccessUtcTicks);
+		});
 	}
 	IPlatformFile& Ipf = FPlatformFileManager::Get().GetPlatformFile();
 	FDirectoryVisitor Visitor(Ipf, OutBlockInfo, MoveTemp(Access));
@@ -1193,13 +1197,12 @@ uint64 FCas::GetBlockInfo(FCasBlockInfoMap& OutBlockInfo) const
 		{
 			const FCasBlockId BlockId = Kv.Key;
 			const uint32 FileSize = Kv.Value.Size;
-
-			const int64* UtcTicks = LastAccess.Find(BlockId);
+			const int64 UtcTicks = Kv.Value.LastAccessUtcTicks;
 
 			OutBlockInfo.Add(BlockId, FCasBlockInfo
 			{
 				.FileSize = uint64(FileSize),
-				.LastAccess = UtcTicks ? (*UtcTicks & ~FCas::DirtyTimestampMask) : TimestampForMissingLastAccess
+				.LastAccess = UtcTicks ? (UtcTicks & ~FCas::DirtyTimestampMask) : TimestampForMissingLastAccess
 			});
 
 			TotalSize += uint64(FileSize);
@@ -1297,7 +1300,6 @@ void FCas::Compact()
 	Lookup.Compact();
 	BlockIds.Compact();
 	ReadHandles.Compact();
-	LastAccess.Compact();
 }
 
 FResult FCas::Verify(TArray<FCasAddr>& OutRemovedChunksAddrs, TMap<FCasBlockId, uint64>& OutOverBudgetBlocks)
@@ -1324,13 +1326,13 @@ FResult FCas::Verify(TArray<FCasAddr>& OutRemovedChunksAddrs, TMap<FCasBlockId,
 
 		if (const uint64* FileSize = BlockFiles.Find(BlockId))
 		{
-			if (int64* TimeStamp = LastAccess.Find(BlockId))
+			if (BlockState.LastAccessUtcTicks)
 			{
-				if (*TimeStamp > Now)
+				if (BlockState.LastAccessUtcTicks > Now)
 				{
 					const FString Filename = GetBlockFilename(BlockId);
 					UE_LOG(LogIoStoreOnDemand, Warning, TEXT("Found future last access time for CAS block '%s'"), *Filename);
-					*TimeStamp = Now;
+					BlockState.LastAccessUtcTicks = Now;
 				}
 			}
 			else
@@ -1344,12 +1346,12 @@ FResult FCas::Verify(TArray<FCasAddr>& OutRemovedChunksAddrs, TMap<FCasBlockId,
 					if (ModTime > Now)
 					{
 						UE_LOG(LogIoStoreOnDemand, Warning, TEXT("Missing last access time and found future mod time so using now time for CAS block '%s'"), *Filename);
-						LastAccess.Add(BlockId, Now);
+						BlockState.LastAccessUtcTicks = Now;
 					}
 					else
 					{
 						UE_LOG(LogIoStoreOnDemand, Warning, TEXT("Missing last access time so using file mod time for CAS block '%s'"), *Filename);
-						LastAccess.Add(BlockId, ModTime.GetTicks());
+						BlockState.LastAccessUtcTicks = ModTime.GetTicks();
 					}
 				}
 				else
@@ -1357,7 +1359,7 @@ FResult FCas::Verify(TArray<FCasAddr>& OutRemovedChunksAddrs, TMap<FCasBlockId,
 				{
 					// Failed to get mod time
 					UE_LOG(LogIoStoreOnDemand, Warning, TEXT("Missing last access time for CAS block '%s'"), *Filename);
-					LastAccess.Add(BlockId, TimestampForMissingLastAccess);
+					BlockState.LastAccessUtcTicks = TimestampForMissingLastAccess;
 				}
 			}
 
@@ -1384,7 +1386,6 @@ FResult FCas::Verify(TArray<FCasAddr>& OutRemovedChunksAddrs, TMap<FCasBlockId,
 			UE_LOG(LogIoStoreOnDemand, Warning, TEXT("%s"), *ErrorMessage);
 		}
 
-		LastAccess.Remove(BlockId);
 		BlockIt.RemoveCurrent();
 		Result = MakeCasError<void>(ECasErrorCode::VerifyFailed, EIoErrorCode::NotFound, MoveTemp(ErrorMessage));
 	}
@@ -1436,9 +1437,9 @@ FCas::FLastAccess FCas::GetAndClearDirtyLastAccess()
 	FCas::FLastAccess DirtyLastAccess;
 	{
 		TUniqueLock Lock(Mutex);
-		for (TPair<FCasBlockId, int64>& Kv : LastAccess)
+		for (TPair<FCasBlockId, FCasBlockState>& Kv : BlockIds)
 		{
-			int64& Timestamp = Kv.Value;
+			int64& Timestamp = Kv.Value.LastAccessUtcTicks;
 			if (Timestamp & FCas::DirtyTimestampMask)
 			{
 				DirtyLastAccess.Add(Kv.Key, Timestamp);
@@ -1604,6 +1605,8 @@ TResult<FCasSnapshot> FCasSnapshot::FromJournal(const FString& JournalFile)
 				break;
 			}
 			};
+
+			return MakeValue();
 		});
 
 	if (ReplayResult.HasError())
@@ -1816,16 +1819,15 @@ void FCas::LoadSnapshot(FCasSnapshot&& Snapshot)
 	}
 
 	BlockIds.Reserve(Snapshot.Blocks.Num());
-	LastAccess.Reserve(Snapshot.Blocks.Num());
 	for (const FCasSnapshot::FBlock& Block : Snapshot.Blocks)
 	{
-		BlockIds.Add(Block.BlockId);
-		LastAccess.Add(Block.BlockId, Block.LastAccess);
+		FCasBlockState& BlockState = BlockIds.Add(Block.BlockId);
+		BlockState.LastAccessUtcTicks = Block.LastAccess;
 	}
 }
 
 ///////////////////////////////////////////////////////////////////////////////
-void FCas::ApplyJournalEntry(const FCasJournal::FEntry& JournalEntry)
+FResult FCas::ApplyJournalEntry(const FCasJournal::FEntry& JournalEntry)
 {
 	switch (JournalEntry.Type())
 	{
@@ -1853,17 +1855,23 @@ void FCas::ApplyJournalEntry(const FCasJournal::FEntry& JournalEntry)
 	{
 		const FCasJournal::FEntry::FBlockOperation& Op = JournalEntry.BlockOperation;
 		BlockIds.Remove(Op.BlockId);
-		LastAccess.Remove(Op.BlockId);
 		break;
 	}
 	case FCasJournal::FEntry::EType::BlockAccess:
 	{
 		const FCasJournal::FEntry::FBlockOperation& Op = JournalEntry.BlockOperation;
 		constexpr const bool bDirty = false;
-		check(TrackAccessIf(ECasTrackAccessType::Always, Op.BlockId, Op.UtcTicks, bDirty));
+
+		if (ensure(TrackAccessIf(ECasTrackAccessType::Always, Op.BlockId, Op.UtcTicks, bDirty)) == false)
+		{
+			using namespace UE::UnifiedError::IoStoreOnDemand;
+			return MakeJournalError(ECasErrorCode::ReplayJournalFailed, EIoErrorCode::FileCorrupt, TEXT("Journal BlockAccess Op references an invalid block"));
+		}
 		break;
 	}
 	};
+
+	return MakeValue();
 }
 
 ///////////////////////////////////////////////////////////////////////////////
@@ -2111,7 +2119,11 @@ FOnDemandInstallCache::FOnDemandInstallCache(const FOnDemandInstallCacheConfig&
 		// Replay the journal 
 		InitResult = FCasJournal::Replay(JournalFile, [this](const FCasJournal::FEntry& JournalEntry)
 		{
-			Cas.ApplyJournalEntry(JournalEntry);
+			FResult Result = Cas.ApplyJournalEntry(JournalEntry);
+			if (Result.HasError())
+			{
+				return Result;
+			}
 
 			switch(JournalEntry.Type())
 			{
@@ -2129,6 +2141,8 @@ FOnDemandInstallCache::FOnDemandInstallCache(const FOnDemandInstallCacheConfig&
 				break;
 			}
 			};
+
+			return Result;
 		});
 	}
 
@@ -2194,7 +2208,7 @@ void FOnDemandInstallCache::Shutdown()
 		UE_LOG(LogIoStoreOnDemand, Error, TEXT("Failed to flush pending chunks on shutdown, reason '%s'"), *Result.GetError().GetErrorMessage(true).ToString());
 	}
 
-	FCas::FLastAccess LastAccess = Cas.ConsumeLastAcccess();
+	FCas::FLastAccess LastAccess = Cas.GetAndClearDirtyLastAccess();
 
 	const FString JournalFile = GetJournalFilename();
 	FCasJournal::FTransaction Transaction = FCasJournal::Begin(JournalFile);
@@ -4108,6 +4122,8 @@ TEST_CASE("IoStore::OnDemand::InstallCache::Journal", "[IoStoreOnDemand][Install
 					CHECK(false);
 					break;
 				};
+
+				return MakeValue();
 			});
 		CHECK(Result.HasValue());
 		CHECK(Locs.Num() == ExpectedAddresses.Num());
@@ -4166,6 +4182,8 @@ TEST_CASE("IoStore::OnDemand::InstallCache::Journal", "[IoStoreOnDemand][Install
 					CHECK(false);
 					break;
 				};
+
+				return MakeValue();
 			});
 
 		CHECK(Result.HasValue());
@@ -4213,6 +4231,8 @@ TEST_CASE("IoStore::OnDemand::InstallCache::Journal", "[IoStoreOnDemand][Install
 					CHECK(false);
 					break;
 				};
+
+				return MakeValue();
 			});
 
 		CHECK(Result.HasValue());
