GeometryCollection - Check for valid material index before grabbing its proxy

#Synced-CL 49554335

diff --git a/Engine/Source/Runtime/Experimental/GeometryCollectionEngine/Private/GeometryCollection/GeometryCollectionSceneProxy.cpp b/Engine/Source/Runtime/Experimental/GeometryCollectionEngine/Private/GeometryCollection/GeometryCollectionSceneProxy.cpp
--- a/Engine/Source/Runtime/Experimental/GeometryCollectionEngine/Private/GeometryCollection/GeometryCollectionSceneProxy.cpp
+++ b/Engine/Source/Runtime/Experimental/GeometryCollectionEngine/Private/GeometryCollection/GeometryCollectionSceneProxy.cpp
@@ -498,7 +498,7 @@ void FGeometryCollectionSceneProxyBase::GetSectionMaterialProxies(TConstArrayVie
 	{
 		const FGeometryCollectionMeshElement& Section = InSectionArray[SectionIndex];
 
-		FMaterialRenderProxy* MaterialProxy = Materials[Section.MaterialIndex]->GetRenderProxy();
+		FMaterialRenderProxy* MaterialProxy = Materials.IsValidIndex(Section.MaterialIndex) ? Materials[Section.MaterialIndex]->GetRenderProxy() : nullptr;
 
 		if (MaterialProxy == nullptr)
 		{
