Dataflow : add compiler classes
- those classes are not used yet and do not yet support conditional execution of graph branches

#Synced-CL 49572843

diff --git a/Engine/Source/Runtime/Experimental/Dataflow/Core/Private/Dataflow/DataflowCompiledGraph.cpp b/Engine/Source/Runtime/Experimental/Dataflow/Core/Private/Dataflow/DataflowCompiledGraph.cpp
new file mode 100644
--- /dev/null
+++ b/Engine/Source/Runtime/Experimental/Dataflow/Core/Private/Dataflow/DataflowCompiledGraph.cpp
@@ -0,0 +1,94 @@
+// Copyright Epic Games, Inc. All Rights Reserved.
+
+#include "Dataflow/DataflowCompiledGraph.h"
+
+namespace UE::Dataflow
+{
+	int32 FCompiledGraph::StartTask()
+	{
+		return Tasks.Emplace(FTask{ .CommandIndex = Commands.Num() });
+	}
+
+	void FCompiledGraph::AddNode(FGuid NodeId)
+	{
+		if (ensure(Tasks.Num() > 0))
+		{
+			const FNode* ExistingNode = Nodes.FindByPredicate([&NodeId](const FNode& Node) { return Node.NodeId == NodeId; });
+			ensure(ExistingNode == nullptr);
+			const int32 CurrentTaskIndex = Tasks.Num() - 1;
+			Tasks[CurrentTaskIndex].NumCommands++;
+			const int32 NodeIndex = Nodes.Emplace(FNode{ .NodeId = NodeId, .TaskIndex = CurrentTaskIndex });
+			Commands.Emplace(FCommand{ .Type = ECommandType::AddNode, .Parameter = NodeIndex });
+		}
+	}
+	void FCompiledGraph::AddDependency(int32 TaskId)
+	{
+		if (ensure(Tasks.Num() > 0) && ensure(Tasks.IsValidIndex(TaskId)))
+		{
+			const int32 CurrentTaskIndex = Tasks.Num() - 1;
+			Tasks[CurrentTaskIndex].NumCommands++;
+			Commands.Emplace(FCommand{ .Type = ECommandType::AddDependency, .Parameter = TaskId });
+		}
+	}
+
+	void FCompiledGraph::EndTask()
+	{
+		const int32 CurrentTaskIndex = Tasks.Num() - 1;
+		if (ensure(Tasks.IsValidIndex(CurrentTaskIndex)))
+		{
+			Tasks[CurrentTaskIndex].NumCommands++;
+			Commands.Emplace(FCommand{ .Type = ECommandType::AddDependency, .Parameter = CurrentTaskIndex });
+		}
+	}
+
+	int32 FCompiledGraph::GetNumNodes() const
+	{
+		return Nodes.Num();
+	}
+
+	FGuid FCompiledGraph::GetNextNodeId(int32 NodeIndex) const
+	{
+		if (NodeIndex >= INDEX_NONE)
+		{
+			const int32 NextNodeIndex = NodeIndex + 1;
+			if (Nodes.IsValidIndex(NextNodeIndex))
+			{
+				return Nodes[NextNodeIndex].NodeId;
+			}
+		}
+		return FGuid();
+	}
+
+	int32 FCompiledGraph::GetNumTasks() const
+	{
+		return Tasks.Num();
+	}
+
+	void FCompiledGraph::GetTasksNodesAndDependencies(int32 TaskId, TArray<FGuid>& OutNodeIds, TArray<int32>& OutParentTaskIds) const
+	{
+		OutNodeIds.Reset();
+		OutParentTaskIds.Reset();
+		if (Tasks.IsValidIndex(TaskId))
+		{
+			for (int32 Offset = 0; Offset < Tasks[TaskId].NumCommands; ++Offset)
+			{
+				const int32 CommandIndex = Tasks[TaskId].CommandIndex + Offset;
+				if (Commands.IsValidIndex(CommandIndex))
+				{
+					if (Commands[CommandIndex].Type == ECommandType::AddNode)
+					{
+						const int32 NodeIndex = Commands[CommandIndex].Parameter;
+						if (Nodes.IsValidIndex(NodeIndex))
+						{
+							OutNodeIds.Add(Nodes[NodeIndex].NodeId);
+						}
+					}
+					if (Commands[CommandIndex].Type == ECommandType::AddDependency)
+					{
+						OutParentTaskIds.Add(Commands[CommandIndex].Parameter);
+					}
+				}
+			}
+		}
+	}
+}
diff --git a/Engine/Source/Runtime/Experimental/Dataflow/Core/Public/Dataflow/DataflowCompiledGraph.h b/Engine/Source/Runtime/Experimental/Dataflow/Core/Public/Dataflow/DataflowCompiledGraph.h
new file mode 100644
--- /dev/null
+++ b/Engine/Source/Runtime/Experimental/Dataflow/Core/Public/Dataflow/DataflowCompiledGraph.h
@@ -0,0 +1,80 @@
+// Copyright Epic Games, Inc. All Rights Reserved.
+
+#pragma once
+
+#include "Containers/Array.h"
+#include "HAL/Platform.h"
+#include "Misc/Guid.h"
+
+namespace UE::Dataflow
+{
+	/** 
+	* This structure stores a compiled graph
+	* A compiled graph allow a dataflow evaluator to evaluate a graph efficiently without having to parse the graph again
+	* Reference to nodes are sorted by tasks in order of execution and dependency
+	* Making it easy to spawn actual async task from for example
+	* The data is also organized in arrays of fixed sized structure for faster cache access and limited allocation
+	* 
+	* Note that a compiled graph is associated to a specific main or sub graph
+	* 
+	* Todo(dataflow) : A compiled graph do not yet provide any information about branching and conditional, this will be added in the near future 
+	*/
+	struct FCompiledGraph
+	{
+	public:
+		DATAFLOWCORE_API int32 StartTask();
+		DATAFLOWCORE_API void AddNode(FGuid NodeId);
+		DATAFLOWCORE_API void AddDependency(int32 TaskId);
+		DATAFLOWCORE_API void EndTask();
+
+		DATAFLOWCORE_API int32 GetNumNodes() const;
+		DATAFLOWCORE_API FGuid GetNextNodeId(int32 NodeIndex) const;
+
+		DATAFLOWCORE_API int32 GetNumTasks() const;
+		DATAFLOWCORE_API void GetTasksNodesAndDependencies(int32 TaskId, TArray<FGuid>& OutNodeIds, TArray<int32>& OutParentTaskIds) const;
+
+	private:
+		enum class ECommandType
+		{
+			AddNode,		// Parameter is an index into a node array 
+			AddDependency,	// Parameter is an index into a task array
+			LaunchTask		// Parameter is an index into a task array 
+		};
+
+		struct FTask
+		{
+			/** Index into the command buffer ( Start of the task ) */
+			int32 CommandIndex = INDEX_NONE; 
+
+			/** Number of commands for this task */
+			int32 NumCommands = 0;	
+ 		};
+
+		struct FNode
+		{
+			/** Guid uniquely identifying a node in the graph */
+			FGuid NodeId;
+
+			/** index of the task in the task array */
+			int32 TaskIndex = INDEX_NONE; // index into the task buffer
+		};
+
+		struct FCommand
+		{
+			/** Type of the command (see ECommandType) */
+			ECommandType Type = ECommandType::LaunchTask;
+
+			/** Parameter of the command (an index if the task or node array) */
+			int32 Parameter = 0;
+		};
+
+		/** List of tasks (index into the command buffer array) */
+		TArray<FTask> Tasks;
+
+		/** List of nodes (with back reference to the task array) */
+		TArray<FNode> Nodes;
+
+		/** List of commands (with reference to the task and node arrays) */
+		TArray<FCommand> Commands;
+	};
+}
\ No newline at end of file
diff --git a/Engine/Source/Runtime/Experimental/Dataflow/Engine/Private/Dataflow/DataflowCompiler.cpp b/Engine/Source/Runtime/Experimental/Dataflow/Engine/Private/Dataflow/DataflowCompiler.cpp
new file mode 100644
--- /dev/null
+++ b/Engine/Source/Runtime/Experimental/Dataflow/Engine/Private/Dataflow/DataflowCompiler.cpp
@@ -0,0 +1,354 @@
+// Copyright Epic Games, Inc. All Rights Reserved.
+
+#include "Dataflow/DataflowCompiler.h"
+
+#include "Dataflow/DataflowEdNode.h"
+#include "Dataflow/DataflowGraph.h"
+#include "Dataflow/DataflowNode.h"
+#include "Dataflow/DataflowObject.h"
+#include "Dataflow/DataflowSubGraph.h"
+
+#define LOCTEXT_NAMESPACE "DataflowCompiler"
+
+namespace UE::Dataflow
+{
+	namespace Private
+	{
+		struct FNodeSequence;
+		using FNodeSequenceRef = TSharedRef<FNodeSequence>;
+
+		// this represent a linear list of node that can be sequencially executed
+		// and make sure data is ready when they do execute
+		// this can be executed in a single task 
+		struct FNodeSequence
+		{
+		public:
+			void Build(FConstNodeRef StartNode);
+
+			const TArray<FConstNodeRef>& GetNodes() const { return Nodes; }
+			const TArray<FConstNodeRef>& GetDownstreamNodes() const { return DownstreamNodes; }
+
+			bool IsEmpty() const
+			{
+				return Nodes.IsEmpty();
+			}
+
+			static void AddDependency(FNodeSequenceRef UpstreamSequence, FNodeSequenceRef DownstreamSequence)
+			{
+				if (ensure(!UpstreamSequence->IsEmpty() && !DownstreamSequence->IsEmpty()))
+				{
+					UpstreamSequence->AddDownstreamDependency(DownstreamSequence);
+					DownstreamSequence->AddUpstreamDependency(UpstreamSequence);
+
+					FConstNodeRef DownstreamNode = DownstreamSequence->GetNodes()[0];
+					ensure(UpstreamSequence->GetDownstreamNodes().Contains(DownstreamNode));
+				}
+			}
+
+			int32 GetExecutionIndex() const
+			{
+				return ExecutionIndex;
+			}
+
+			void InitExecutionIndex(int32 Value)
+			{
+				ExecutionIndex = FMath::Max(ExecutionIndex, Value);
+				for (FNodeSequenceRef DownstreamSequence : DownstreamDependencies)
+				{
+					DownstreamSequence->InitExecutionIndex(ExecutionIndex + 1);
+				}
+			}
+
+			void SetTaskIndex(int32 Index) { TaskIndex = Index; }
+			int32 GetTaskIndex() const { return TaskIndex; }
+
+			void AddToCompiledGraph(FCompiledGraph& OutCompiledGraph)
+			{
+				const int32 NewTaskIndex = OutCompiledGraph.StartTask();
+				{
+					for (FConstNodeRef Node : Nodes)
+					{
+						OutCompiledGraph.AddNode(Node->GetGuid());
+					}
+
+					for (FNodeSequenceRef Dependency : UpstreamDependencies)
+					{
+						OutCompiledGraph.AddDependency(Dependency->GetTaskIndex());
+					}
+				}
+				OutCompiledGraph.EndTask();
+				ensure(NewTaskIndex == GetTaskIndex());
+			}
+
+			void AddDownstreamDependency(FNodeSequenceRef DownstreamSequence)
+			{
+				DownstreamDependencies.AddUnique(DownstreamSequence);
+			}
+
+			void AddUpstreamDependency(FNodeSequenceRef UpstreamSequence)
+			{
+				UpstreamDependencies.AddUnique(UpstreamSequence);
+			}
+
+		private:
+			int32 TaskIndex = INDEX_NONE;
+			int32 ExecutionIndex = INDEX_NONE;
+			TArray<FConstNodeRef> Nodes;
+			TArray<FConstNodeRef> DownstreamNodes;
+			TArray<FNodeSequenceRef> DownstreamDependencies;
+			TArray<FNodeSequenceRef> UpstreamDependencies;
+		};
+
+		void FNodeSequence::Build(FConstNodeRef StartNode)
+		{
+			Nodes.Reset();
+			DownstreamNodes.Reset();
+
+			Nodes.Add(StartNode);
+
+			TArray<FConstNodeRef> UpstreamNodes;
+
+			FConstNodePtr NodeToProcess = StartNode;
+			while (NodeToProcess)
+			{
+				// check if we have a singular node connected to this one
+				FCompiler::FindDownstreamNodes(NodeToProcess.ToSharedRef(), DownstreamNodes);
+				if (DownstreamNodes.Num() != 1)
+				{
+					// No more nodes or too many, stop building 
+					break;
+				}
+
+				// one downstream, let's make sure this connected node is not a merge node (upstream noded different from the one we process )
+				FCompiler::FindUpstreamNodes(DownstreamNodes[0], UpstreamNodes);
+				if (UpstreamNodes.Num() != 1 || UpstreamNodes[0] != NodeToProcess)
+				{
+					// This is a merge node, stop the build here
+					break;
+				}
+
+				Nodes.Add(DownstreamNodes[0]);
+				NodeToProcess = Nodes.Last();
+				DownstreamNodes.Reset();
+			}
+		}
+	}
+
+	/////////////////////////////////////////////////////////////////////////////////////////////////////
+
+	bool FCompilerResults::Succedded() const
+	{
+		return NumErrors == 0;
+	}
+
+	bool FCompilerResults::Failed() const
+	{
+		return NumErrors > 0;
+	}
+
+	void FCompilerResults::AddMessage(EMessageSeverity::Type InSeverity, const FGuid& InNode, const FGuid& InOutput, const FText& Text)
+	{
+		if (InSeverity == EMessageSeverity::Error)
+		{
+			NumErrors++;
+		}
+		else if (InSeverity == EMessageSeverity::Warning)
+		{
+			NumWarnings++;
+		}
+		Messages.Emplace(
+			FCompilerMessage
+			{
+				.Severity = InSeverity,
+				.Node = InNode,
+				.Output = InOutput,
+				.Text = Text,
+			});
+	}
+
+	///////////////////////////////////////////////////////////////////////////////////////////////////////////////////////////
+
+
+
+	///////////////////////////////////////////////////////////////////////////////////////////////////////////////////////////
+
+	/*static*/ bool FCompiler::Compile(const FCompilerParameters& Params, FCompilerResults& Results)
+	{
+		static FGuid NoGuid;
+
+		const double StartCompilationTimeMs = FPlatformTime::ToMilliseconds64(FPlatformTime::Cycles64());
+
+		// Get dataflow Object to compile
+		TStrongObjectPtr<const UDataflow> DataflowObject = Params.DataflowObjectToCompile.Pin();
+		if (DataflowObject)
+		{
+			TArray<FConstNodeRef> NodesInGraph;
+
+			// compile main graph 
+			CompileGraph(*DataflowObject, Params, Results);
+
+			// now compile subgraphs
+			for (const UDataflowSubGraph* SubGraph : DataflowObject->GetSubGraphs())
+			{
+				if (SubGraph)
+				{
+					CompileGraph(*SubGraph, Params, Results);
+				}
+			}
+		}
+		else
+		{
+			Results.AddMessage(EMessageSeverity::Error, NoGuid, NoGuid, LOCTEXT("NullDataflowObject", "Dataflow object is null"));
+		}
+
+		const double EndCompilationTimeMs = FPlatformTime::ToMilliseconds64(FPlatformTime::Cycles64());
+		Results.CompilationTimeMs = (EndCompilationTimeMs - StartCompilationTimeMs);
+
+		return Results.Succedded();
+	}
+
+	/*static*/ bool FCompiler::CompileGraph(const UEdGraph& GraphObject, const FCompilerParameters& Params, FCompilerResults& Results)
+	{
+		// Now build the command buffer
+		TSharedRef<FCompiledGraph> CompiledGraph = MakeShared<FCompiledGraph>();
+		Results.AddCompiledGraph(&GraphObject, CompiledGraph);
+
+		// Find first the left nost thread as start node for the node thread
+		TArray<FConstNodeRef> StartNodes;
+		FindStartNodes(GraphObject, StartNodes);
+
+		TArray<FConstNodeRef> NodeToProcess;
+		NodeToProcess.Append(StartNodes);
+
+		TMap<FConstNodeRef, Private::FNodeSequenceRef> NodeSequencesByStartNode;
+
+		// let create all the needed node sequences
+		while (NodeToProcess.Num() > 0)
+		{
+			FConstNodeRef StartNode = NodeToProcess.Pop();
+			if (ensure(!NodeSequencesByStartNode.Contains(StartNode)))
+			{
+				Private::FNodeSequenceRef NodeSequence = MakeShared<Private::FNodeSequence>();
+				NodeSequence->Build(StartNode);
+				NodeSequencesByStartNode.Add(StartNode, NodeSequence);
+
+				for (FConstNodeRef DownstreamNode : NodeSequence->GetDownstreamNodes())
+				{
+					if (!NodeSequencesByStartNode.Contains(DownstreamNode))
+					{
+						NodeToProcess.AddUnique(DownstreamNode);
+					}
+				}
+			}
+		}
+
+		const Private::FNodeSequenceRef EmptyNodeSequence = MakeShared<Private::FNodeSequence>();
+
+		// Connect the sequences together
+		for (const TPair<FConstNodeRef, Private::FNodeSequenceRef>& NodeSequenceByStartNode : NodeSequencesByStartNode)
+		{
+			const Private::FNodeSequenceRef UpstreamNodeSequence = NodeSequenceByStartNode.Value;
+			const TArray<FConstNodeRef>& DownstreamNodes = UpstreamNodeSequence->GetDownstreamNodes();
+			for (FConstNodeRef DownstreamNode: DownstreamNodes)
+			{
+				const Private::FNodeSequenceRef DownstreamNodeSequence = NodeSequencesByStartNode.FindRef(DownstreamNode, EmptyNodeSequence);
+				if (!DownstreamNodeSequence->IsEmpty())
+				{
+					Private::FNodeSequence::AddDependency(UpstreamNodeSequence, DownstreamNodeSequence);
+				}
+			}
+		}
+
+		// sort the sequences
+		TArray<Private::FNodeSequenceRef> SortedSequencesByDependencies;
+		NodeSequencesByStartNode.GenerateValueArray(SortedSequencesByDependencies);
+
+		for (Private::FNodeSequenceRef Sequence : SortedSequencesByDependencies)
+		{
+			Sequence->InitExecutionIndex(0);
+		}
+
+		Algo::Sort(SortedSequencesByDependencies,
+			[](const Private::FNodeSequenceRef& A, const Private::FNodeSequenceRef& B) -> bool
+			{
+				return (A->GetExecutionIndex() < B->GetExecutionIndex());
+			});
+
+		for (int32 TaskIndex = 0; TaskIndex < SortedSequencesByDependencies.Num(); ++TaskIndex)
+		{
+			SortedSequencesByDependencies[TaskIndex]->SetTaskIndex(TaskIndex);
+		}
+		for (Private::FNodeSequenceRef Sequence : SortedSequencesByDependencies)
+		{
+			Sequence->AddToCompiledGraph(*CompiledGraph);
+		}
+		// TODO : transfer the compiled graph over to the results
+		return Results.Succedded();
+	}
+
+	/*static*/ void FCompiler::FindStartNodes(const UEdGraph& GraphObject, TArray<FConstNodeRef>& OutStartNodes)
+	{
+		OutStartNodes.Reset();
+		for (const UEdGraphNode* EdNode : GraphObject.Nodes)
+		{
+			if (const UDataflowEdNode* DataflowEdNode = Cast<UDataflowEdNode>(EdNode))
+			{
+				if (TSharedPtr<const FDataflowNode> DataflowNode = DataflowEdNode->GetDataflowNode())
+				{
+					bool bHasAnyInputConnected = false;
+					for (const FDataflowInput* Input : DataflowNode->GetInputs())
+					{
+						if (Input && Input->IsConnected())
+						{
+							bHasAnyInputConnected = true;
+							break;
+						}
+					}
+					if (!bHasAnyInputConnected)
+					{
+						OutStartNodes.Add(DataflowNode.ToSharedRef());
+					}
+				}
+			}
+		}
+	}
+
+	/*static*/ void FCompiler::FindDownstreamNodes(const FConstNodeRef& Node, TArray<FConstNodeRef>& OutDownstreamNodes)
+	{
+		OutDownstreamNodes.Reset();
+		for (const FDataflowOutput* Output : Node->GetOutputs())
+		{
+			if (Output && Output->IsConnected())
+			{
+				for (const FDataflowInput* Input : Output->GetConnections())
+				{
+					if (Input && Input->GetOwningNode())
+					{
+						OutDownstreamNodes.AddUnique(Input->GetOwningNode()->AsShared());
+					}
+				}
+			}
+		}
+	}
+
+	/*static*/ void FCompiler::FindUpstreamNodes(const FConstNodeRef& Node, TArray<FConstNodeRef>& OutUpstreamNodes)
+	{
+		OutUpstreamNodes.Reset();
+		for (const FDataflowInput* Input : Node->GetInputs())
+		{
+			if (Input && Input->IsConnected())
+			{
+				if (const FDataflowOutput* Output = Input->GetConnection())
+				{
+					if (Output && Output->GetOwningNode())
+					{
+						OutUpstreamNodes.AddUnique(Output->GetOwningNode()->AsShared());
+					}
+				}
+			}
+		}
+	}
+
+}
+
+#undef LOCTEXT_NAMESPACE
\ No newline at end of file
diff --git a/Engine/Source/Runtime/Experimental/Dataflow/Engine/Public/Dataflow/DataflowCompiler.h b/Engine/Source/Runtime/Experimental/Dataflow/Engine/Public/Dataflow/DataflowCompiler.h
new file mode 100644
--- /dev/null
+++ b/Engine/Source/Runtime/Experimental/Dataflow/Engine/Public/Dataflow/DataflowCompiler.h
@@ -0,0 +1,82 @@
+// Copyright Epic Games, Inc. All Rights Reserved.
+
+#pragma once
+
+#include "Containers/Array.h"
+#include "Dataflow/DataflowCompiledGraph.h"
+#include "HAL/Platform.h"
+#include "Logging/TokenizedMessage.h"
+#include "Misc/Guid.h"
+#include "Templates/SharedPointer.h"
+#include "UObject/WeakObjectPtr.h"
+
+struct FDataflowNode;
+class UDataflow;
+class UEdGraph;
+
+namespace UE::Dataflow
+{
+	class FGraph;
+
+	using FConstNodeRef = TSharedRef<const FDataflowNode>;
+	using FConstNodePtr = TSharedPtr<const FDataflowNode>;
+	/**
+	* Dataflow compiler takes a UDataflowGraph and build a representation of the graph execution and stores it into compiled graph objects
+	* the compiled graph are only valid for the topology of the graph at the time of compilation 
+	*/
+	struct FCompilerParameters
+	{
+	public:
+		bool OptimizeRerouteNodes = true;
+		TWeakObjectPtr<const UDataflow> DataflowObjectToCompile;
+	};
+
+	struct FCompilerMessage
+	{
+		EMessageSeverity::Type Severity = EMessageSeverity::Info;
+		FGuid Node;
+		FGuid Output;
+		FText Text;
+	};
+
+	struct FCompilerResults
+	{
+	public:
+		bool Succedded() const;
+		bool Failed() const;
+
+		void AddMessage(EMessageSeverity::Type InSeverity, const FGuid& InNode, const FGuid& InOutput, const FText& Text);
+
+		void AddCompiledGraph(const UObject* GraphObject, TSharedRef<FCompiledGraph> CompiledGraph)
+		{
+			ensure(!CompiledGraphs.Contains(GraphObject));
+			CompiledGraphs.Add(GraphObject, CompiledGraph);
+		}
+
+		TSharedPtr<const FCompiledGraph> GetCompiledGraph(const UObject* GraphObject) const
+		{
+			const TSharedRef<FCompiledGraph>* CompiledGraphPtr = CompiledGraphs.Find(GraphObject);
+			return CompiledGraphPtr ? CompiledGraphPtr->ToSharedPtr() : TSharedPtr<const FCompiledGraph>();
+		}
+
+	private:
+		friend struct FCompiler;
+
+		double CompilationTimeMs = 0;
+		int32 NumErrors = 0;
+		int32 NumWarnings = 0;
+		TArray<FCompilerMessage> Messages;
+		TMap<const UObject*, TSharedRef<FCompiledGraph>> CompiledGraphs;
+	};
+
+	struct FCompiler
+	{
+	public:
+		DATAFLOWENGINE_API static bool Compile(const FCompilerParameters& Params, FCompilerResults& Results);
+
+		static bool CompileGraph(const UEdGraph& GraphObject, const FCompilerParameters& Params, FCompilerResults& Results);
+		static void FindStartNodes(const UEdGraph& GraphObject, TArray<FConstNodeRef>& OutStartNodes);
+		static void FindDownstreamNodes(const FConstNodeRef& Node, TArray<FConstNodeRef>& OutDownstreamNodes);
+		static void FindUpstreamNodes(const FConstNodeRef& Node, TArray<FConstNodeRef>& OutUpstreamNodes);
+	};
+}
\ No newline at end of file
