Chaos Visual Debugger: Fix for an issue in the start frame CVD macro that causes a compilation error when pasing *this to it as second argument, due to operator precedence (opeator.() has higher precendence than operator*()).

#Synced-CL 49576089

diff --git a/Engine/Source/Runtime/Experimental/Chaos/Public/ChaosVisualDebugger/ChaosVDTraceMacros.h b/Engine/Source/Runtime/Experimental/Chaos/Public/ChaosVisualDebugger/ChaosVDTraceMacros.h
--- a/Engine/Source/Runtime/Experimental/Chaos/Public/ChaosVisualDebugger/ChaosVDTraceMacros.h
+++ b/Engine/Source/Runtime/Experimental/Chaos/Public/ChaosVisualDebugger/ChaosVDTraceMacros.h
@@ -63,7 +63,7 @@
 		#define CVD_TRACE_SOLVER_START_FRAME(SolverType, SolverRef) \
 			FChaosVDContext StartContextData; \
 			FChaosVisualDebuggerTrace::GetCVDContext<SolverType>(SolverRef, StartContextData); \
-			FChaosVisualDebuggerTrace::TraceSolverFrameStart(StartContextData, FChaosVisualDebuggerTrace::GetDebugName<SolverType>(SolverRef), SolverRef.GetCVDFrameNumber());
+			FChaosVisualDebuggerTrace::TraceSolverFrameStart(StartContextData, FChaosVisualDebuggerTrace::GetDebugName<SolverType>(SolverRef), (SolverRef).GetCVDFrameNumber());
 	#endif
 
 	#ifndef CVD_TRACE_SOLVER_END_FRAME
