IAD cache block modtime tracking.
Fix size of critical error transaction.

#Synced-CL 49554261

diff --git a/Engine/Source/Runtime/Experimental/IoStore/OnDemand/Private/OnDemandInstallCache.cpp b/Engine/Source/Runtime/Experimental/IoStore/OnDemand/Private/OnDemandInstallCache.cpp
--- a/Engine/Source/Runtime/Experimental/IoStore/OnDemand/Private/OnDemandInstallCache.cpp
+++ b/Engine/Source/Runtime/Experimental/IoStore/OnDemand/Private/OnDemandInstallCache.cpp
@@ -68,20 +68,6 @@ namespace CVars
 		GIoStoreOnDemandEnableDefrag,
 		TEXT("Whether to enable defrag when purging")
 	);
-
-	static bool GIoStoreOnDemandUseNewGetBlockInfo = true;
-	static FAutoConsoleVariableRef CVar_IoStoreOnDemandUseNewGetBlockInfo(
-		TEXT("iostore.UseNewGetBlockInfo"),
-		GIoStoreOnDemandUseNewGetBlockInfo,
-		TEXT("Use runtime blocksize tracking for GetBlockInfo")
-	);
-
-	static bool GIoStoreOnDemandVerifyNewGetBlockInfo = false;
-	static FAutoConsoleVariableRef CVar_IoStoreOnDemandVerifyNewGetBlockInfo(
-		TEXT("iostore.VerifyNewGetBlockInfo"),
-		GIoStoreOnDemandVerifyNewGetBlockInfo,
-		TEXT("Verify new GetBlockInfo against old GetBlockInfo")
-	);
 }
 
 ///////////////////////////////////////////////////////////////////////////////
@@ -177,6 +163,7 @@ struct FCasBlockInfo
 {
 	uint64	FileSize = 0;
 	int64	LastAccess = 0;
+	int64	LastModification = 0;
 	uint64	RefSize = 0;
 
 	bool operator==(const FCasBlockInfo& Other) const
@@ -184,6 +171,7 @@ struct FCasBlockInfo
 		return
 			FileSize == Other.FileSize && 
 			LastAccess == Other.LastAccess && 
+			LastModification == Other.LastModification &&
 			RefSize == Other.RefSize;
 	}
 };
@@ -198,6 +186,13 @@ enum class ECasTrackAccessType : uint8
 	Granular
 };
 
+///////////////////////////////////////////////////////////////////////////////
+enum class ECasTrackModificationType : uint8
+{
+	Always,
+	Newer
+};
+
 ///////////////////////////////////////////////////////////////////////////////
 struct FCasJournal
 {
@@ -251,7 +246,9 @@ struct FCasJournal
 			BlockCreated,
 			BlockDeleted,
 			BlockAccess,
-			CriticalError
+			CriticalError,
+			BlockModification,
+			// New values must be added at the end
 		};
 
 		struct FChunkLocation
@@ -276,9 +273,11 @@ struct FCasJournal
 		struct FCriticalError
 		{
 			EType		Type = EType::CriticalError;
+			uint8		Pad[3] = { 0 };
 			EErrorCode	ErrorCode = EErrorCode::None;
+			uint8		Pad1[16] = { 0 };
 		};
-		static_assert(sizeof(FBlockOperation) == 24);
+		static_assert(sizeof(FCriticalError) == 24);
 
 		union
 		{
@@ -289,6 +288,8 @@ struct FCasJournal
 
 		EType			Type() const { return *reinterpret_cast<const EType*>(this); }
 		static int64	Size() { return sizeof(FEntry); }
+
+		FEntry() {}
 	};
 	static_assert(sizeof(FEntry) == 24);
 
@@ -298,6 +299,7 @@ struct FCasJournal
 		void			BlockCreated(FCasBlockId BlockId);
 		void			BlockDeleted(FCasBlockId BlockId);
 		void			BlockAccess(FCasBlockId BlockId, int64 UtcTicks);
+		void			BlockModification(FCasBlockId BlockId, int64 UtcTicks);
 		void			CriticalError(FCasJournal::EErrorCode ErrorCode);
 
 		FString			JournalFile;
@@ -623,7 +625,7 @@ FResult FCasJournal::Commit(FTransaction&& Transaction)
 
 void FCasJournal::FTransaction::ChunkLocation(const FCasLocation& Location, const FCasAddr& Addr)
 {
-	Entries.AddZeroed_GetRef().ChunkLocation = FEntry::FChunkLocation
+	new(&Entries.AddDefaulted_GetRef().ChunkLocation) FEntry::FChunkLocation
 	{
 		.CasLocation = Location,
 		.CasAddr = Addr
@@ -632,7 +634,7 @@ void FCasJournal::FTransaction::ChunkLocation(const FCasLocation& Location, cons
 
 void FCasJournal::FTransaction::BlockCreated(FCasBlockId BlockId)
 {
-	Entries.AddZeroed_GetRef().BlockOperation = FEntry::FBlockOperation
+	new(&Entries.AddDefaulted_GetRef().BlockOperation) FEntry::FBlockOperation
 	{
 		.Type = FEntry::EType::BlockCreated,
 		.BlockId = BlockId,
@@ -642,7 +644,7 @@ void FCasJournal::FTransaction::BlockCreated(FCasBlockId BlockId)
 
 void FCasJournal::FTransaction::BlockDeleted(FCasBlockId BlockId)
 {
-	Entries.AddZeroed_GetRef().BlockOperation = FEntry::FBlockOperation
+	new(&Entries.AddDefaulted_GetRef().BlockOperation) FEntry::FBlockOperation
 	{
 		.Type = FEntry::EType::BlockDeleted,
 		.BlockId = BlockId,
@@ -652,7 +654,7 @@ void FCasJournal::FTransaction::BlockDeleted(FCasBlockId BlockId)
 
 void FCasJournal::FTransaction::BlockAccess(FCasBlockId BlockId, int64 UtcTicks)
 {
-	Entries.AddZeroed_GetRef().BlockOperation = FEntry::FBlockOperation
+	new(&Entries.AddDefaulted_GetRef().BlockOperation) FEntry::FBlockOperation
 	{
 		.Type = FEntry::EType::BlockAccess,
 		.BlockId = BlockId,
@@ -660,9 +662,19 @@ void FCasJournal::FTransaction::BlockAccess(FCasBlockId BlockId, int64 UtcTicks)
 	};
 }
 
+void FCasJournal::FTransaction::BlockModification(FCasBlockId BlockId, int64 UtcTicks)
+{
+	new(&Entries.AddDefaulted_GetRef().BlockOperation) FEntry::FBlockOperation
+	{
+		.Type = FEntry::EType::BlockModification,
+		.BlockId = BlockId,
+		.UtcTicks = UtcTicks
+	};
+}
+
 void FCasJournal::FTransaction::CriticalError(FCasJournal::EErrorCode ErrorCode)
 {
-	Entries.AddZeroed_GetRef().CriticalError = FEntry::FCriticalError
+	new(&Entries.AddDefaulted_GetRef().CriticalError) FEntry::FCriticalError
 	{
 		.Type = FEntry::EType::CriticalError,
 		.ErrorCode = ErrorCode
@@ -675,6 +687,7 @@ struct FCasBlockState
 	uint32 Size = 0;
 	uint32 ReadHandleCount = 0;
 	int64 LastAccessUtcTicks = 0;
+	int64 LastModificationUtcTicks = 0;
 };
 
 struct FCasSnapshot;
@@ -709,13 +722,18 @@ struct FCas
 		return TrackAccessIf(Type, BlockId, FDateTime::UtcNow().GetTicks(), bDirty);
 	};
 	bool						UnlockedTrackAccessIf(ECasTrackAccessType Type, uint32 BlockIdHash, FCasBlockId BlockId, int64 UtcTicks, bool bDirty);
+	bool						TrackModificationIf(ECasTrackModificationType Type, FCasBlockId BlockId, int64 UtcTicks);
+	bool						TrackModificationIf(ECasTrackModificationType Type, FCasBlockId BlockId)
+	{
+		return TrackModificationIf(Type, BlockId, FDateTime::UtcNow().GetTicks());
+	};
+	bool						UnlockedTrackModificationIf(ECasTrackModificationType Type, uint32 BlockIdHash, FCasBlockId BlockId, int64 UtcTicks);
 	void						UpdateBlock(
 									FCasBlockId BlockId, 
 									uint32 BytesWritten, 
 									TConstArrayView<FIoHash> ChunkHashes, 
 									TConstArrayView<int64> Offsets,
 									TFunctionRef<void(const FCasAddr&, const FCasLocation&)> OutAddrs);
-	uint64						GetBlockInfoOld(FCasBlockInfoMap& OutBlockInfo) const;
 	uint64						GetBlockInfo(FCasBlockInfoMap& OutBlockInfo) const;
 	uint64						GetTotalBlockSize() const;
 	FBlockFileSizes				FindAllBlockFiles() const;
@@ -1073,6 +1091,45 @@ bool FCas::UnlockedTrackAccessIf(const ECasTrackAccessType Type, const uint32 Bl
 	return false;
 }
 
+bool FCas::TrackModificationIf(ECasTrackModificationType Type, FCasBlockId BlockId, int64 UtcTicks)
+{
+	const uint32 BlockIdHash = GetTypeHash(BlockId);
+
+	UE::TUniqueLock Lock(Mutex);
+
+	return UnlockedTrackModificationIf(Type, BlockIdHash, BlockId, UtcTicks);
+}
+
+bool FCas::UnlockedTrackModificationIf(ECasTrackModificationType Type, uint32 BlockIdHash, FCasBlockId BlockId, int64 UtcTicks)
+{
+	check(BlockId.IsValid());
+
+	FCasBlockState* BlockState = BlockIds.FindByHash(BlockIdHash, BlockId);
+	if (ensure(BlockState) == false)
+	{
+		return false;
+	}
+
+	int64& FoundTicks = BlockState->LastModificationUtcTicks;
+	const int64 PrevTicks = FoundTicks;
+	
+	bool bUpdate = true;
+	switch (Type)
+	{
+	case ECasTrackModificationType::Newer:
+		bUpdate = PrevTicks < UtcTicks;
+		break;
+	}
+
+	if (bUpdate)
+	{
+		FoundTicks = UtcTicks;
+		return true;
+	}
+
+	return false;
+}
+
 void FCas::UpdateBlock(
 	FCasBlockId BlockId, 
 	uint32 BytesWritten, 
@@ -1100,92 +1157,9 @@ void FCas::UpdateBlock(
 	}
 }
 
-uint64 FCas::GetBlockInfoOld(FCasBlockInfoMap& OutBlockInfo) const
-{
-	TStringBuilder<256> Path;
-	FPathViews::Append(Path, RootDirectory, TEXT("blocks"));
-
-	struct FDirectoryVisitor final
-		: public IPlatformFile::FDirectoryVisitor
-	{
-		FDirectoryVisitor(IPlatformFile& PlatformFile, FCasBlockInfoMap& InBlockInfo, FLastAccess&& Access)
-			: Ipf(PlatformFile)
-			, BlockInfo(InBlockInfo)
-			, LastAccess(MoveTemp(Access))
-		{ }
-		
-		virtual bool Visit(const TCHAR* FilenameOrDirectory, bool bIsDirectory) override
-		{
-			if (bIsDirectory)
-			{
-				return true;
-			}
-
-			const FStringView Filename(FilenameOrDirectory);
-			if (FPathViews::GetExtension(Filename) == TEXTVIEW("ucas") == false)
-			{
-				return true;
-			}
-
-			const int64			FileSize = Ipf.FileSize(FilenameOrDirectory);
-			const FStringView	IndexHex = FPathViews::GetBaseFilename(Filename);
-			const FCasBlockId	BlockId(FParse::HexNumber(WriteToString<128>(IndexHex).ToString()));
-
-			if (BlockId.IsValid() == false || FileSize < 0)
-			{
-				UE_LOG(LogIoStoreOnDemand, Warning, TEXT("Found invalid CAS block '%s', FileSize=%" INT64_FMT),
-					FilenameOrDirectory, FileSize);
-				return true;
-			}
-
-			if (BlockInfo.Contains(BlockId))
-			{
-				UE_LOG(LogIoStoreOnDemand, Warning, TEXT("Found duplicate CAS block '%s'"), FilenameOrDirectory);
-				return true;
-			}
-
-			const int64* UtcTicks = LastAccess.Find(BlockId);
-
-			BlockInfo.Add(BlockId, FCasBlockInfo
-			{
-				.FileSize = uint64(FileSize),
-				.LastAccess = (UtcTicks && *UtcTicks) ? (*UtcTicks & ~FCas::DirtyTimestampMask) : TimestampForMissingLastAccess
-			});
-			TotalSize += uint64(FileSize);
-
-			return true;
-		}
-
-		IPlatformFile&		Ipf;
-		FCasBlockInfoMap&	BlockInfo;
-		FLastAccess			LastAccess;
-		const int64			TimestampForMissingLastAccess = FCas::GetTimestampForMissingLastAccess();
-		uint64				TotalSize = 0;
-	};
-
-	FLastAccess Access;
-	{
-		TUniqueLock Lock(Mutex);
-		Access.Reserve(BlockIds.Num());
-		Algo::Transform(BlockIds, Access, [](const TPair<FCasBlockId, FCasBlockState>& Kv) 
-		{
-			return FLastAccess::ElementType(Kv.Key, Kv.Value.LastAccessUtcTicks);
-		});
-	}
-	IPlatformFile& Ipf = FPlatformFileManager::Get().GetPlatformFile();
-	FDirectoryVisitor Visitor(Ipf, OutBlockInfo, MoveTemp(Access));
-	Ipf.IterateDirectory(Path.ToString(), Visitor);
-
-	return Visitor.TotalSize;
-}
 
 uint64 FCas::GetBlockInfo(FCasBlockInfoMap& OutBlockInfo) const
 {
-	if (CVars::GIoStoreOnDemandUseNewGetBlockInfo == false)
-	{
-		return GetBlockInfoOld(OutBlockInfo);
-	}
-
 	const int64	TimestampForMissingLastAccess = FCas::GetTimestampForMissingLastAccess();
 
 	uint64 TotalSize = 0;
@@ -1197,26 +1171,20 @@ uint64 FCas::GetBlockInfo(FCasBlockInfoMap& OutBlockInfo) const
 		{
 			const FCasBlockId BlockId = Kv.Key;
 			const uint32 FileSize = Kv.Value.Size;
-			const int64 UtcTicks = Kv.Value.LastAccessUtcTicks;
+			const int64 LastAccessUtcTicks = Kv.Value.LastAccessUtcTicks;
+			const int64 LastModification = Kv.Value.LastModificationUtcTicks;
 
 			OutBlockInfo.Add(BlockId, FCasBlockInfo
 			{
 				.FileSize = uint64(FileSize),
-				.LastAccess = UtcTicks ? (UtcTicks & ~FCas::DirtyTimestampMask) : TimestampForMissingLastAccess
+				.LastAccess = LastAccessUtcTicks ? (LastAccessUtcTicks & ~FCas::DirtyTimestampMask) : TimestampForMissingLastAccess,
+				.LastModification = LastModification
 			});
 
 			TotalSize += uint64(FileSize);
 		}
 	}
 
-	if (CVars::GIoStoreOnDemandVerifyNewGetBlockInfo)
-	{
-		FCasBlockInfoMap OldMap;
-		uint64 TotalSizeOld = GetBlockInfoOld(OldMap);
-		check(TotalSize == TotalSizeOld);
-		check(OutBlockInfo.OrderIndependentCompareEqual(OldMap))
-	}
-
 	return TotalSize;
 }
 
@@ -1315,53 +1283,60 @@ FResult FCas::Verify(TArray<FCasAddr>& OutRemovedChunksAddrs, TMap<FCasBlockId,
 
 	IPlatformFile& Ipf = FPlatformFileManager::Get().GetPlatformFile();
 
-	UE::TUniqueLock Lock(Mutex);
-
-	for (auto BlockIt = BlockIds.CreateIterator(); BlockIt; ++BlockIt)
+	auto VerifyTimestamp = [this, &Ipf, Now, TimestampValueIfMissing = TimestampForMissingLastAccess](
+		int64& InOutUtcTicks, const FCasBlockId BlockId, const TCHAR* TimestampName)
 	{
-		const FCasBlockId BlockId = BlockIt->Key;
-		FCasBlockState& BlockState = BlockIt->Value;
-
-		FString ErrorMessage;
-
-		if (const uint64* FileSize = BlockFiles.Find(BlockId))
+		if (InOutUtcTicks)
 		{
-			if (BlockState.LastAccessUtcTicks)
+			if (InOutUtcTicks > Now)
 			{
-				if (BlockState.LastAccessUtcTicks > Now)
-				{
-					const FString Filename = GetBlockFilename(BlockId);
-					UE_LOG(LogIoStoreOnDemand, Warning, TEXT("Found future last access time for CAS block '%s'"), *Filename);
-					BlockState.LastAccessUtcTicks = Now;
-				}
+				const FString Filename = FCas::GetBlockFilename(BlockId);
+				UE_LOG(LogIoStoreOnDemand, Warning, TEXT("Found future %s time for CAS block '%s'"), TimestampName, *Filename);
+				InOutUtcTicks = Now;
 			}
-			else
+		}
+		else
+		{
+			const FString Filename = FCas::GetBlockFilename(BlockId);
+
+	#if UE_ONDEMANDINSTALLCACHE_USE_MODTIME
+			FDateTime ModTime = Ipf.GetTimeStamp(*Filename);
+			if (ensure(ModTime > FDateTime::MinValue()))
 			{
-				const FString Filename = GetBlockFilename(BlockId);
-				
-#if UE_ONDEMANDINSTALLCACHE_USE_MODTIME
-				FDateTime ModTime = Ipf.GetTimeStamp(*Filename);
-				if (ensure(ModTime > FDateTime::MinValue()))
+				if (ModTime > Now)
 				{
-					if (ModTime > Now)
-					{
-						UE_LOG(LogIoStoreOnDemand, Warning, TEXT("Missing last access time and found future mod time so using now time for CAS block '%s'"), *Filename);
-						BlockState.LastAccessUtcTicks = Now;
-					}
-					else
-					{
-						UE_LOG(LogIoStoreOnDemand, Warning, TEXT("Missing last access time so using file mod time for CAS block '%s'"), *Filename);
-						BlockState.LastAccessUtcTicks = ModTime.GetTicks();
-					}
+					UE_LOG(LogIoStoreOnDemand, Warning, TEXT("Missing %s time and found future mod time so using now time for CAS block '%s'"), TimestampName, *Filename);
+					InOutUtcTicks = Now;
 				}
 				else
-#endif // UE_ONDEMANDINSTALLCACHE_USE_MODTIME
 				{
-					// Failed to get mod time
-					UE_LOG(LogIoStoreOnDemand, Warning, TEXT("Missing last access time for CAS block '%s'"), *Filename);
-					BlockState.LastAccessUtcTicks = TimestampForMissingLastAccess;
+					UE_LOG(LogIoStoreOnDemand, Warning, TEXT("Missing %s time so using file mod time for CAS block '%s'"), TimestampName, *Filename);
+					InOutUtcTicks = ModTime.GetTicks();
 				}
 			}
+			else
+	#endif // UE_ONDEMANDINSTALLCACHE_USE_MODTIME
+			{
+				// Failed to get mod time
+				UE_LOG(LogIoStoreOnDemand, Warning, TEXT("Missing %s time for CAS block '%s'"), TimestampName, *Filename);
+				InOutUtcTicks = TimestampValueIfMissing;
+			}
+		}
+	};
+
+	UE::TUniqueLock Lock(Mutex);
+
+	for (auto BlockIt = BlockIds.CreateIterator(); BlockIt; ++BlockIt)
+	{
+		const FCasBlockId BlockId = BlockIt->Key;
+		FCasBlockState& BlockState = BlockIt->Value;
+
+		FString ErrorMessage;
+
+		if (const uint64* FileSize = BlockFiles.Find(BlockId))
+		{
+			VerifyTimestamp(BlockState.LastAccessUtcTicks,			BlockId, TEXT("last access"));
+			VerifyTimestamp(BlockState.LastModificationUtcTicks,	BlockId, TEXT("modification"));
 
 			if (*FileSize > FCas::GetMaxBlockSize())
 			{
@@ -1479,11 +1454,14 @@ struct FCasSnapshot
 	{
 		Invalid = 0,
 		Initial,
+		AddModificationTime,
 
 		LatestPlusOne,
 		Latest = LatestPlusOne - 1
 	};
 
+	static const FGuid CustomVersionGuid;
+
 	struct FHeader
 	{
 		static const inline uint8 MagicSequence[16] = {'+', 'S', 'N', 'A', 'P', 'S', 'H', 'O', 'T', 'H', 'E', 'A', 'D', 'E', 'R', '+'};
@@ -1514,11 +1492,18 @@ struct FCasSnapshot
 		{
 			Ar << Block.BlockId;
 			Ar << Block.LastAccess;
+
+			if (Ar.CustomVer(FCasSnapshot::CustomVersionGuid) >= static_cast<int32>(EVersion::AddModificationTime))
+			{
+				Ar << Block.ModTime;
+			}
+
 			return Ar;
 		}
 
 		FCasBlockId BlockId;
 		int64		LastAccess = 0;
+		int64		ModTime = 0;
 	};
 
 	using FChunkLocation = TPair<FCasAddr, FCasLocation>;
@@ -1533,6 +1518,8 @@ struct FCasSnapshot
 	FCasBlockId							CurrentBlockId;
 };
 
+const FGuid FCasSnapshot::CustomVersionGuid(0x217bee7d, 0x23a54a23, 0x9ef080d2, 0x3df768b1);
+
 ///////////////////////////////////////////////////////////////////////////////
 bool FCasSnapshot::FHeader::IsValid() const
 {
@@ -1556,14 +1543,13 @@ bool FCasSnapshot::FFooter::IsValid() const
 
 TResult<FCasSnapshot> FCasSnapshot::FromJournal(const FString& JournalFile)
 {
-	FCas::FLookup		CasLookup;
-	FCas::FLastAccess	LastAccess;
-	TSet<FCasBlockId>	BlockIds;
-	FCasBlockId			CurrentBlockId;
+	FCas::FLookup				CasLookup;
+	TMap<FCasBlockId, FBlock>	BlockMap;
+	FCasBlockId					CurrentBlockId;
 
 	FResult ReplayResult = FCasJournal::Replay(
 		JournalFile,
-		[&CasLookup, &LastAccess, &BlockIds, &CurrentBlockId](const FCasJournal::FEntry& JournalEntry)
+		[&CasLookup, &BlockMap, &CurrentBlockId](const FCasJournal::FEntry& JournalEntry) -> FResult
 		{
 			switch(JournalEntry.Type())
 			{
@@ -1585,13 +1571,13 @@ TResult<FCasSnapshot> FCasSnapshot::FromJournal(const FString& JournalFile)
 			{
 				const FCasJournal::FEntry::FBlockOperation& Op = JournalEntry.BlockOperation;
 				CurrentBlockId = Op.BlockId;
-				BlockIds.Add(Op.BlockId);
+				BlockMap.Emplace(Op.BlockId, FBlock{ .BlockId = Op.BlockId });
 				break;
 			}
 			case FCasJournal::FEntry::EType::BlockDeleted:
 			{
 				const FCasJournal::FEntry::FBlockOperation& Op = JournalEntry.BlockOperation;
-				BlockIds.Remove(Op.BlockId);
+				BlockMap.Remove(Op.BlockId);
 				if (CurrentBlockId == Op.BlockId)
 				{
 					CurrentBlockId = FCasBlockId::Invalid; 
@@ -1601,7 +1587,29 @@ TResult<FCasSnapshot> FCasSnapshot::FromJournal(const FString& JournalFile)
 			case FCasJournal::FEntry::EType::BlockAccess:
 			{
 				const FCasJournal::FEntry::FBlockOperation& Op = JournalEntry.BlockOperation;
-				LastAccess.Add(Op.BlockId, Op.UtcTicks);
+				if (FBlock* Block = BlockMap.Find(Op.BlockId))
+				{
+					Block->LastAccess = Op.UtcTicks;
+				}
+				else
+				{
+					using namespace UE::UnifiedError::IoStoreOnDemand;
+					return MakeJournalError(ECasErrorCode::ReplayJournalFailed, EIoErrorCode::FileCorrupt, TEXT("Journal BlockAccess Op references an invalid block"));
+				}
+				break;
+			}
+			case FCasJournal::FEntry::EType::BlockModification:
+			{
+				const FCasJournal::FEntry::FBlockOperation& Op = JournalEntry.BlockOperation;
+				if (FBlock* Block = BlockMap.Find(Op.BlockId))
+				{
+					Block->ModTime = Op.UtcTicks;
+				}
+				else
+				{
+					using namespace UE::UnifiedError::IoStoreOnDemand;
+					return MakeJournalError(ECasErrorCode::ReplayJournalFailed, EIoErrorCode::FileCorrupt, TEXT("Journal BlockModification Op references an invalid block"));
+				}
 				break;
 			}
 			};
@@ -1617,16 +1625,7 @@ TResult<FCasSnapshot> FCasSnapshot::FromJournal(const FString& JournalFile)
 	const int64 TimestampForMissingLastAccess = FCas::GetTimestampForMissingLastAccess();
 
 	FCasSnapshot Snapshot;
-	Snapshot.Blocks.Reserve(BlockIds.Num());
-	for (FCasBlockId BlockId : BlockIds)
-	{
-		const int64* AccessTime = LastAccess.Find(BlockId);
-		Snapshot.Blocks.Add(FBlock
-		{
-			.BlockId	= BlockId,
-			.LastAccess = AccessTime != nullptr ? *AccessTime : TimestampForMissingLastAccess
-		});
-	}
+	BlockMap.GenerateValueArray(Snapshot.Blocks);
 
 	Snapshot.ChunkLocations = CasLookup.Array();
 	Snapshot.CurrentBlockId = CurrentBlockId;
@@ -1652,6 +1651,8 @@ TResult<int64> FCasSnapshot::Save(const FCasSnapshot& Snapshot, const FString& S
 	FMemory::Memcpy(&Header.Magic, &FHeader::MagicSequence, sizeof(FHeader::MagicSequence));
 	Header.Version = EVersion::Latest;
 
+	Ar->SetCustomVersion(CustomVersionGuid, static_cast<int32>(Header.Version), "CasSnapshotVersion");
+
 	Ar->Serialize(reinterpret_cast<uint8*>(&Header), FHeader::Size());
 	if (Ar->IsError())
 	{
@@ -1726,6 +1727,8 @@ TResult<FCasSnapshot> FCasSnapshot::Load(const FString& SnapshotFile, int64* Out
 		return MakeSnapshotError<FCasSnapshot>(ECasErrorCode::LoadSnapshotFailed, EIoErrorCode::ReadError, FString::Printf(TEXT("Failed to validate snapshot header in '%s'"), *SnapshotFile));
 	}
 
+	Ar->SetCustomVersion(CustomVersionGuid, static_cast<int32>(Header.Version), "CasSnapshotVersion");
+
 	FCasSnapshot Snapshot;
 	*Ar << Snapshot.Blocks;
 	*Ar << Snapshot.ChunkLocations;
@@ -1823,6 +1826,7 @@ void FCas::LoadSnapshot(FCasSnapshot&& Snapshot)
 	{
 		FCasBlockState& BlockState = BlockIds.Add(Block.BlockId);
 		BlockState.LastAccessUtcTicks = Block.LastAccess;
+		BlockState.LastModificationUtcTicks = Block.ModTime;
 	}
 }
 
@@ -1869,6 +1873,15 @@ FResult FCas::ApplyJournalEntry(const FCasJournal::FEntry& JournalEntry)
 		}
 		break;
 	}
+	case FCasJournal::FEntry::EType::BlockModification:
+	{
+		const FCasJournal::FEntry::FBlockOperation& Op = JournalEntry.BlockOperation;
+		if (ensure(TrackModificationIf(ECasTrackModificationType::Always, Op.BlockId, Op.UtcTicks)) == false)
+		{
+			using namespace UE::UnifiedError::IoStoreOnDemand;
+			return MakeJournalError(ECasErrorCode::ReplayJournalFailed, EIoErrorCode::FileCorrupt, TEXT("Journal BlockModification Op references an invalid block"));
+		}
+	}
 	};
 
 	return MakeValue();
@@ -2030,8 +2043,8 @@ private:
 													uint64* OutTotalFreedBytes = nullptr);
 	bool										Resolve(FIoRequestImpl* Request);
 	void										CompleteRequest(FIoRequestImpl* Request, EIoErrorCode Status);
-	FResult										FlushPendingChunks(FPendingChunks& Chunks, int64 UtcAccessTicks = 0);
-	FResult										FlushPendingChunksImpl(const FPendingChunks& Chunks, int64 UtcAccessTicks = 0);
+	FResult										FlushPendingChunks(FPendingChunks& Chunks, int64 UtcAccessTicks = 0, int64 UtcModificationTicks = 0);
+	FResult										FlushPendingChunksImpl(const FPendingChunks& Chunks, int64 UtcAccessTicks = 0, int64 UtcModificationTicks = 0);
 	FString										GetJournalFilename() const { return CacheDirectory / TEXT("cas.jrn"); }
 	FString										GetSnapshotFilename() const { return CacheDirectory / TEXT("cas.snp"); }
 	UE::UnifiedError::IoStoreOnDemand::FInstallCacheErrorContext MakeInstallCacheErrorContext(uint64 TotalCachedBytes = 0, uint64 RequestedBytes = 0, uint32 LineNo = __builtin_LINE());
@@ -2234,7 +2247,7 @@ void FOnDemandInstallCache::Shutdown()
 		else
 		{
 			UE_LOG(LogIoStoreOnDemand, Error, TEXT("Failed to create CAS snapshot from journal '%s', error '%s'"),
-				*JournalFile, *LexToString(SnapshotResult.GetError()));
+				*JournalFile, *SnapshotResult.GetError().GetErrorMessage(true).ToString());
 		}
 	}
 
@@ -3236,7 +3249,8 @@ FResult FOnDemandInstallCache::PurgeInternal(FCasBlockInfoMap& BlockInfo, const
 			return Result;
 		}
 
-		FOnDemandInstallCacheStats::OnBlockDeleted(Info.LastAccess, false /*FromDefrag*/);
+		constexpr bool bFromDefragFalse = false;
+		FOnDemandInstallCacheStats::OnBlockDeleted(Info.LastAccess, Info.LastModification, bFromDefragFalse);
 
 		// This should be the only thread writing to CurrentBlock
 		FCasBlockId MaybeCurrentBlock = BlockId;
@@ -3297,7 +3311,14 @@ FResult FOnDemandInstallCache::Defrag(
 	{
 		FCasBlockId BlockId;
 		int64 LastAccess = 0;
+		int64 LastModification = 0;
 		TArray<FDefragBlockReferencedChunk> ReferencedChunks;
+
+		FDefragBlock(const FCasBlockId InBlockId, const FCasBlockInfo& InBlockInfo)
+			: BlockId(InBlockId)
+			, LastAccess(InBlockInfo.LastAccess)
+			, LastModification(InBlockInfo.LastModification)
+		{ }
 	};
 
 	// Build the list of blocks to defrag and determine if its possible to free enough data through defragging
@@ -3335,7 +3356,7 @@ FResult FOnDemandInstallCache::Defrag(
 				FreedBlockBytes += Info.FileSize;
 				NewBlockBytes += Info.RefSize; // For now, assume that nothing will be moved to the current block
 
-				BlocksToDefrag.Add(FDefragBlock{ .BlockId = BlockId, .LastAccess = Info.LastAccess });
+				BlocksToDefrag.Emplace(BlockId, Info);
 
 				if (FreedBlockBytes >= NewBlockBytes && FreedBlockBytes - NewBlockBytes >= *TotalBytesToFree)
 				{
@@ -3352,7 +3373,7 @@ FResult FOnDemandInstallCache::Defrag(
 
 				BlocksToDefragTotalSize += Info.FileSize;
 
-				BlocksToDefrag.Add(FDefragBlock{ .BlockId = BlockId, .LastAccess = Info.LastAccess });
+				BlocksToDefrag.Emplace(BlockId, Info);
 			}
 		}
 
@@ -3385,7 +3406,7 @@ FResult FOnDemandInstallCache::Defrag(
 				BlocksToDefragFragmentedSize += (Info.FileSize - Info.RefSize);
 				BlocksToDefragTotalSize += Info.FileSize;
 
-				BlocksToDefrag.Add(FDefragBlock{ .BlockId = BlockId, .LastAccess = Info.LastAccess });
+				BlocksToDefrag.Emplace(BlockId, Info);
 			}
 			else if (Info.FileSize < Cas.GetMinBlockSize())
 			{
@@ -3397,7 +3418,7 @@ FResult FOnDemandInstallCache::Defrag(
 
 				BlocksToDefragTotalSize += Info.FileSize;
 
-				BlocksToDefrag.Add(FDefragBlock{ .BlockId = BlockId, .LastAccess = Info.LastAccess });
+				BlocksToDefrag.Emplace(BlockId, Info);
 			}
 		}
 
@@ -3557,7 +3578,7 @@ FResult FOnDemandInstallCache::Defrag(
 						UE_LOGFMT(LogIoStoreOnDemand, Error, "Failed to commit critical errors to CAS journal ({Error})", CommitResult.GetError());
 					}
 
-					if (FResult Result = FlushPendingChunks(DefragPendingChunks, DefragBlock.LastAccess); Result.HasError())
+					if (FResult Result = FlushPendingChunks(DefragPendingChunks, DefragBlock.LastAccess, DefragBlock.LastModification); Result.HasError())
 					{
 						FOnDemandInstallCacheStats::OnDefrag(Result, TotalFragmentedBytes, TotalCachedBytes, BlocksToDefragFragmentedSize, BlocksToDefragTotalSize);
 						return Result;
@@ -3577,7 +3598,7 @@ FResult FOnDemandInstallCache::Defrag(
 
 				if (DefragPendingChunks.TotalSize > FPendingChunks::MaxPendingBytes)
 				{
-					if (FResult Result = FlushPendingChunks(DefragPendingChunks, DefragBlock.LastAccess); Result.HasError())
+					if (FResult Result = FlushPendingChunks(DefragPendingChunks, DefragBlock.LastAccess, DefragBlock.LastModification); Result.HasError())
 					{
 						FOnDemandInstallCacheStats::OnDefrag(Result, TotalFragmentedBytes, TotalCachedBytes, BlocksToDefragFragmentedSize, BlocksToDefragTotalSize);
 						return Result;
@@ -3590,7 +3611,7 @@ FResult FOnDemandInstallCache::Defrag(
 
 			FileHandle.Reset();
 
-			if (FResult Result = FlushPendingChunks(DefragPendingChunks, DefragBlock.LastAccess); Result.HasError())
+			if (FResult Result = FlushPendingChunks(DefragPendingChunks, DefragBlock.LastAccess, DefragBlock.LastModification); Result.HasError())
 			{
 				FOnDemandInstallCacheStats::OnDefrag(Result, TotalFragmentedBytes, TotalCachedBytes, BlocksToDefragFragmentedSize, BlocksToDefragTotalSize);
 				return Result;
@@ -3613,7 +3634,8 @@ FResult FOnDemandInstallCache::Defrag(
 			return Result;
 		}
 
-		FOnDemandInstallCacheStats::OnBlockDeleted(DefragBlock.LastAccess, true /*FromDefrag*/);
+		constexpr bool bFromDefragTrue = true;
+		FOnDemandInstallCacheStats::OnBlockDeleted(DefragBlock.LastAccess, DefragBlock.LastModification, bFromDefragTrue);
 
 		Transaction.BlockDeleted(DefragBlock.BlockId);
 
@@ -3745,7 +3767,7 @@ FOnDemandInstallCacheUsage FOnDemandInstallCache::GetCacheUsage()
 	};
 }
 
-FResult FOnDemandInstallCache::FlushPendingChunks(FPendingChunks& Chunks, int64 UtcAccessTicks)
+FResult FOnDemandInstallCache::FlushPendingChunks(FPendingChunks& Chunks, int64 UtcAccessTicks, int64 UtcModificationTicks)
 {
 	if (Chunks.IsEmpty())
 	{
@@ -3775,11 +3797,11 @@ FResult FOnDemandInstallCache::FlushPendingChunks(FPendingChunks& Chunks, int64
 		Chunks.Reset();
 	};
 
-	return FlushPendingChunksImpl(Chunks, UtcAccessTicks);
+	return FlushPendingChunksImpl(Chunks, UtcAccessTicks, UtcModificationTicks);
 #endif // UE_ONDEMANDINSTALLCACHE_EXCLUSIVE_WRITE
 }
 
-FResult FOnDemandInstallCache::FlushPendingChunksImpl(const FPendingChunks& Chunks, int64 UtcAccessTicks)
+FResult FOnDemandInstallCache::FlushPendingChunksImpl(const FPendingChunks& Chunks, int64 UtcAccessTicks, int64 UtcModificationTicks)
 {
 	using namespace UE::UnifiedError::IoStoreOnDemand;
 
@@ -3892,6 +3914,20 @@ FResult FOnDemandInstallCache::FlushPendingChunksImpl(const FPendingChunks& Chun
 				Transaction.BlockAccess(CurrentBlockId, Now);
 			}
 
+			if (UtcModificationTicks)
+			{
+				if (Cas.TrackModificationIf(ECasTrackModificationType::Newer, CurrentBlockId, UtcModificationTicks))
+				{
+					Transaction.BlockModification(CurrentBlockId, UtcModificationTicks);
+				}
+			}
+			else
+			{
+				const int64 Now = FDateTime::UtcNow().GetTicks();
+				check(Cas.TrackModificationIf(ECasTrackModificationType::Always, CurrentBlockId, Now));
+				Transaction.BlockModification(CurrentBlockId, Now);
+			}
+
 			Cas.UpdateBlock(CurrentBlockId, IntCastChecked<uint32>(BytesToWrite), ChunkHashes, Offsets,
 				[&Transaction](const FCasAddr& CasAddr, const FCasLocation& Loc)
 				{
@@ -4239,6 +4275,55 @@ TEST_CASE("IoStore::OnDemand::InstallCache::Journal", "[IoStoreOnDemand][Install
 		CHECK(BlockId == ExpectedBlockId);
 		CHECK(Ticks == ExpectedTicks);
 	}
+
+	SECTION("ReplayBlockModification")
+	{
+		// Arrange
+		const FCasBlockId ExpectedBlockId(462);
+		const uint64 ExpectedTicks = FDateTime::UtcNow().GetTicks();
+
+		// Act
+		FTmpDirectoryScope _(TestBaseDir);
+		const FString JournalFile = TestBaseDir / TEXT("test.jrn");
+
+		FResult Result = FCasJournal::Create(JournalFile);
+		CHECK(Result.HasValue());
+
+		FCasJournal::FTransaction Tx = FCasJournal::Begin(JournalFile);
+		Tx.BlockModification(ExpectedBlockId, ExpectedTicks);
+
+		Result = FCasJournal::Commit(MoveTemp(Tx));
+		CHECK(Result.HasValue());
+
+		// Assert
+		FCasBlockId BlockId;
+		uint64 Ticks = 0;
+
+		Result = FCasJournal::Replay(
+			JournalFile,
+			[&BlockId, &Ticks](const FCasJournal::FEntry& JournalEntry)
+			{
+				switch (JournalEntry.Type())
+				{
+				case FCasJournal::FEntry::EType::BlockModification:
+				{
+					const FCasJournal::FEntry::FBlockOperation& Op = JournalEntry.BlockOperation;
+					BlockId = Op.BlockId;
+					Ticks = Op.UtcTicks;
+					break;
+				}
+				default:
+					CHECK(false);
+					break;
+				};
+
+				return MakeValue();
+			});
+
+		CHECK(Result.HasValue());
+		CHECK(BlockId == ExpectedBlockId);
+		CHECK(Ticks == ExpectedTicks);
+	}
 }
 
 TEST_CASE("IoStore::OnDemand::InstallCache::Snapshot", "[IoStoreOnDemand][InstallCache]")
@@ -4250,12 +4335,15 @@ TEST_CASE("IoStore::OnDemand::InstallCache::Snapshot", "[IoStoreOnDemand][Instal
 		// Arrange
 		FCasSnapshot ExpectedSnapshot;
 
+		const int64 Now = FDateTime::UtcNow().GetTicks();
+
 		for (uint32 Id = 1; Id <= 10; ++Id)
 		{
 			ExpectedSnapshot.Blocks.Add(FCasSnapshot::FBlock
 			{
 				.BlockId	= FCasBlockId(Id),
-				.LastAccess = FDateTime::UtcNow().GetTicks()
+				.LastAccess = Now + Id,
+				.ModTime = Now + 10 + Id
 			});
 
 			for (uint32 Idx = 1; Idx <= 10; ++Idx)
@@ -4284,6 +4372,7 @@ TEST_CASE("IoStore::OnDemand::InstallCache::Snapshot", "[IoStoreOnDemand][Instal
 		{
 			CHECK(Snapshot.Blocks[Idx].BlockId == ExpectedSnapshot.Blocks[Idx].BlockId);
 			CHECK(Snapshot.Blocks[Idx].LastAccess == ExpectedSnapshot.Blocks[Idx].LastAccess);
+			CHECK(Snapshot.Blocks[Idx].ModTime == ExpectedSnapshot.Blocks[Idx].ModTime);
 		}
 		CHECK(Snapshot.ChunkLocations.Num() == ExpectedSnapshot.ChunkLocations.Num());
 		for (int32 Idx = 0; Idx < Snapshot.ChunkLocations.Num(); ++Idx)
@@ -4306,8 +4395,15 @@ TEST_CASE("IoStore::OnDemand::InstallCache::Snapshot", "[IoStoreOnDemand][Instal
 
 		FCasJournal::FTransaction Tx = FCasJournal::Begin(JournalFile);
 
+		const int64 Now = FDateTime::UtcNow().GetTicks();
+		const int64 ExpectedModTime = Now + 2;
+		const int64 ExpectedLastAccess = Now + 3;
+
 		// Add a block and some chunk locations
-		Tx.BlockCreated(FCasBlockId(1));
+		FCasBlockId Block1(1);
+		Tx.BlockCreated(Block1);
+		Tx.BlockModification(Block1, Now);
+		Tx.BlockAccess(Block1, Now + 1);
 		for (int32 Idx = 1; Idx <= 10; ++Idx)
 		{
 			Tx.ChunkLocation(FCasLocation
@@ -4327,6 +4423,8 @@ TEST_CASE("IoStore::OnDemand::InstallCache::Snapshot", "[IoStoreOnDemand][Instal
 
 		// Add a second block and some chunk locations
 		Tx.BlockCreated(ExpectedCurrentBlockId);
+		Tx.BlockModification(ExpectedCurrentBlockId, ExpectedModTime);
+		Tx.BlockAccess(ExpectedCurrentBlockId, ExpectedLastAccess);
 		for (int32 Idx = 1; Idx <= 10; ++Idx)
 		{
 			Tx.ChunkLocation(FCasLocation
@@ -4346,6 +4444,8 @@ TEST_CASE("IoStore::OnDemand::InstallCache::Snapshot", "[IoStoreOnDemand][Instal
 		// Assert
 		CHECK(Snapshot.CurrentBlockId == ExpectedCurrentBlockId);
 		CHECK(Snapshot.Blocks.Num() == 1); 
+		CHECK(Snapshot.Blocks[0].ModTime == ExpectedModTime);
+		CHECK(Snapshot.Blocks[0].LastAccess == ExpectedLastAccess);
 		CHECK(Snapshot.ChunkLocations.Num() == 10); 
 		for (int32 Idx = 1; Idx < Snapshot.ChunkLocations.Num(); ++Idx)
 		{
diff --git a/Engine/Source/Runtime/Experimental/IoStore/OnDemand/Private/Statistics.cpp b/Engine/Source/Runtime/Experimental/IoStore/OnDemand/Private/Statistics.cpp
--- a/Engine/Source/Runtime/Experimental/IoStore/OnDemand/Private/Statistics.cpp
+++ b/Engine/Source/Runtime/Experimental/IoStore/OnDemand/Private/Statistics.cpp
@@ -1747,26 +1747,36 @@ void FOnDemandInstallCacheStats::OnDefrag(
 	SendImmediateAnalytic(MoveTemp(Event));
 }
 
-void FOnDemandInstallCacheStats::OnBlockDeleted(int64 LastAccessTicks, bool bFromDefrag)
+void FOnDemandInstallCacheStats::OnBlockDeleted(int64 LastAccessTicks, int64 LastModificationTicks, bool bFromDefrag)
 {
 	const FDateTime Now = FDateTime::UtcNow();
 	const FDateTime BlockAccess(LastAccessTicks);
+	const FDateTime BlockModification(LastModificationTicks);
 
 	// Theoretically all of the following checks shouldn't be necessary, as they should be 
 	// impossible. Emperical data is not fun :(
-	FTimespan BlockAge = FTimespan::MaxValue();
+	FTimespan BlockAccessAge = FTimespan::MaxValue();
 	if (ensure(BlockAccess.GetTicks() >= 0) &&
 		ensure(Now.GetTicks() >= 0) &&
 		ensure(Now >= BlockAccess))
 	{
-		BlockAge = Now - BlockAccess;
+		BlockAccessAge = Now - BlockAccess;
+	}
+
+	FTimespan BlockModAge = FTimespan::MaxValue();
+	if (ensure(BlockModification.GetTicks() >= 0) &&
+		ensure(Now.GetTicks() >= 0) &&
+		ensure(Now >= BlockModification))
+	{
+		BlockModAge = Now - BlockModification;
 	}
 
 	SendImmediateAnalytic(FOnDemandImmediateAnalytic
 	{
 		TEXT("IadCacheBlockDeleted"),
 		MakeAnalyticsEventAttributeArray(
-			TEXT("BlockAgeMinutes"), BlockAge.GetTotalMinutes(),
+			TEXT("BlockAgeMinutes"), BlockAccessAge.GetTotalMinutes(),
+			TEXT("BlockModAgeMinutes"), BlockModAge.GetTotalMinutes(),
 			TEXT("FromDefrag"), bFromDefrag
 		)
 	});
diff --git a/Engine/Source/Runtime/Experimental/IoStore/OnDemand/Private/Statistics.h b/Engine/Source/Runtime/Experimental/IoStore/OnDemand/Private/Statistics.h
--- a/Engine/Source/Runtime/Experimental/IoStore/OnDemand/Private/Statistics.h
+++ b/Engine/Source/Runtime/Experimental/IoStore/OnDemand/Private/Statistics.h
@@ -181,7 +181,7 @@ public:
 		uint64 TotalCachedBytes,
 		uint64 BlocksToDefragFragmentedSize, 
 		uint64 BlocksToDefragTotalSize) IAS_STATISTICS_IMPL()
-	static void OnBlockDeleted(int64 LastAccessTicks, bool bFromDefrag) IAS_STATISTICS_IMPL()
+	static void OnBlockDeleted(int64 LastAccessTicks, int64 LastModificationTicks, bool bFromDefrag) IAS_STATISTICS_IMPL()
 	static void OnCacheUsage(
 		uint64 MaxCacheSize,
 		uint64 CacheSize,
