Optimize Collision Detection by running Narrow Phase without waiting end of Broad Phase if no MidPhase callbacks are called.

#Synced-CL 49560376

diff --git a/Engine/Source/Runtime/Experimental/Chaos/Private/Chaos/PBDRigidsEvolutionGBF.cpp b/Engine/Source/Runtime/Experimental/Chaos/Private/Chaos/PBDRigidsEvolutionGBF.cpp
--- a/Engine/Source/Runtime/Experimental/Chaos/Private/Chaos/PBDRigidsEvolutionGBF.cpp
+++ b/Engine/Source/Runtime/Experimental/Chaos/Private/Chaos/PBDRigidsEvolutionGBF.cpp
@@ -1295,8 +1295,10 @@ void FPBDRigidsEvolutionGBF::DispatchKinematicTargetsTasks(int32 DispatchBatchIn
 
 void FPBDRigidsEvolutionGBF::ApplyMidPhaseModifier(FReal Dt)
 {
-	if (MidPhaseModifiers != nullptr)
+	if (MidPhaseModifiers != nullptr && !MidPhaseModifiers->IsEmpty())
 	{
+		UE::Tasks::Wait(CollisionDetector.GetBroadPhase().GetBroadToNarrowTasks());
+
 		// Consolidate midphases from parallel threads into the main lists held by the allocator.
 		// This makes them accessible to MidPhaseModifiers.
 		CollisionConstraints.GetConstraintAllocator().ProcessNewMidPhases();
diff --git a/Engine/Source/Runtime/Experimental/Chaos/Public/Chaos/Collision/SpatialAccelerationBroadPhase.h b/Engine/Source/Runtime/Experimental/Chaos/Public/Chaos/Collision/SpatialAccelerationBroadPhase.h
--- a/Engine/Source/Runtime/Experimental/Chaos/Public/Chaos/Collision/SpatialAccelerationBroadPhase.h
+++ b/Engine/Source/Runtime/Experimental/Chaos/Public/Chaos/Collision/SpatialAccelerationBroadPhase.h
@@ -387,7 +387,6 @@ namespace Chaos
 				FChaosVDContextWrapper CVDContext;
 				CVD_GET_WRAPPED_CURRENT_CONTEXT(CVDContext);
 
-				TArray<UE::Tasks::FTask> NextPendingTasks;
 				const int32 NumContextTask = PendingTasks.Num();
 				check(NumContextTask <= NumActiveBroadphaseContexts);
 				// If bSingleWorkerPhysics is true, NumActiveBroadphaseContexts shouldn't bigger than 1
@@ -400,6 +399,7 @@ namespace Chaos
 				}
 				else
 				{
+					BroadToNarrowTasks.Reset();
 					for (int32 ContextIndex = 0; ContextIndex < NumContextTask; ContextIndex++)
 					{
 						Private::FBroadPhaseContext& BroadPhaseContext = BroadphaseContexts[ContextIndex];
@@ -409,11 +409,8 @@ namespace Chaos
 							CVD_SCOPE_CONTEXT(CVDContext.Context);
 							AssignMidPhases(BroadPhaseContext);
 						}, PendingTasks[ContextIndex]);
-						NextPendingTasks.Add(PendingTask);
+						BroadToNarrowTasks.Add(PendingTask);
 					}
-					// It is required to wait for the middle phase modifiers, which is a callback to modify middle phases.
-					// This could be removed by adding dependent tasks for the middle phase modifiers.
-					UE::Tasks::Wait(NextPendingTasks);
 				}
 			}
 
@@ -458,21 +455,22 @@ namespace Chaos
 				for (int32 ContextIndex = 0; ContextIndex < NumContextTask; ContextIndex++)
 				{
 					Private::FBroadPhaseContext& BroadphaseContext = BroadphaseContexts[ContextIndex];
-					// Accessing the overlaps array is safe here, all overlaps have been finished to be computed at this point
-					if (!BroadphaseContext.Overlaps.IsEmpty())
+					UE::Tasks::FTask ProcessMidPhaseTask = UE::Tasks::Launch(UE_SOURCE_LOCATION, [this, &BroadphaseContext, Dt, CVDContext]()
 					{
-						UE::Tasks::FTask ProcessMidPhaseTask = UE::Tasks::Launch(UE_SOURCE_LOCATION, [this, &BroadphaseContext, Dt, CVDContext]()
-						{
-							CVD_SCOPE_CONTEXT(CVDContext.Context);
-							ProcessMidPhases(Dt, BroadphaseContext);
-						}, UE::Tasks::ETaskPriority::High);
-						PendingTasks.Add(ProcessMidPhaseTask);
-						ContextsWithConstraints.Add(ContextIndex);
-					}
+						CVD_SCOPE_CONTEXT(CVDContext.Context);
+						ProcessMidPhases(Dt, BroadphaseContext);
+					}, BroadToNarrowTasks[ContextIndex], UE::Tasks::ETaskPriority::High);
+					PendingTasks.Add(ProcessMidPhaseTask);
+					ContextsWithConstraints.Add(ContextIndex);
 				}
 			}
 		}
 
+		const TArray<UE::Tasks::FTask>& GetBroadToNarrowTasks() const 
+		{ 
+			return BroadToNarrowTasks; 
+		}
+
 		void GatherConstraints(bool bIsDeterministic)
 		{
 			if (bSingleWorkerPhysics)
@@ -800,7 +798,7 @@ namespace Chaos
 								ProduceParticleOverlaps<bOnlyRigid>(Dt, TransientHandle.Handle(), InSpatialAcceleration, BroadPhaseContext);
 							}
 						}
-					});
+					}, UE::Tasks::ETaskPriority::High);
 					PendingTasks.Add(PendingTask);
 				}
 				else
@@ -1079,6 +1077,7 @@ namespace Chaos
 		int32 NumMidPhases;
 
 		TArray<UE::Tasks::FTask> PendingTasks;
+		TArray<UE::Tasks::FTask> BroadToNarrowTasks;
 		// Contexts indexes in which there are constraints
 		TArray<int32> ContextsWithConstraints;
 	};
