IAD Cache - Fix UtcModificationTicks not being forwarded when UE_ONDEMANDINSTALLCACHE_EXCLUSIVE_WRITE is enabled.

#Synced-CL 49554261

diff --git a/Engine/Source/Runtime/Experimental/IoStore/OnDemand/Private/OnDemandInstallCache.cpp b/Engine/Source/Runtime/Experimental/IoStore/OnDemand/Private/OnDemandInstallCache.cpp
--- a/Engine/Source/Runtime/Experimental/IoStore/OnDemand/Private/OnDemandInstallCache.cpp
+++ b/Engine/Source/Runtime/Experimental/IoStore/OnDemand/Private/OnDemandInstallCache.cpp
@@ -3775,7 +3775,7 @@ FResult FOnDemandInstallCache::FlushPendingChunks(FPendingChunks& Chunks, int64
 	}
 
 #if UE_ONDEMANDINSTALLCACHE_EXCLUSIVE_WRITE
-	UE::Tasks::TTask<FResult> Task = ExclusivePipe.Launch(UE_SOURCE_LOCATION, [this, &Chunks, UtcAccessTicks]
+	UE::Tasks::TTask<FResult> Task = ExclusivePipe.Launch(UE_SOURCE_LOCATION, [this, &Chunks, UtcAccessTicks, UtcModificationTicks]
 	{
 		ON_SCOPE_EXIT 
 		{ 
@@ -3783,7 +3783,7 @@ FResult FOnDemandInstallCache::FlushPendingChunks(FPendingChunks& Chunks, int64
 			Chunks.Reset();
 		};
 
-		return FlushPendingChunksImpl(Chunks, UtcAccessTicks);
+		return FlushPendingChunksImpl(Chunks, UtcAccessTicks, UtcModificationTicks);
 	}, UE::Tasks::ETaskPriority::BackgroundHigh);
 
 	Task.Wait();
