Minor fix for  Collision Filter Refactor Part 5.5. Fixed a UE_INTERNAL warning from a inlined header function
#Synced-CL 49621724

diff --git a/Engine/Source/Runtime/Experimental/ChaosVisualDebugger/Public/DataWrappers/ChaosVDQueryDataWrappers.h b/Engine/Source/Runtime/Experimental/ChaosVisualDebugger/Public/DataWrappers/ChaosVDQueryDataWrappers.h
--- a/Engine/Source/Runtime/Experimental/ChaosVisualDebugger/Public/DataWrappers/ChaosVDQueryDataWrappers.h
+++ b/Engine/Source/Runtime/Experimental/ChaosVisualDebugger/Public/DataWrappers/ChaosVDQueryDataWrappers.h
@@ -79,7 +79,9 @@ struct FChaosVDCollisionObjectQueryParams : public FChaosVDWrapperDataBase
 	template <typename TOther>
 	void CopyFrom(const TOther& Other)
 	{
+		PRAGMA_DISABLE_INTERNAL_WARNINGS
 		ObjectTypesToQuery64 = Other.GetObjectTypesToQuery();
+		PRAGMA_ENABLE_INTERNAL_WARNINGS
 		IgnoreMask = Other.IgnoreMask;
 		bHasValidData = true;
 	}
