Chaos: Add support for max force in joint angular drive when using SimD.

#Synced-CL 49268397

diff --git a/Engine/Source/Runtime/Experimental/Chaos/Private/Chaos/Joint/PBDJointCachedSolverGaussSeidel.cpp b/Engine/Source/Runtime/Experimental/Chaos/Private/Chaos/Joint/PBDJointCachedSolverGaussSeidel.cpp
--- a/Engine/Source/Runtime/Experimental/Chaos/Private/Chaos/Joint/PBDJointCachedSolverGaussSeidel.cpp
+++ b/Engine/Source/Runtime/Experimental/Chaos/Private/Chaos/Joint/PBDJointCachedSolverGaussSeidel.cpp
@@ -2802,7 +2802,8 @@ void FPBDJointCachedSolver::ApplyRotationDrivesSimd(
 		VectorMultiplyAdd(RotationDrives.Simd.ConstraintSoftDamping, CVDt, RotationDrives.Simd.ConstraintLambda)), RotationDrives.Simd.ConstraintSoftIM));
 	RotationDrives.Simd.ConstraintLambda = VectorAdd(RotationDrives.Simd.ConstraintLambda, DeltaLambda);
 
-	// Should check MaxLambda eventually
+	// Check and apply MaxLambda 
+	RotationDrives.ApplyMaxLambdaSIMD(DeltaLambda, RotationDrives.Simd.ConstraintLambda);
 
 	VectorRegister4Float DeltaLambdas[3];
 	DeltaLambdas[0] = VectorReplicate(DeltaLambda, 0);
@@ -3007,7 +3008,15 @@ void FAxisConstraintDatas::SetMaxForce(
 		FReal MaxLambda = InMaxForce * Dt * Dt;
 		if (bAccelerationMode)
 		{
-			MaxLambda /= Data.ConstraintHardIM[ConstraintIndex];
+			FReal ConstraintIM = Data.ConstraintHardIM[ConstraintIndex];
+			if (bUseSimd)
+			{
+				ConstraintIM = FReal(VectorGetComponentDynamic(Simd.ConstraintHardIM, ConstraintIndex));
+			}
+			if (ConstraintIM > UE_SMALL_NUMBER)
+			{
+				MaxLambda /= ConstraintIM;
+			}
 		}
 		ConstraintMaxLambda[ConstraintIndex] = MaxLambda;
 	}
@@ -3033,6 +3042,39 @@ void FAxisConstraintDatas::ApplyMaxLambda(
 	}
 }
 
+void FAxisConstraintDatas::ApplyMaxLambdaSIMD(
+	VectorRegister4Float& DeltaLambda,
+	VectorRegister4Float& Lambda)
+{
+	// 1. Load MaxLambdas
+	VectorRegister4Float MaxLambda = MakeVectorRegisterFloatFromDouble(MakeVectorRegister(ConstraintMaxLambda[0], ConstraintMaxLambda[1], ConstraintMaxLambda[2], 0.0f));
+
+	// 2. Prepare -MaxLambda
+	VectorRegister4Float NegMaxLambda = VectorNegate(MaxLambda);
+
+	// 3. Create a mask where MaxLambda is greater than zero
+	VectorRegister4Float Zero = VectorZero();
+	VectorRegister4Float IsActiveMask = VectorCompareGT(MaxLambda, Zero);
+
+	// 4. If (Lambda > MaxLambda) -> Clamp Top
+	VectorRegister4Float TooHighMask = VectorCompareGT(Lambda, MaxLambda);
+
+	// 5. else if (Lambda < -MaxLambda) -> Clamp Bottom
+	VectorRegister4Float TooLowMask = VectorCompareLT(Lambda, NegMaxLambda);
+
+	// 6. Calculate clamped Lambda values
+	// FinalLambda = (Lambda > Max) ? Max : ((Lambda < -Max) ? -Max : Lambda)
+	VectorRegister4Float ClampedLambda = VectorSelect(TooHighMask, MaxLambda, Lambda);
+	ClampedLambda = VectorSelect(TooLowMask, NegMaxLambda, ClampedLambda);
+
+	// 7. Apply check (ConstraintMaxLambda > 0)
+	VectorRegister4Float OriginalLambda = Lambda;
+	Lambda = VectorSelect(IsActiveMask, ClampedLambda, Lambda);
+
+	// 8. Adjust DeltaLambda
+	DeltaLambda = VectorAdd(DeltaLambda, VectorSubtract(Lambda, OriginalLambda));
+}
+
 }
 
 
diff --git a/Engine/Source/Runtime/Experimental/Chaos/Public/Chaos/Joint/PBDJointCachedSolverGaussSeidel.h b/Engine/Source/Runtime/Experimental/Chaos/Public/Chaos/Joint/PBDJointCachedSolverGaussSeidel.h
--- a/Engine/Source/Runtime/Experimental/Chaos/Public/Chaos/Joint/PBDJointCachedSolverGaussSeidel.h
+++ b/Engine/Source/Runtime/Experimental/Chaos/Public/Chaos/Joint/PBDJointCachedSolverGaussSeidel.h
@@ -62,6 +62,11 @@ struct FAxisConstraintDatas
 		const int32 ConstraintIndex,
 		FReal& DeltaLambda, 
 		FReal& Lambda);
+
+	/** Apply the impulse limits to the impulse delta and net impulse in SIMD */
+	void ApplyMaxLambdaSIMD(
+		VectorRegister4Float& DeltaLambda,
+		VectorRegister4Float& Lambda);
 	
 	
 
