Fixed static analysis warnings for new 64 bit collision filters.
#Synced-CL 49621724

diff --git a/Engine/Source/Runtime/Engine/Public/Physics/PhysicsFiltering.h b/Engine/Source/Runtime/Engine/Public/Physics/PhysicsFiltering.h
--- a/Engine/Source/Runtime/Engine/Public/Physics/PhysicsFiltering.h
+++ b/Engine/Source/Runtime/Engine/Public/Physics/PhysicsFiltering.h
@@ -245,7 +245,8 @@ inline FCollisionResponseContainer ExtractSimCollisionResponseContainer(const FC
 	
 	for (int32 ChannelIndex = 0; ChannelIndex < UE_ARRAY_COUNT(CollisionResponseContainer.EnumArray); ++ChannelIndex)
 	{
-		if ((InSimFilterData.Word1) & (1 << ChannelIndex))
+		const uint64 ChannelBit = 1LLU << ChannelIndex;
+		if (uint64(InSimFilterData.Word1) & ChannelBit)
 		{
 			CollisionResponseContainer.EnumArray[ChannelIndex] = ECR_Block;
 		}
@@ -265,11 +266,12 @@ inline FCollisionResponseContainer ExtractQueryCollisionResponseContainer(const
 	
 	for (int32 ChannelIndex = 0; ChannelIndex < UE_ARRAY_COUNT(CollisionResponseContainer.EnumArray); ++ChannelIndex)
 	{
-		if ((InQueryFilterData.Word1) & (1 << ChannelIndex))
+		const uint64 ChannelBit = 1LLU << ChannelIndex;
+		if (uint64(InQueryFilterData.Word1) & ChannelBit)
 		{
 			CollisionResponseContainer.EnumArray[ChannelIndex] = ECR_Block;
 		}
-		else if ((InQueryFilterData.Word2) & (1 << ChannelIndex))
+		else if (uint64(InQueryFilterData.Word2) & ChannelBit)
 		{
 			CollisionResponseContainer.EnumArray[ChannelIndex] = ECR_Overlap;
 		}
diff --git a/Engine/Source/Runtime/Experimental/Chaos/Private/Chaos/ParticleDirtyFlags.cpp b/Engine/Source/Runtime/Experimental/Chaos/Private/Chaos/ParticleDirtyFlags.cpp
--- a/Engine/Source/Runtime/Experimental/Chaos/Private/Chaos/ParticleDirtyFlags.cpp
+++ b/Engine/Source/Runtime/Experimental/Chaos/Private/Chaos/ParticleDirtyFlags.cpp
@@ -5,15 +5,11 @@
 namespace Chaos
 {
 	PRAGMA_DISABLE_DEPRECATION_WARNINGS
-	constexpr uint32 QueryDataStart = offsetof(FCollisionData, QueryData);
-	constexpr uint32 SimDataStart = offsetof(FCollisionData, SimData);
-	constexpr uint32 SimDataEnd = SimDataStart + sizeof(FCollisionFilterData);
-	constexpr uint32 ShapeFilterDataStart = QueryDataStart + offsetof(FCollisionFilterData, Word2);
-	// Validate the total size is sufficient to store the two new structs
-	static_assert(sizeof(Filter::FInstanceData) + sizeof(Filter::FShapeFilterData) <= 2 * sizeof(FCollisionFilterData));
-	// Validate the instance data and shape filters fit in the correct spots
-	static_assert(sizeof(Filter::FInstanceData) <= ShapeFilterDataStart - QueryDataStart);
-	static_assert(sizeof(Filter::FShapeFilterData) <= SimDataEnd - ShapeFilterDataStart);
+	// The new instance data and shape filter data are stored in the old Query/Sim Data until deprecation allows removing the old fields.
+	static_assert(sizeof(FCollisionFilterData::Word0) + sizeof(FCollisionFilterData::Word1) == sizeof(Filter::FInstanceData));
+	constexpr size_t ShapeFilterPart1Size = sizeof(FCollisionFilterData::Word2) + sizeof(FCollisionFilterData::Word3);
+	constexpr size_t ShapeFilterPart2Size = sizeof(Filter::FShapeFilterData) - ShapeFilterPart1Size;
+	static_assert(ShapeFilterPart2Size <= sizeof(FCollisionFilterData));
 	PRAGMA_ENABLE_DEPRECATION_WARNINGS
 
 	FCollisionFilterData FCollisionData::GetQueryData() const
@@ -44,7 +40,8 @@ namespace Chaos
 	{
 		PRAGMA_DISABLE_DEPRECATION_WARNINGS
 		Filter::FShapeFilterData ShapeFilter;
-		memcpy(&ShapeFilter, &QueryData.Word2, sizeof(Filter::FShapeFilterData));
+		memcpy(&ShapeFilter, &QueryData.Word2, ShapeFilterPart1Size);
+		memcpy(((unsigned char*)&ShapeFilter) + ShapeFilterPart1Size, &SimData, ShapeFilterPart2Size);
 		return ShapeFilter;
 		PRAGMA_ENABLE_DEPRECATION_WARNINGS
 	}
@@ -52,7 +49,8 @@ namespace Chaos
 	void FCollisionData::SetShapeFilterData(const Chaos::Filter::FShapeFilterData& ShapeFilter)
 	{
 		PRAGMA_DISABLE_DEPRECATION_WARNINGS
-		memcpy(&QueryData.Word2, &ShapeFilter, sizeof(Filter::FShapeFilterData));
+		memcpy(&QueryData.Word2, &ShapeFilter, ShapeFilterPart1Size);
+		memcpy(&SimData, ((unsigned char*)&ShapeFilter) + ShapeFilterPart1Size, ShapeFilterPart2Size);
 		PRAGMA_ENABLE_DEPRECATION_WARNINGS
 	}
 
@@ -74,20 +72,14 @@ namespace Chaos
 
 	Chaos::Filter::FCombinedShapeFilterData FCollisionData::GetCombinedShapeFilterData() const
 	{
-		PRAGMA_DISABLE_DEPRECATION_WARNINGS
-		Filter::FInstanceData InstanceData;
-		Filter::FShapeFilterData ShapeFilter;
-		memcpy(&InstanceData, &QueryData, sizeof(Filter::FInstanceData));
-		memcpy(&ShapeFilter, &QueryData.Word2, sizeof(Filter::FShapeFilterData));
+		const Filter::FInstanceData& InstanceData = GetFilterInstanceData();
+		const Filter::FShapeFilterData& ShapeFilter = GetShapeFilterData();
 		return Filter::FCombinedShapeFilterData(ShapeFilter, InstanceData);
-		PRAGMA_ENABLE_DEPRECATION_WARNINGS
 	}
 
 	void FCollisionData::SetCombinedShapeFilterData(const Chaos::Filter::FCombinedShapeFilterData& CombinedShapeFilter)
 	{
-		PRAGMA_DISABLE_DEPRECATION_WARNINGS
-		memcpy(&QueryData, &CombinedShapeFilter.GetInstanceData(), sizeof(Filter::FInstanceData));
-		memcpy(&QueryData.Word2, &CombinedShapeFilter.GetShapeFilterData(), sizeof(Filter::FShapeFilterData));
-		PRAGMA_ENABLE_DEPRECATION_WARNINGS
+		SetFilterInstanceData(CombinedShapeFilter.GetInstanceData());
+		SetShapeFilterData(CombinedShapeFilter.GetShapeFilterData());
 	}
 } // namespace Chaos
diff --git a/Engine/Source/Runtime/Experimental/ChaosVisualDebugger/Private/DataWrappers/ChaosVDCollisionDataWrappers.cpp b/Engine/Source/Runtime/Experimental/ChaosVisualDebugger/Private/DataWrappers/ChaosVDCollisionDataWrappers.cpp
--- a/Engine/Source/Runtime/Experimental/ChaosVisualDebugger/Private/DataWrappers/ChaosVDCollisionDataWrappers.cpp
+++ b/Engine/Source/Runtime/Experimental/ChaosVisualDebugger/Private/DataWrappers/ChaosVDCollisionDataWrappers.cpp
@@ -226,7 +226,7 @@ uint32 FChaosVDCollisionFilterDataV2::GetFlags() const
 
 void FChaosVDCollisionFilterDataV2::SetMaskFilterAndFlags(uint32 InMaskFilter, uint32 InFlags)
 {
-	const uint64 MaskFilterPart = (InMaskFilter << MaskFilterOffset) & MaskFilterMask;
+	const uint64 MaskFilterPart = ((uint64)InMaskFilter << MaskFilterOffset) & MaskFilterMask;
 	const uint64 FlagsPart = InFlags & FlagsMask;
 	MaskFilterAndFlags = MaskFilterPart | FlagsPart;
 }
