Fixed additional static analysis warnings for new 64 bit collision filters.
#Synced-CL 49621724

diff --git a/Engine/Source/Runtime/Experimental/Chaos/Private/Chaos/CollisionFilterData.cpp b/Engine/Source/Runtime/Experimental/Chaos/Private/Chaos/CollisionFilterData.cpp
--- a/Engine/Source/Runtime/Experimental/Chaos/Private/Chaos/CollisionFilterData.cpp
+++ b/Engine/Source/Runtime/Experimental/Chaos/Private/Chaos/CollisionFilterData.cpp
@@ -305,6 +305,22 @@ namespace Chaos::Filter
 		return NewToLegacyWord3(Word3);
 	}
 
+	void FShapeFilterData::Store(FCollisionFilterData& QueryData, FCollisionFilterData& SimData) const
+	{
+		QueryData.Word2 = uint32(BlockChannels >> 32);
+		QueryData.Word3 = uint32(BlockChannels) & 0xFFFFFFFF;
+		SimData.Word0 = uint32(OverlapChannels >> 32);
+		SimData.Word1 = uint32(OverlapChannels) & 0xFFFFFFFF;
+		SimData.Word2 = Word3;
+	}
+
+	void FShapeFilterData::Load(const FCollisionFilterData& QueryData, const FCollisionFilterData& SimData)
+	{
+		BlockChannels = (uint64(QueryData.Word2) << 32) | uint64(QueryData.Word3);
+		OverlapChannels = (uint64(SimData.Word0) << 32) | uint64(SimData.Word1);
+		Word3 = SimData.Word2;
+	}
+
 	void FShapeUnionFilterData::Combine(const FShapeFilterData& ShapeFilter, bool bQueryEnabled, bool bSimEnabled)
 	{
 		if (bQueryEnabled)
diff --git a/Engine/Source/Runtime/Experimental/Chaos/Private/Chaos/ParticleDirtyFlags.cpp b/Engine/Source/Runtime/Experimental/Chaos/Private/Chaos/ParticleDirtyFlags.cpp
--- a/Engine/Source/Runtime/Experimental/Chaos/Private/Chaos/ParticleDirtyFlags.cpp
+++ b/Engine/Source/Runtime/Experimental/Chaos/Private/Chaos/ParticleDirtyFlags.cpp
@@ -40,8 +40,7 @@ namespace Chaos
 	{
 		PRAGMA_DISABLE_DEPRECATION_WARNINGS
 		Filter::FShapeFilterData ShapeFilter;
-		memcpy(&ShapeFilter, &QueryData.Word2, ShapeFilterPart1Size);
-		memcpy(((unsigned char*)&ShapeFilter) + ShapeFilterPart1Size, &SimData, ShapeFilterPart2Size);
+		ShapeFilter.Load(QueryData, SimData);
 		return ShapeFilter;
 		PRAGMA_ENABLE_DEPRECATION_WARNINGS
 	}
@@ -49,8 +48,7 @@ namespace Chaos
 	void FCollisionData::SetShapeFilterData(const Chaos::Filter::FShapeFilterData& ShapeFilter)
 	{
 		PRAGMA_DISABLE_DEPRECATION_WARNINGS
-		memcpy(&QueryData.Word2, &ShapeFilter, ShapeFilterPart1Size);
-		memcpy(&SimData, ((unsigned char*)&ShapeFilter) + ShapeFilterPart1Size, ShapeFilterPart2Size);
+		ShapeFilter.Store(QueryData, SimData);
 		PRAGMA_ENABLE_DEPRECATION_WARNINGS
 	}
 
diff --git a/Engine/Source/Runtime/Experimental/Chaos/Public/Chaos/CollisionFilterData.h b/Engine/Source/Runtime/Experimental/Chaos/Public/Chaos/CollisionFilterData.h
--- a/Engine/Source/Runtime/Experimental/Chaos/Public/Chaos/CollisionFilterData.h
+++ b/Engine/Source/Runtime/Experimental/Chaos/Public/Chaos/CollisionFilterData.h
@@ -158,6 +158,9 @@ namespace Chaos::Filter
 
 		UE_INTERNAL uint32 GetLegacyWord3() const;
 
+		UE_INTERNAL void Store(FCollisionFilterData& QueryData, FCollisionFilterData& SimData) const;
+		UE_INTERNAL void Load(const FCollisionFilterData& QueryData, const FCollisionFilterData& SimData);
+
 	private:
 		friend class FShapeFilterBuilder;
 		friend FQueryFilterData;
