Chaos Visual Debugger: Isolating extra resim tracing to the ResimDetails channel to keep basic captures lean. Adding more "frames" of tracing to CVD for those so they don't appear at the end of frames they don't actually correspond to, or as in between frame steps. Eventually I'd like to add a mode in CVD that combines all the contiguous captured frames corresponding to an actual single local solver frame as one displayed frame. If this gets done it will be done in a separate CL.

#Synced-CL 49576089

diff --git a/Engine/Source/Runtime/Experimental/Chaos/Private/PBDRigidsSolver.cpp b/Engine/Source/Runtime/Experimental/Chaos/Private/PBDRigidsSolver.cpp
--- a/Engine/Source/Runtime/Experimental/Chaos/Private/PBDRigidsSolver.cpp
+++ b/Engine/Source/Runtime/Experimental/Chaos/Private/PBDRigidsSolver.cpp
@@ -29,6 +29,7 @@
 #include "Chaos/PhysicsSolverBaseImpl.h"
 #include "Chaos/ConvexOptimizer.h"
 #include "Chaos/AsyncInitBodyHelper.h"
+#include "ChaosVisualDebugger/ChaosVDTraceMacros.h"
 
 #include "Chaos/DebugDraw/DebugDrawParticle.h"
 #include "ChaosDebugDraw/ChaosDDContext.h"
@@ -2166,15 +2167,30 @@ namespace Chaos
 					FDurationTimer ResimTimer(DebugInfo.ResimTime);
 					{
 						QUICK_SCOPE_CYCLE_COUNTER(ChaosRewindAndResim_Resim);
+
+#if WITH_CHAOS_VISUAL_DEBUGGER
+						bool bIsTracingResimDetails = (FChaosVisualDebuggerTrace::IsTracing() && CVDDC_ResimDetails->IsChannelEnabled());
+						if (bIsTracingResimDetails)
+						{
+							// This frame is redundant with the one in AdvanceSolver, but if we don't do this, subsequent particle traces and steps will appear as the last steps of the current frame we rewind from
+							// and CVD will not correctly advance the server timeline to the frame we are landing on, making it harder to compare the target set locally with server state at the same frame
+							CVD_TRACE_SOLVER_START_FRAME(FPhysicsSolverBase, *this);
+							CVD_SCOPE_TRACE_SOLVER_STEP(CVDDC_ResimDetails, TEXT("Rewind To Frame"));
+							MRewindData->TraceParticlesModifiedDuringRewind();
+						}
+#endif
+
 						// Do rollback as necessary
 						for (int32 Step = ResimStep; Step <= LastStep; ++Step)
 						{
-							if (!bFirst)
+#if WITH_CHAOS_VISUAL_DEBUGGER
+							if (!bFirst && bIsTracingResimDetails)
 							{
-								FPhysicsSolverBase& Solver = *this;
-								CVD_SCOPE_TRACE_SOLVER_STEP(CVDDC_Default, TEXT("Resim Step Start"));
+								CVD_TRACE_SOLVER_START_FRAME(FPhysicsSolverBase, *this);
+								CVD_SCOPE_TRACE_SOLVER_STEP(CVDDC_ResimDetails, TEXT("Resim Step Start"));
 								CVD_TRACE_PARTICLES_SOA(GetEvolution()->GetParticles(), &GetEvolution()->GetRigidClustering());
 							}
+#endif
 
 							if ((LastStep - Step) < RecordedPushData.Num())
 							{
@@ -2199,6 +2215,7 @@ namespace Chaos
 								}
 
 								GetEvolution()->SetReset(bFirst);
+
 #if !(UE_BUILD_SHIPPING || UE_BUILD_TEST)
 								if (bEnableNetworkPredictionDebug)
 								{
@@ -2225,18 +2242,27 @@ namespace Chaos
 									}
 								}
 #endif
-								if (bFirst)
+#if WITH_CHAOS_VISUAL_DEBUGGER
+								if (bFirst && bIsTracingResimDetails)
 								{
-									FPhysicsSolverBase& Solver = *this;
-									CVD_SCOPE_TRACE_SOLVER_STEP(CVDDC_Default, TEXT("Resim Step Start"));
+									CVD_TRACE_SOLVER_END_FRAME(FPhysicsSolverBase, *this);
+									CVD_TRACE_SOLVER_START_FRAME(FPhysicsSolverBase, *this);
+									CVD_SCOPE_TRACE_SOLVER_STEP(CVDDC_ResimDetails, TEXT("Resim Step Start"));
 									CVD_TRACE_PARTICLES_SOA(GetEvolution()->GetParticles(), &GetEvolution()->GetRigidClustering());
 								}
+#endif
 
 								{
 									QUICK_SCOPE_CYCLE_COUNTER(ChaosRewindAndResim_RewindCallback_PreResimStep_Internal);
 									MRewindCallback->PreResimStep_Internal(Step, bFirst);
 								}
 
+#if WITH_CHAOS_VISUAL_DEBUGGER
+								if (bIsTracingResimDetails)
+								{
+									CVD_TRACE_SOLVER_END_FRAME(FPhysicsSolverBase, *this);
+								}
+#endif
 								// Run the advance but omit the game thread callbacks as we're executing during the physics tick
 								FSolverTasksPTOnly ImmediateTask(*this, PushData);
 								//ensure(bSolverHasFrozenGameThreadCallbacks == false);	//We don't support this for resim as it's very expensive and difficult to schedule
diff --git a/Engine/Source/Runtime/Experimental/Chaos/Private/RewindData.cpp b/Engine/Source/Runtime/Experimental/Chaos/Private/RewindData.cpp
--- a/Engine/Source/Runtime/Experimental/Chaos/Private/RewindData.cpp
+++ b/Engine/Source/Runtime/Experimental/Chaos/Private/RewindData.cpp
@@ -275,7 +275,7 @@ void FRewindData::RewindStates(const int32 RewindFrame, const bool bResetSolver)
 void FRewindData::ApplyTargets(const int32 Frame, const bool bResetSimulation)
 {
 	QUICK_SCOPE_CYCLE_COUNTER(ChaosRewindAndResim_ApplyTargets);
-	CVD_SCOPE_TRACE_SOLVER_STEP(CVDDC_Default, TEXT("Rewind Apply Targets"));
+	CVD_SCOPE_TRACE_SOLVER_STEP(CVDDC_ResimDetails, TEXT("Rewind Apply Targets"));
 
 	PRAGMA_DISABLE_DEPRECATION_WARNINGS
 	RewindStates(Frame, bResetSimulation);
@@ -362,7 +362,12 @@ void FRewindData::ApplyTargets(const int32 Frame, const bool bResetSimulation)
 			}
 		});
 
-		CVD_TRACE_PARTICLE(PTParticle);
+#if WITH_CHAOS_VISUAL_DEBUGGER
+		if (FChaosVisualDebuggerTrace::IsTracing() && CVDDC_ResimDetails->IsChannelEnabled())
+		{
+			CVD_TRACE_PARTICLE(PTParticle);
+		}
+#endif
 
 #if !(UE_BUILD_SHIPPING || UE_BUILD_TEST)
 		if (!History.TargetPositions.IsClean(RewindFrameAndPhase) && Chaos::FPhysicsSolverBase::CanDebugNetworkPhysicsPrediction())
@@ -671,11 +676,7 @@ void FRewindData::FreezeParticlesOutsideResimBubbles()
 		FPBDRigidParticleHandle* RigidParticleHandle = ParticlesToFreezeInResim[ParticleToFreezeIndex];
 		if (IPhysicsProxyBase* PhysicsProxy = RigidParticleHandle->PhysicsProxy())
 		{
-			TPair<FParticleSafeRef, EObjectStateType> FrozenParticleState;
-			FrozenParticleState.Key.Proxy = static_cast<FSingleParticlePhysicsProxy*>(PhysicsProxy);
-			FrozenParticleState.Key.Timestamp = PhysicsProxy->GetSyncTimestamp();
-			FrozenParticleState.Value = RigidParticleHandle->ObjectState();
-			FrozenParticles.Add(MoveTemp(FrozenParticleState));
+			FrozenParticles.Add({PhysicsProxy, RigidParticleHandle->ObjectState()});
 		}
 		EObjectStateType FrozenObjectState = (ResimConsoleVars::BubbleResimFreezeMode == 0) ? EObjectStateType::Static : EObjectStateType::Kinematic;
 		if (Chaos::ResimConsoleVars::BubbleResimEnableCheapFreeze)
@@ -707,13 +708,11 @@ void FRewindData::UnfreezeParticlesOutsideResimBubbles()
 	// Revert to pre-resim object state
 	{
 		QUICK_SCOPE_CYCLE_COUNTER(ChaosRewindAndResim_RevertObjectState);
-		for (TPair<FParticleSafeRef, EObjectStateType>& FrozenParticle : FrozenParticles)
+		for (TPair<Chaos::FPhysicsProxyHandle, EObjectStateType>& FrozenParticle : FrozenParticles)
 		{
-			IPhysicsProxyBase* Proxy = FrozenParticle.Key.Proxy;
-			TSharedPtr<FProxyTimestampBase, ESPMode::ThreadSafe> Timestamp = FrozenParticle.Key.Timestamp;
-			EObjectStateType PreResimObjectState = FrozenParticle.Value;
-			if (Timestamp.IsValid() && !Timestamp->bDeleted && Proxy && Proxy->GetSyncTimestamp() == Timestamp && Proxy->GetType() == EPhysicsProxyType::SingleParticleProxy)
+			if (IPhysicsProxyBase* Proxy = FrozenParticle.Key.Get())
 			{
+				EObjectStateType PreResimObjectState = FrozenParticle.Value;
 				FSingleParticlePhysicsProxy* SingleParticlePhysicsProxy = static_cast<FSingleParticlePhysicsProxy*>(Proxy);
 				if (FPBDRigidParticleHandle* RigidParticleHandle = SingleParticlePhysicsProxy->GetHandle_LowLevel()->CastToRigidParticle())
 				{
@@ -731,8 +730,13 @@ void FRewindData::UnfreezeParticlesOutsideResimBubbles()
 							Evolution->SetParticleObjectState(RigidParticleHandle, PreResimObjectState);
 						}
 					}
-					// Trace modified particles. Unfortunately there is no way currently to only trace their object state so we trace the whole particle.
-					CVD_TRACE_PARTICLE(RigidParticleHandle);
+#if WITH_CHAOS_VISUAL_DEBUGGER
+					if (FChaosVisualDebuggerTrace::IsTracing() && CVDDC_ResimDetails->IsChannelEnabled())
+					{
+						// Trace modified particles. Unfortunately there is no way currently to only trace their object state so we trace the whole particle.
+						CVD_TRACE_PARTICLE(RigidParticleHandle);
+					}
+#endif
 				}
 			}
 		}
@@ -808,8 +812,6 @@ void FRewindData::RewindToFrame(int32 RewindFrame)
 
 	QUICK_SCOPE_CYCLE_COUNTER(RewindToFrame);
 
-	CVD_SCOPE_TRACE_SOLVER_STEP(CVDDC_Default, TEXT("Rewind To Frame"));
-
 	EnsureIsInPhysicsThreadContext();
 
 	if (Chaos::ResimConsoleVars::bEnableRewindLogs)
@@ -986,24 +988,18 @@ void FRewindData::RewindToFrame(int32 RewindFrame)
 
 			// If we modified the particle, trace it to CVD
 			// Bubble Resim: CVD tracing for bubble resim is done at the end for this function, after we freeze particles
+#if WITH_CHAOS_VISUAL_DEBUGGER
 			if (bAnyChange && !ResimConsoleVars::bIsBubbleResimEnabled)
 			{
-				CVD_TRACE_PARTICLE(GeometryParticle);
+				ParticlesModifiedDuringRewind.Add(GeometryParticle);
 			}
+#endif
 		}
 	}
 
 	// Bubble resim only: "Freeze" particles outside resim bubbles
 	FreezeParticlesOutsideResimBubbles();
 
-	// With bubble resim, anything that's not static ends up either being frozen (outside the resim bubbles) or marked as FullResim (inside the bubble)
-	// So we are better off tracing everything. We're only tracing the ActiveDynamicMovingKinematicParticlesView to try and keep the size down
-	// TODO: Implement particle tracing to trace only object state and resim type but on all particles to make sure we've not forgotten some updates
-	if (ResimConsoleVars::bIsBubbleResimEnabled)
-	{
-		CVD_TRACE_PARTICLES_VIEW(Evolution->GetParticles().GetActiveDynamicMovingKinematicParticlesView());
-	}
-
 #if !UE_BUILD_SHIPPING
 	// For now, just ensure that the joints are InSync
 	for(FDirtyJointInfo& DirtyJointInfo : DirtyJoints)
@@ -1015,9 +1011,35 @@ void FRewindData::RewindToFrame(int32 RewindFrame)
 #endif
 }
 
+#if WITH_CHAOS_VISUAL_DEBUGGER
+void FRewindData::TraceParticlesModifiedDuringRewind()
+{
+	if (FChaosVisualDebuggerTrace::IsTracing() && CVDDC_ResimDetails->IsChannelEnabled())
+	{
+		// With bubble resim, anything that's not static ends up either being frozen (outside the resim bubbles) or marked as FullResim (inside the bubble)
+		// So we are better off tracing everything. We're only tracing the ActiveDynamicMovingKinematicParticlesView to try and keep the size down
+		// TODO: Implement particle tracing to trace only object state and resim type but on all particles to make sure we've not forgotten some updates
+		if (ResimConsoleVars::bIsBubbleResimEnabled)
+		{
+			if (FPBDRigidsEvolution* Evolution = Solver ? Solver->GetEvolution() : nullptr)
+			{
+				CVD_TRACE_PARTICLES_VIEW(Evolution->GetParticles().GetActiveDynamicMovingKinematicParticlesView());
+			}
+		}
+		else
+		{
+			for (const FGeometryParticleHandle* GeometryParticle : ParticlesModifiedDuringRewind)
+			{
+				CVD_TRACE_PARTICLE(GeometryParticle);
+			}
+		}
+	}
+}
+#endif
+
 void FRewindData::EndResim()
 {
-	CVD_SCOPE_TRACE_SOLVER_STEP(CVDDC_Default, TEXT("End Resim"));
+	CVD_SCOPE_TRACE_SOLVER_STEP(CVDDC_ResimDetails, TEXT("End Resim"));
 	QUICK_SCOPE_CYCLE_COUNTER(ChaosRewindAndResim_EndResim);
 
 	// Bubble resim only: "Unfreeze" particles that were frozen because they were outside resim bubbles
@@ -1025,12 +1047,16 @@ void FRewindData::EndResim()
 
 	// Clear the array of resim bubbles
 	ResimBubbles.Empty();
+
+#if WITH_CHAOS_VISUAL_DEBUGGER
+	ParticlesModifiedDuringRewind.Empty();
+#endif
 }
 
 void FRewindData::StepNonResimParticles(const int32 Frame)
 {
 	QUICK_SCOPE_CYCLE_COUNTER(ChaosRewindAndResim_StepNonResimParticles);
-	CVD_SCOPE_TRACE_SOLVER_STEP(CVDDC_Default, TEXT("Step Resim As Followers"));
+	CVD_SCOPE_TRACE_SOLVER_STEP(CVDDC_ResimDetails, TEXT("Step Non Resim Particles"));
 
 	const FFrameAndPhase FrameAndPhase{ Frame, FFrameAndPhase::PrePushData };
 	auto RewindHelper = [FrameAndPhase, this](auto Obj, auto& Property, const auto& RewindFunc) -> bool
@@ -1171,10 +1197,12 @@ void FRewindData::StepNonResimParticles(const int32 Frame)
 				});
 		}
 
-		if (bHasChanged && RigidParticle)
+#if WITH_CHAOS_VISUAL_DEBUGGER
+		if (bHasChanged && RigidParticle && FChaosVisualDebuggerTrace::IsTracing() && CVDDC_ResimDetails->IsChannelEnabled())
 		{
 			CVD_TRACE_PARTICLE(RigidParticle);
 		}
+#endif
 	}
 }
 
@@ -2634,4 +2662,4 @@ void FRewindData::BlockResim()
 	}
 }
 
-}
\ No newline at end of file
+}
diff --git a/Engine/Source/Runtime/Experimental/Chaos/Public/RewindData.h b/Engine/Source/Runtime/Experimental/Chaos/Public/RewindData.h
--- a/Engine/Source/Runtime/Experimental/Chaos/Public/RewindData.h
+++ b/Engine/Source/Runtime/Experimental/Chaos/Public/RewindData.h
@@ -9,6 +9,7 @@
 #include "Chaos/ResimCacheBase.h"
 #include "Chaos/PBDJointConstraints.h"
 #include "Chaos/AABB.h"
+#include "Chaos/PhysicsProxyHandle.h"
 
 #ifndef VALIDATE_REWIND_DATA
 #define VALIDATE_REWIND_DATA 0
@@ -2248,6 +2249,11 @@ private:
 public:
 	// This can be used to determine if an update that affects physics and restricts its effects to within an AABB can be skipped during a resim
 	bool CHAOS_API OverlapsAnyResimBubble(const Chaos::FAABB3& InAABB) const;
+	
+#if WITH_CHAOS_VISUAL_DEBUGGER
+	// Traces the particles that were modified during rewind. in the bubble resim case this also traces particles that were frozen.
+	void TraceParticlesModifiedDuringRewind();
+#endif
 private:
 	// Freeze particles outside resim bubbles before rewind/resim
 	void FreezeParticlesOutsideResimBubbles();
@@ -2258,15 +2264,13 @@ private:
 	// Adds all particles needing resim and particle likely to interact with them during resim to OutParticlesOverlappingResimBubbles
 	// This is accomplished by creating resim "bubbles" around particles needing resim and growing them as they overlap other particles
 	void FindParticlesOverlappingResimBubbles(const FFrameAndPhase& RewindFrameAndPhase, float TotalResimDuration, TSet<FGeometryParticleHandle*>& OutParticlesOverlappingResimBubbles);
-	// This structure references a particle in a way that is safe to use several frames after having been established
-	struct FParticleSafeRef
-	{
-		IPhysicsProxyBase* Proxy;
-		TSharedPtr<FProxyTimestampBase, ESPMode::ThreadSafe> Timestamp;
-	};
 	// Used to keep track of objects we "freeze" during resim (by making them kinematic and marking them FrozenDuringResim)
 	// and their object state prior to the freezing so it can be set back to its original value
-	TArray<TPair<FParticleSafeRef, EObjectStateType>> FrozenParticles;
+	TArray<TPair<Chaos::FPhysicsProxyHandle, EObjectStateType>> FrozenParticles;
+#if WITH_CHAOS_VISUAL_DEBUGGER
+	// Used to keep track of which particles were rewound during the last rewind. Emptied at the end of resim.
+	TArray<FGeometryParticleHandle*> ParticlesModifiedDuringRewind;
+#endif
 
 	// Resim bubbles start as the swept AABBs of particles needing resim, and grow to encompass nearby kinematic and dynamic particles involved in the resim
 	// Their formation happens in FindParticlesOverlappingResimBubbles()
