Fixed some LogCategories declared in multiple dlls (some were defined >200 times) because the were using DEFINE_LOG_CATEGORY_STATIC in header files instead of DECLARE_LOG_CATEGORY_EXTERN

#Synced-CL 49574741

diff --git a/Engine/Source/Editor/ClassViewer/Private/ClassViewerModule.cpp b/Engine/Source/Editor/ClassViewer/Private/ClassViewerModule.cpp
--- a/Engine/Source/Editor/ClassViewer/Private/ClassViewerModule.cpp
+++ b/Engine/Source/Editor/ClassViewer/Private/ClassViewerModule.cpp
@@ -24,6 +24,8 @@
 
 IMPLEMENT_MODULE( FClassViewerModule, ClassViewer );
 
+DEFINE_LOG_CATEGORY(LogEditorClassViewer);
+
 namespace ClassViewerModule
 {
 	static const FName ClassViewerApp = FName("ClassViewerApp");
diff --git a/Engine/Source/Editor/ClassViewer/Public/ClassViewerModule.h b/Engine/Source/Editor/ClassViewer/Public/ClassViewerModule.h
--- a/Engine/Source/Editor/ClassViewer/Public/ClassViewerModule.h
+++ b/Engine/Source/Editor/ClassViewer/Public/ClassViewerModule.h
@@ -18,7 +18,7 @@ class IPropertyHandle;
 class SWidget;
 class UClass;
 
-DEFINE_LOG_CATEGORY_STATIC(LogEditorClassViewer, Log, All);
+CLASSVIEWER_API DECLARE_LOG_CATEGORY_EXTERN(LogEditorClassViewer, Log, All);
 
 /** Delegate used with the Class Viewer in 'class picking' mode.  You'll bind a delegate when the
     class viewer widget is created, which will be fired off when a class is selected in the list */
diff --git a/Engine/Source/Editor/UnrealEd/Private/Commandlets/SummarizeTraceUtils.cpp b/Engine/Source/Editor/UnrealEd/Private/Commandlets/SummarizeTraceUtils.cpp
--- a/Engine/Source/Editor/UnrealEd/Private/Commandlets/SummarizeTraceUtils.cpp
+++ b/Engine/Source/Editor/UnrealEd/Private/Commandlets/SummarizeTraceUtils.cpp
@@ -28,6 +28,8 @@
 #include "Math/NumericLimits.h"
 #include "Math/UnrealMathSSE.h"
 
+DEFINE_LOG_CATEGORY(LogSummarizeTrace);
+
 uint64 FIncrementalVariance::GetCount() const
 {
 	return Count;
diff --git a/Engine/Source/Editor/UnrealEd/Private/Commandlets/SummarizeTraceUtils.h b/Engine/Source/Editor/UnrealEd/Private/Commandlets/SummarizeTraceUtils.h
--- a/Engine/Source/Editor/UnrealEd/Private/Commandlets/SummarizeTraceUtils.h
+++ b/Engine/Source/Editor/UnrealEd/Private/Commandlets/SummarizeTraceUtils.h
@@ -9,7 +9,7 @@ namespace TraceServices { class IAnalysisSession; }
 namespace UE::Trace { class IAnalyzer; }
 
 // Defined here to prevent adding logs in the code above, which will likely be moved elsewhere.
-DEFINE_LOG_CATEGORY_STATIC(LogSummarizeTrace, Log, All);
+DECLARE_LOG_CATEGORY_EXTERN(LogSummarizeTrace, Log, All);
 
 /*
 *
diff --git a/Engine/Source/Editor/UnrealEd/Private/WorldPartition/WorldPartitionBuilderHelpers.cpp b/Engine/Source/Editor/UnrealEd/Private/WorldPartition/WorldPartitionBuilderHelpers.cpp
--- a/Engine/Source/Editor/UnrealEd/Private/WorldPartition/WorldPartitionBuilderHelpers.cpp
+++ b/Engine/Source/Editor/UnrealEd/Private/WorldPartition/WorldPartitionBuilderHelpers.cpp
@@ -14,6 +14,7 @@
 #include "ISourceControlProvider.h"
 #include "SourceControlOperations.h"
 
+DEFINE_LOG_CATEGORY(LogWorldPartitionBuilderSourceControlHelper);
 
 void FBuilderModifiedFiles::Add(EFileOperation FileOp, const FString& File)
 {
diff --git a/Engine/Source/Editor/UnrealEd/Public/WorldPartition/WorldPartitionBuilderHelpers.h b/Engine/Source/Editor/UnrealEd/Public/WorldPartition/WorldPartitionBuilderHelpers.h
--- a/Engine/Source/Editor/UnrealEd/Public/WorldPartition/WorldPartitionBuilderHelpers.h
+++ b/Engine/Source/Editor/UnrealEd/Public/WorldPartition/WorldPartitionBuilderHelpers.h
@@ -6,7 +6,7 @@
 #include "PackageSourceControlHelper.h" 
 #include "WorldPartition/WorldPartition.h" // for ISourceControlHelper
 
-DEFINE_LOG_CATEGORY_STATIC(LogWorldPartitionBuilderSourceControlHelper, All, All);
+UNREALED_API DECLARE_LOG_CATEGORY_EXTERN(LogWorldPartitionBuilderSourceControlHelper, All, All);
 
 struct FBuilderModifiedFiles
 {
diff --git a/Engine/Source/Runtime/Experimental/Chaos/Private/BoneHierarchy.cpp b/Engine/Source/Runtime/Experimental/Chaos/Private/BoneHierarchy.cpp
new file mode 100644
--- /dev/null
+++ b/Engine/Source/Runtime/Experimental/Chaos/Private/BoneHierarchy.cpp
@@ -0,0 +1,5 @@
+// Copyright Epic Games, Inc. All Rights Reserved.
+
+#include "BoneHierarchy.h"
+
+DEFINE_LOG_CATEGORY(USkeletalMeshSimulationComponentLogging);
diff --git a/Engine/Source/Runtime/Experimental/Chaos/Private/Chaos/Deformable/GaussSeidelMainConstraint.cpp b/Engine/Source/Runtime/Experimental/Chaos/Private/Chaos/Deformable/GaussSeidelMainConstraint.cpp
--- a/Engine/Source/Runtime/Experimental/Chaos/Private/Chaos/Deformable/GaussSeidelMainConstraint.cpp
+++ b/Engine/Source/Runtime/Experimental/Chaos/Private/Chaos/Deformable/GaussSeidelMainConstraint.cpp
@@ -3,6 +3,8 @@
 #include "Chaos/Math/Krylov.h"
 #include "ChaosLog.h"
 
+DEFINE_LOG_CATEGORY(LogDeformableGaussSeidelMainConstraint);
+
 namespace Chaos::Softs
 {
 	int32 MaxItCG = 50;
diff --git a/Engine/Source/Runtime/Experimental/Chaos/Public/BoneHierarchy.h b/Engine/Source/Runtime/Experimental/Chaos/Public/BoneHierarchy.h
--- a/Engine/Source/Runtime/Experimental/Chaos/Public/BoneHierarchy.h
+++ b/Engine/Source/Runtime/Experimental/Chaos/Public/BoneHierarchy.h
@@ -3,8 +3,8 @@
 
 #include "PhysicsProxy/AnalyticImplicitGroup.h"
 
-DEFINE_LOG_CATEGORY_STATIC(USkeletalMeshSimulationComponentLogging, NoLogging, All);
-//DEFINE_LOG_CATEGORY_STATIC(USkeletalMeshSimulationComponentLogging, Log, All);
+CHAOS_API DECLARE_LOG_CATEGORY_EXTERN(USkeletalMeshSimulationComponentLogging, NoLogging, All);
+//CHAOS_API DECLARE_LOG_CATEGORY_EXTERN(USkeletalMeshSimulationComponentLogging, Log, All);
 
 /**
  * A hierarchy of transforms (bones) that tracks the state of transforms within 
diff --git a/Engine/Source/Runtime/Experimental/Chaos/Public/Chaos/Deformable/GaussSeidelMainConstraint.h b/Engine/Source/Runtime/Experimental/Chaos/Public/Chaos/Deformable/GaussSeidelMainConstraint.h
--- a/Engine/Source/Runtime/Experimental/Chaos/Public/Chaos/Deformable/GaussSeidelMainConstraint.h
+++ b/Engine/Source/Runtime/Experimental/Chaos/Public/Chaos/Deformable/GaussSeidelMainConstraint.h
@@ -1,4 +1,4 @@
-// Copyright Epic Games, Inc. All Rights Reserved.
+﻿// Copyright Epic Games, Inc. All Rights Reserved.
 #pragma once
 
 #include "Chaos/PBDSpringConstraintsBase.h"
@@ -19,7 +19,7 @@ DECLARE_CYCLE_STAT(TEXT("Chaos.Deformable.GSMainConstraint.Init"), STAT_ChaosGSM
 DECLARE_CYCLE_STAT(TEXT("Chaos.Deformable.GSMainConstraint.InitTransientColor"), STAT_ChaosGSMainConstraint_InitTransientColor, STATGROUP_Chaos);
 DECLARE_CYCLE_STAT(TEXT("Chaos.Deformable.GSMainConstraint.InitDynamicColor"), STAT_ChaosGSMainConstraint_InitDynamicColor, STATGROUP_Chaos);
 
-DEFINE_LOG_CATEGORY_STATIC(LogDeformableGaussSeidelMainConstraint, Log, All);
+CHAOS_API DECLARE_LOG_CATEGORY_EXTERN(LogDeformableGaussSeidelMainConstraint, Log, All);
 
 namespace Chaos::Softs
 {
diff --git a/Engine/Source/Runtime/Experimental/ChaosSolverEngine/Private/Chaos/ChaosSolver.cpp b/Engine/Source/Runtime/Experimental/ChaosSolverEngine/Private/Chaos/ChaosSolver.cpp
--- a/Engine/Source/Runtime/Experimental/ChaosSolverEngine/Private/Chaos/ChaosSolver.cpp
+++ b/Engine/Source/Runtime/Experimental/ChaosSolverEngine/Private/Chaos/ChaosSolver.cpp
@@ -4,8 +4,6 @@
 
 #include UE_INLINE_GENERATED_CPP_BY_NAME(ChaosSolver)
 
-DEFINE_LOG_CATEGORY_STATIC(FSC_Log, NoLogging, All);
-
 UChaosSolver::UChaosSolver(const FObjectInitializer& ObjectInitializer)
 	: UObject(ObjectInitializer)
 {
diff --git a/Engine/Source/Runtime/RenderCore/Private/RenderGraphPrivate.h b/Engine/Source/Runtime/RenderCore/Private/RenderGraphPrivate.h
--- a/Engine/Source/Runtime/RenderCore/Private/RenderGraphPrivate.h
+++ b/Engine/Source/Runtime/RenderCore/Private/RenderGraphPrivate.h
@@ -9,8 +9,6 @@
 #include "RHI.h"
 #include "DataDrivenShaderPlatformInfo.h"
 
-DEFINE_LOG_CATEGORY_STATIC(LogRDG, Log, All);
-
 #define RDG_DUMP_GRAPH_PRODUCERS 1
 #define RDG_DUMP_GRAPH_RESOURCES 2
 #define RDG_DUMP_GRAPH_TRACKS 3
diff --git a/Engine/Source/Runtime/RenderCore/Private/RenderGraphValidation.cpp b/Engine/Source/Runtime/RenderCore/Private/RenderGraphValidation.cpp
--- a/Engine/Source/Runtime/RenderCore/Private/RenderGraphValidation.cpp
+++ b/Engine/Source/Runtime/RenderCore/Private/RenderGraphValidation.cpp
@@ -6,6 +6,8 @@
 #include "MultiGPU.h"
 #include "RenderGraphPass.h"
 
+DEFINE_LOG_CATEGORY_STATIC(LogRDG, Log, All);
+
 #if RDG_ENABLE_DEBUG
 
 namespace
