Ignore FVCellProperty in JsonObjectGraph::Stringify

UVerseClass and UVerseStruct each contain an FVCellProperty pointing back to their originating VClass. GenerateBlueprintDependencies and FBlueprintCompileReinstancer::CalculateConstructionHash pass these UStructs through JsonObjectGraph::Stringify in order to hash the result, but Stringify does not currently support VCells.

However, these VClasses do not need to be included in these hashes, because all the relevant data is also represented in the UStruct. Any actual instance properties that should be hashed use FVRestValueProperty instead of FVCellProperty, so Stringify can safely ignore FVCellProperty while still failing if it encounters an FVRestValueProperty.

(As the Verse VM is used in more places, or more Verse instances are referenced by Blueprints, we may need to implement actual support for VCells in Stringify, but that is not yet necessary.)

Before running its own logic on reflected properties, Stringify calls Serialize with an FJsonStringifyStringifyArchive, then falls back to an FJsonStringifyArchive. While the latter just ignores all reflected properties directly, the former instead must filter them out itself- so this change also ignores VCell references in text-formatted FArchiveFromStructuredArchiveImpl to support that scenario.

#rnx
#Synced-CL 49643405

diff --git a/Engine/Source/Runtime/Core/Private/Serialization/ArchiveFromStructuredArchive.cpp b/Engine/Source/Runtime/Core/Private/Serialization/ArchiveFromStructuredArchive.cpp
--- a/Engine/Source/Runtime/Core/Private/Serialization/ArchiveFromStructuredArchive.cpp
+++ b/Engine/Source/Runtime/Core/Private/Serialization/ArchiveFromStructuredArchive.cpp
@@ -183,6 +183,25 @@ FArchive& FArchiveFromStructuredArchiveImpl::operator<<(class UObject*& Value)
 	return *this;
 }
 
+#if WITH_VERSE_VM || defined(__INTELLISENSE__)
+FArchive& FArchiveFromStructuredArchiveImpl::operator<<(Verse::VCell*& Value)
+{
+	OpenArchive();
+
+	if (InnerArchive.IsTextFormat())
+	{
+		// VCell references are not generally supported in structured text archives, but
+		// ignoring them while saving makes things easier for use cases like debug printing.
+		check(!IsLoading());
+	}
+	else
+	{
+		InnerArchive << Value;
+	}
+	return *this;
+}
+#endif
+
 FArchive& FArchiveFromStructuredArchiveImpl::operator<<(class FText& Value)
 {
 	OpenArchive();
diff --git a/Engine/Source/Runtime/Core/Public/Serialization/StructuredArchiveAdapters.h b/Engine/Source/Runtime/Core/Public/Serialization/StructuredArchiveAdapters.h
--- a/Engine/Source/Runtime/Core/Public/Serialization/StructuredArchiveAdapters.h
+++ b/Engine/Source/Runtime/Core/Public/Serialization/StructuredArchiveAdapters.h
@@ -54,6 +54,9 @@ public:
 	//~ Begin FArchive Interface
 	CORE_API virtual FArchive& operator<<(class FName& Value) override;
 	CORE_API virtual FArchive& operator<<(class UObject*& Value) override;
+#if WITH_VERSE_VM || defined(__INTELLISENSE__)
+	CORE_API virtual FArchive& operator<<(Verse::VCell*& Value) override;
+#endif
 	CORE_API virtual FArchive& operator<<(class FText& Value) override;
 	//~ End FArchive Interface
 
diff --git a/Engine/Source/Runtime/Experimental/JsonObjectGraph/Private/JsonStringifyArchive.cpp b/Engine/Source/Runtime/Experimental/JsonObjectGraph/Private/JsonStringifyArchive.cpp
--- a/Engine/Source/Runtime/Experimental/JsonObjectGraph/Private/JsonStringifyArchive.cpp
+++ b/Engine/Source/Runtime/Experimental/JsonObjectGraph/Private/JsonStringifyArchive.cpp
@@ -117,6 +117,15 @@ FArchive& FJsonStringifyArchive::operator<<(FField*& Value)
 	return *this;
 }
 
+#if WITH_VERSE_VM || defined(__INTELLISENSE__)
+FArchive& FJsonStringifyArchive::operator<<(Verse::VCell*& Value)
+{
+	// VCell references should only be serialized as reflected properties, which this archive skips.
+	check(0);
+	return *this;
+}
+#endif
+
 FArchive& FJsonStringifyArchive::operator<<(struct FLazyObjectPtr& Value)
 {
 	UObject* Resolved = Value.Get();
diff --git a/Engine/Source/Runtime/Experimental/JsonObjectGraph/Private/JsonStringifyArchive.h b/Engine/Source/Runtime/Experimental/JsonObjectGraph/Private/JsonStringifyArchive.h
--- a/Engine/Source/Runtime/Experimental/JsonObjectGraph/Private/JsonStringifyArchive.h
+++ b/Engine/Source/Runtime/Experimental/JsonObjectGraph/Private/JsonStringifyArchive.h
@@ -36,6 +36,9 @@ private:
 	#endif
 	virtual FArchive& operator<<(UObject*& Value) override;
 	virtual FArchive& operator<<(FField*& Value) override;
+#if WITH_VERSE_VM || defined(__INTELLISENSE__)
+	virtual FArchive& operator<<(Verse::VCell*& Value) override;
+#endif
 	virtual FArchive& operator<<(struct FLazyObjectPtr& Value) override;
 	virtual FArchive& operator<<(struct FObjectPtr& Value) override;
 	virtual FArchive& operator<<(struct FSoftObjectPtr& Value) override;
diff --git a/Engine/Source/Runtime/Experimental/JsonObjectGraph/Private/JsonStringifyImpl.cpp b/Engine/Source/Runtime/Experimental/JsonObjectGraph/Private/JsonStringifyImpl.cpp
--- a/Engine/Source/Runtime/Experimental/JsonObjectGraph/Private/JsonStringifyImpl.cpp
+++ b/Engine/Source/Runtime/Experimental/JsonObjectGraph/Private/JsonStringifyImpl.cpp
@@ -19,6 +19,7 @@
 #include "UObject/PropertyOptional.h"
 #include "UObject/Package.h"
 #include "UObject/TextProperty.h"
+#include "UObject/VerseValueProperty.h"
 
 #if WITH_EDITORONLY_DATA
 #include "UObject/InstanceDataObjectUtils.h"
@@ -551,6 +552,15 @@ void FJsonStringifyImpl::WriteIdentifierAndValueToJson(const void* Container, co
 		return;
 	}
 
+#if WITH_VERSE_VM || defined(__INTELLISENSE__)
+	// FVCellProperty (as opposed to FVRestValueProperty) is only used by UVerseClass/UVerseStruct,
+	// where it just points back to the VClass the UStruct was derived from. It can be safely ignored.
+	if (const FVCellProperty* VerseProperty = CastField<FVCellProperty>(Property))
+	{
+		return;
+	}
+#endif
+
 	// this is inefficient for structs... but it papers over some
 	// problematic identical implementations. We will improve it
 	bool bMatchesDefault = DefaultContainer != nullptr && IsDeltaEncoding();
