Constraints: enable fast sequential evaluation by default

#rnx

#Synced-CL 49659971

diff --git a/Engine/Source/Runtime/Experimental/Animation/Constraints/Private/ConstraintsEvaluation.cpp b/Engine/Source/Runtime/Experimental/Animation/Constraints/Private/ConstraintsEvaluation.cpp
--- a/Engine/Source/Runtime/Experimental/Animation/Constraints/Private/ConstraintsEvaluation.cpp
+++ b/Engine/Source/Runtime/Experimental/Animation/Constraints/Private/ConstraintsEvaluation.cpp
@@ -9,11 +9,11 @@
 namespace UE::Anim::Constraints
 {
 
-static bool	bUseFastEvaluation = false;
+static bool	bUseFastEvaluation = true;
 static FAutoConsoleVariableRef CVarUseFastEvaluation(
 	TEXT("Constraints.UseFastEvaluation"),
 	bUseFastEvaluation,
-	TEXT("Whether to enable fast sequential constraints evaluation. (default is false)") );
+	TEXT("Whether to enable fast sequential constraints evaluation. (default is true)") );
 
 bool UseFastEvaluation()
 {
