Geometry Collections - Support Custom Depth pass

Commit 1 of 2 | Engine Code

#Synced-CL 49673251

diff --git a/Engine/Source/Runtime/Experimental/GeometryCollectionEngine/Private/GeometryCollection/GeometryCollectionBlueprintLibrary.cpp b/Engine/Source/Runtime/Experimental/GeometryCollectionEngine/Private/GeometryCollection/GeometryCollectionBlueprintLibrary.cpp
--- a/Engine/Source/Runtime/Experimental/GeometryCollectionEngine/Private/GeometryCollection/GeometryCollectionBlueprintLibrary.cpp
+++ b/Engine/Source/Runtime/Experimental/GeometryCollectionEngine/Private/GeometryCollection/GeometryCollectionBlueprintLibrary.cpp
@@ -11,36 +11,25 @@
 #endif
 #include UE_INLINE_GENERATED_CPP_BY_NAME(GeometryCollectionBlueprintLibrary)
 
-void UGeometryCollectionBlueprintLibrary::SetCustomInstanceDataByIndex(UGeometryCollectionComponent* GeometryCollectionComponent, int32 CustomDataIndex, float CustomDataValue)
+static IGeometryCollectionCustomDataInterface* GetCustomRenderer(UGeometryCollectionComponent* GeometryCollectionComponent)
 {
-	if (GeometryCollectionComponent == nullptr)
-	{
-		return;
-	}
+	return GeometryCollectionComponent ? Cast<IGeometryCollectionCustomDataInterface>(GeometryCollectionComponent->GetCustomRenderer()) : nullptr;
+}
 
-	IGeometryCollectionCustomDataInterface* CustomRenderer = Cast<IGeometryCollectionCustomDataInterface>(GeometryCollectionComponent->GetCustomRenderer());
-	if (CustomRenderer == nullptr)
+void UGeometryCollectionBlueprintLibrary::SetCustomInstanceDataByIndex(UGeometryCollectionComponent* GeometryCollectionComponent, int32 CustomDataIndex, float CustomDataValue)
+{
+	if (IGeometryCollectionCustomDataInterface* CustomRenderer = GetCustomRenderer(GeometryCollectionComponent))
 	{
-		return;
+		CustomRenderer->SetCustomInstanceData(CustomDataIndex, CustomDataValue);
 	}
-
-	CustomRenderer->SetCustomInstanceData(CustomDataIndex, CustomDataValue);
 }
 
 void UGeometryCollectionBlueprintLibrary::SetCustomInstanceDataByName(UGeometryCollectionComponent* GeometryCollectionComponent, FName CustomDataName, float CustomDataValue)
 {
-	if (GeometryCollectionComponent == nullptr)
+	if (IGeometryCollectionCustomDataInterface* CustomRenderer = GetCustomRenderer(GeometryCollectionComponent))
 	{
-		return;
+		CustomRenderer->SetCustomInstanceData(CustomDataName, CustomDataValue);
 	}
-
-	IGeometryCollectionCustomDataInterface* CustomRenderer = Cast<IGeometryCollectionCustomDataInterface>(GeometryCollectionComponent->GetCustomRenderer());
-	if (CustomRenderer == nullptr)
-	{
-		return;
-	}
-
-	CustomRenderer->SetCustomInstanceData(CustomDataName, CustomDataValue);
 }
 
 void UGeometryCollectionBlueprintLibrary::SetISMPoolCustomInstanceData(UGeometryCollectionComponent* GeometryCollectionComponent, int32 CustomDataIndex, float CustomDataValue)
@@ -48,6 +37,13 @@ void UGeometryCollectionBlueprintLibrary::SetISMPoolCustomInstanceData(UGeometry
 	SetCustomInstanceDataByIndex(GeometryCollectionComponent, CustomDataIndex, CustomDataValue);
 }
 
+void UGeometryCollectionBlueprintLibrary::SetCustomDepthStencil(UGeometryCollectionComponent* GeometryCollectionComponent, bool bEnable, int32 CustomStencilValue)
+{
+	if (IGeometryCollectionCustomDataInterface* CustomRenderer = GetCustomRenderer(GeometryCollectionComponent))
+	{
+		CustomRenderer->SetRenderCustomDepthStencil(GeometryCollectionComponent, bEnable, CustomStencilValue);
+	}
+}
 
 bool UGeometryCollectionBlueprintLibrary::AddGeometryCollectionsToBlueprint(UBlueprint* Blueprint, const TArray<UObject*>& Collections)
 {
diff --git a/Engine/Source/Runtime/Experimental/GeometryCollectionEngine/Public/GeometryCollection/GeometryCollectionBlueprintLibrary.h b/Engine/Source/Runtime/Experimental/GeometryCollectionEngine/Public/GeometryCollection/GeometryCollectionBlueprintLibrary.h
--- a/Engine/Source/Runtime/Experimental/GeometryCollectionEngine/Public/GeometryCollection/GeometryCollectionBlueprintLibrary.h
+++ b/Engine/Source/Runtime/Experimental/GeometryCollectionEngine/Public/GeometryCollection/GeometryCollectionBlueprintLibrary.h
@@ -35,6 +35,16 @@ public:
 	UFUNCTION(BlueprintCallable, Category = "ChaosPhysics")
 	static UE_API void SetCustomInstanceDataByName(UGeometryCollectionComponent* GeometryCollectionComponent, FName CustomDataName, float CustomDataValue);
 
+	/** 
+	 * Specify custom depth and a value to write during the custom depth pass for all instances associated with a geometry collection. 
+	 * This assumes that the geometry collection is using a custom renderer that supports IGeometryCollectionCustomDataInterface.
+	 * @param GeometryCollectionComponent	The Geometry Collection Component that we want to set custom depth data on.
+	 * @param bEnable	Enable or disable custom depth writing.
+	 * @param CustomStencilValue	The value which will be written to the stencil buffer during the custom depth pass if bEnable is true.
+	 */
+	UFUNCTION(BlueprintCallable, Category = "ChaosPhysics")
+	static UE_API void SetCustomDepthStencil(UGeometryCollectionComponent* GeometryCollectionComponent, bool bEnable, int32 CustomStencilValue);
+
 	UE_DEPRECATED(5.5, "Use SetCustomInstanceDataByIndex() instead")
 	UFUNCTION(BlueprintCallable, Category = "ChaosPhysics", meta = (DeprecatedFunction, DeprecationMessage = "Use SetCustomInstanceDataByIndex() instead"))
 	static UE_API void SetISMPoolCustomInstanceData(UGeometryCollectionComponent* GeometryCollectionComponent, int32 CustomDataIndex, float CustomDataValue);
diff --git a/Engine/Source/Runtime/Experimental/GeometryCollectionEngine/Public/GeometryCollection/GeometryCollectionExternalRenderInterface.h b/Engine/Source/Runtime/Experimental/GeometryCollectionEngine/Public/GeometryCollection/GeometryCollectionExternalRenderInterface.h
--- a/Engine/Source/Runtime/Experimental/GeometryCollectionEngine/Public/GeometryCollection/GeometryCollectionExternalRenderInterface.h
+++ b/Engine/Source/Runtime/Experimental/GeometryCollectionEngine/Public/GeometryCollection/GeometryCollectionExternalRenderInterface.h
@@ -91,4 +91,7 @@ public:
 
 	/** Set value a custom instance data slot by name. The value will be set for all instances the renderer is spawning. */
 	virtual void SetCustomInstanceData(FName CustomDataName, float CustomDataValue) = 0;
+
+	/** Enable or disable custom depth writing and, if enabled, the value written to the stencil buffer. The value will be set for all instances the renderer is spawning. */
+	virtual void SetRenderCustomDepthStencil(UGeometryCollectionComponent* GeometryCollectionComponent, bool bEnable, int32 CustomStencilValue) = 0;
 };
diff --git a/Engine/Source/Runtime/Experimental/ISMPool/Private/ISMPool/ISMPoolComponent.cpp b/Engine/Source/Runtime/Experimental/ISMPool/Private/ISMPool/ISMPoolComponent.cpp
--- a/Engine/Source/Runtime/Experimental/ISMPool/Private/ISMPool/ISMPoolComponent.cpp
+++ b/Engine/Source/Runtime/Experimental/ISMPool/Private/ISMPool/ISMPoolComponent.cpp
@@ -244,6 +244,8 @@ void FISMPoolISM::InitISM(const FISMPoolStaticMeshInstance& InMeshInstance, bool
 	ISMComponent->bDisallowNanite = bDisallowNanite;
 	ISMComponent->SetMeshDrawCommandStatsCategory(MeshInstance.Desc.StatsCategory);
 	ISMComponent->ComponentTags = MeshInstance.Desc.Tags;
+	ISMComponent->SetRenderCustomDepth(MeshInstance.Desc.Flags & FISMPoolComponentDescription::UsesCustomDepthStencil);
+	ISMComponent->SetCustomDepthStencilValue(MeshInstance.Desc.CustomDepthStencilValue);
 
 	// Use a fixed seed to avoid getting a different seed at every run (see UInstancedStaticMeshComponent::OnRegister())
 	// A possible improvement would be to compute a hash from the owner component and use that as the seed.
diff --git a/Engine/Source/Runtime/Experimental/ISMPool/Public/ISMPool/ISMPoolComponent.h b/Engine/Source/Runtime/Experimental/ISMPool/Public/ISMPool/ISMPoolComponent.h
--- a/Engine/Source/Runtime/Experimental/ISMPool/Public/ISMPool/ISMPoolComponent.h
+++ b/Engine/Source/Runtime/Experimental/ISMPool/Public/ISMPool/ISMPoolComponent.h
@@ -113,6 +113,7 @@ struct FISMPoolComponentDescription
 		AffectDynamicIndirectLighting = 1 << 9,
 		AffectFarShadow = 1 << 10,
 		DistanceCullPrimitive = 1 << 11,
+		UsesCustomDepthStencil = 1 << 12,
 	};
 
 	uint32 Flags = WorldPositionOffsetWritesVelocity|EvaluateWorldPositionOffset|AffectShadow;
@@ -121,6 +122,7 @@ struct FISMPoolComponentDescription
 	int32 StartCullDistance = 0;
 	int32 EndCullDistance = 0;
 	int32 MinLod = 0;
+	int32 CustomDepthStencilValue = 0;
 	uint32 GroupHash = 0;	// Optional, allows identical SMs to be separated into different groups for finer grained culling
 	float LodScale = 1.f;
 	TArray<FName> Tags;
@@ -134,6 +136,7 @@ struct FISMPoolComponentDescription
 			StartCullDistance == Other.StartCullDistance &&
 			EndCullDistance == Other.EndCullDistance &&
 			MinLod == Other.MinLod &&
+			CustomDepthStencilValue == Other.CustomDepthStencilValue &&
 			LodScale == Other.LodScale &&
 			Tags == Other.Tags &&
 			GroupHash == Other.GroupHash &&
@@ -148,6 +151,7 @@ inline uint32 GetTypeHash(const FISMPoolComponentDescription& Desc)
 	Hash = HashCombineFast(Hash, GetTypeHash(Desc.StartCullDistance));
 	Hash = HashCombineFast(Hash, GetTypeHash(Desc.EndCullDistance));
 	Hash = HashCombineFast(Hash, GetTypeHash(Desc.MinLod));
+	Hash = HashCombineFast(Hash, GetTypeHash(Desc.CustomDepthStencilValue));
 	Hash = HashCombineFast(Hash, GetTypeHash(Desc.LodScale));
 	Hash = HashCombineFast(Hash, GetTypeHash(Desc.GroupHash));
 	Hash = HashCombineFast(Hash, GetArrayHash(Desc.Tags.GetData(), Desc.Tags.Num()));
