IAD
- Add function to dump content installer stats to log.
- Add iad.LogStats console command to log stats.
- Automatically log stats at shutdown when running unattended.

#Synced-CL 49554261

diff --git a/Engine/Source/Runtime/Experimental/IoStore/OnDemand/Private/OnDemandContentInstaller.cpp b/Engine/Source/Runtime/Experimental/IoStore/OnDemand/Private/OnDemandContentInstaller.cpp
--- a/Engine/Source/Runtime/Experimental/IoStore/OnDemand/Private/OnDemandContentInstaller.cpp
+++ b/Engine/Source/Runtime/Experimental/IoStore/OnDemand/Private/OnDemandContentInstaller.cpp
@@ -13,6 +13,7 @@
 #include "IO/IoStatus.h"
 #include "IO/IoStoreOnDemand.h"
 #include "IO/OnDemandError.h"
+#include "Misc/App.h"
 #include "Misc/Timespan.h"
 #include "OnDemandInstallCache.h"
 #include "OnDemandIoStore.h"
@@ -1859,6 +1860,11 @@ void FOnDemandContentInstaller::Shutdown()
 		ExecuteRequest(*NextRequest);
 		check(RunningRequests.IsEmpty());
 	}
+
+	if (FApp::IsUnattended())
+	{
+		FOnDemandContentInstallerStats::LogAnalytics();
+	}
 }
 
 } // namespace UE::IoStore
diff --git a/Engine/Source/Runtime/Experimental/IoStore/OnDemand/Private/Statistics.cpp b/Engine/Source/Runtime/Experimental/IoStore/OnDemand/Private/Statistics.cpp
--- a/Engine/Source/Runtime/Experimental/IoStore/OnDemand/Private/Statistics.cpp
+++ b/Engine/Source/Runtime/Experimental/IoStore/OnDemand/Private/Statistics.cpp
@@ -15,6 +15,7 @@
 #include "Internationalization/Internationalization.h"
 #include "Math/UnrealMathUtility.h"
 #include "Misc/CoreDelegates.h"
+#include "Logging/StructuredLog.h"
 #include "OnDemandHttpIoDispatcher.h"
 #include "Templates/Requires.h"
 
@@ -71,6 +72,17 @@ static FAutoConsoleVariableRef CVar_DisplayOnScreenStatistics(
 
 #endif // UE_ENABLE_ONSCREEN_STATISTICS
 
+#if !UE_BUILD_SHIPPING
+FAutoConsoleCommand CmdLogOnDemandContentInstallerStats(
+	TEXT("iad.LogStats"), 
+	TEXT("Dump IAD content installer stats to log"), 
+	FConsoleCommandDelegate::CreateLambda([]()
+	{
+		FOnDemandContentInstallerStats::LogAnalytics();
+	})
+);
+#endif // !UE_BUILD_SHIPPING
+
 extern int32 GIasHttpRateLimitKiBPerSecond;
 extern int32 GIasHttpConcurrentRequests;
 
@@ -1552,9 +1564,14 @@ void FOnDemandContentInstallerStats::OnRequestCompleted(
 	}
 }
 
-void FOnDemandContentInstallerStats::ReportAnalytics(TArray<FAnalyticsEventAttribute>& OutAnalyticsArray)
+static void LogOrReportContentInstallerAnalytics(bool bReport, bool bLog, TArray<FAnalyticsEventAttribute>& OutAnalyticsArray)
 {
 	if (GIadReportAnalyticsEnabled == false)
+	{
+		bReport = false;
+	}
+
+	if (bReport == false && bLog == false)
 	{
 		return;
 	}
@@ -1562,15 +1579,15 @@ void FOnDemandContentInstallerStats::ReportAnalytics(TArray<FAnalyticsEventAttri
 	FInstallerAnalytics CurrentInstallerAnalytics;
 	{
 		TUniqueLock Lock(InstallerAnalytics);
-		CurrentInstallerAnalytics	= InstallerAnalytics;
-		InstallerAnalytics			= FInstallerAnalytics();
+		CurrentInstallerAnalytics = InstallerAnalytics;
+		InstallerAnalytics = FInstallerAnalytics();
 	}
 
 	FInstallCacheAnalytics CurrentInstallCacheAnalytics;
 	{
 		TUniqueLock Lock(InstallCacheAnalytics);
-		CurrentInstallCacheAnalytics	= InstallCacheAnalytics;
-		InstallCacheAnalytics			= FInstallCacheAnalytics();
+		CurrentInstallCacheAnalytics = InstallCacheAnalytics;
+		InstallCacheAnalytics = FInstallCacheAnalytics();
 	}
 
 	const FDateTime Now = FDateTime::UtcNow();
@@ -1582,36 +1599,72 @@ void FOnDemandContentInstallerStats::ReportAnalytics(TArray<FAnalyticsEventAttri
 	const double AvgTotalFragmentationRatio = CurrentInstallerAnalytics.InstallCount > 0 ? (double)CurrentInstallerAnalytics.TotalFragmentionRatio / (double)CurrentInstallerAnalytics.InstallCount : 0.0;
 	const double AvgDefraggedBlockFragmentationRatio = CurrentInstallerAnalytics.InstallCount > 0 ? (double)CurrentInstallerAnalytics.DefraggedBlockFragmentionRatio / (double)CurrentInstallerAnalytics.InstallCount : 0.0;
 
-	AppendAnalyticsEventAttributeArray(OutAnalyticsArray
-		, TEXT("IadTotalInstallCount"), CurrentInstallerAnalytics.InstallCount
-		, TEXT("IadTotalInstallErrorCount"), CurrentInstallerAnalytics.InstallErrorCount 
-		, TEXT("IadTotalInstallCacheMissCount"), CurrentInstallerAnalytics.InstallCacheMissCount
-		, TEXT("IadTotalDownloadedBytes"), CurrentInstallerAnalytics.DownloadedBytes 
-		, TEXT("IadTotalInstallDurationMs"), CurrentInstallerAnalytics.TotalInstallDurationMs 
-		, TEXT("IadAvgInstallDurationMs"), AvgInstallDurationMs
-		, TEXT("IadAvgCacheHitRatio"), AvgCacheHitRatio
-		, TEXT("IadAvgTotalFragmentationRatio"), AvgTotalFragmentationRatio
-		, TEXT("IadAvgDefraggedBlockFragmentationRatio"), AvgDefraggedBlockFragmentationRatio
-		, TEXT("IadInstallCacheStartupErrorCode"), CurrentInstallCacheAnalytics.StartupErrorCode
-		, TEXT("IadInstallCacheVerificationRemovedBlockCount"), CurrentInstallCacheAnalytics.VerificationRemovedBlockCount 
-		, TEXT("IadInstallCacheFlushCount"), CurrentInstallCacheAnalytics.FlushCount 
-		, TEXT("IadInstallCacheFlushErrorCount"), CurrentInstallCacheAnalytics.FlushErrorCount 
-		, TEXT("IadInstallCacheFlushedBytes"), CurrentInstallCacheAnalytics.FlushedBytes
-		, TEXT("IadInstallCachePurgeCount"), CurrentInstallCacheAnalytics.PurgeCount
-		, TEXT("IadInstallCachePurgeErrorCount"), CurrentInstallCacheAnalytics.PurgeErrorCount
-		, TEXT("IadInstallCacheDefragCount"), CurrentInstallCacheAnalytics.DefragCount 
-		, TEXT("IadInstallCacheDefragErrorCount"), CurrentInstallCacheAnalytics.DefragErrorCount 
-		, TEXT("IadInstallCacheJournalCommitCount"), CurrentInstallCacheAnalytics.JournalCommitCount 
-		, TEXT("IadInstallCacheJournalCommitErrorCount"), CurrentInstallCacheAnalytics.JournalCommitErrorCount
-		, TEXT("IadInstallCacheMaxSize"), CurrentInstallCacheAnalytics.MaxCacheSize 
-		, TEXT("IadInstallCacheMaxUsageSize"), CurrentInstallCacheAnalytics.MaxCacheUsageSize 
-		, TEXT("IadInstallCacheMaxReferencedBlockSize"), CurrentInstallCacheAnalytics.MaxReferencedBlockSize
-		, TEXT("IadInstallCacheMaxReferencedSize"), CurrentInstallCacheAnalytics.MaxReferencedSize
-		, TEXT("IadInstallCacheMaxFragmentedSize"), CurrentInstallCacheAnalytics.MaxFragmentedSize
-		, TEXT("IadInstallCacheOldestBlockAgeMinutes"), OldestBlockAge.GetTotalMinutes()
-		, TEXT("IadInstallCacheReadCount"), CurrentInstallCacheAnalytics.ReadCount
-		, TEXT("IadInstallCacheReadErrorCount"), CurrentInstallCacheAnalytics.ReadErrorCount
-	);
+	#define INSTALLER_ANALYTICS_LIST(XANALYTIC) \
+		XANALYTIC(IadTotalInstallCount, CurrentInstallerAnalytics.InstallCount) \
+		XANALYTIC(IadTotalInstallErrorCount, CurrentInstallerAnalytics.InstallErrorCount) \
+		XANALYTIC(IadTotalInstallCacheMissCount, CurrentInstallerAnalytics.InstallCacheMissCount) \
+		XANALYTIC(IadTotalDownloadedBytes, CurrentInstallerAnalytics.DownloadedBytes) \
+		XANALYTIC(IadTotalInstallDurationMs, CurrentInstallerAnalytics.TotalInstallDurationMs) \
+		XANALYTIC(IadAvgInstallDurationMs, AvgInstallDurationMs) \
+		XANALYTIC(IadAvgCacheHitRatio, AvgCacheHitRatio) \
+		XANALYTIC(IadAvgTotalFragmentationRatio, AvgTotalFragmentationRatio) \
+		XANALYTIC(IadAvgDefraggedBlockFragmentationRatio, AvgDefraggedBlockFragmentationRatio) \
+		XANALYTIC(IadInstallCacheStartupErrorCode, CurrentInstallCacheAnalytics.StartupErrorCode) \
+		XANALYTIC(IadInstallCacheVerificationRemovedBlockCount, CurrentInstallCacheAnalytics.VerificationRemovedBlockCount) \
+		XANALYTIC(IadInstallCacheFlushCount, CurrentInstallCacheAnalytics.FlushCount) \
+		XANALYTIC(IadInstallCacheFlushErrorCount, CurrentInstallCacheAnalytics.FlushErrorCount) \
+		XANALYTIC(IadInstallCacheFlushedBytes, CurrentInstallCacheAnalytics.FlushedBytes) \
+		XANALYTIC(IadInstallCachePurgeCount, CurrentInstallCacheAnalytics.PurgeCount) \
+		XANALYTIC(IadInstallCachePurgeErrorCount, CurrentInstallCacheAnalytics.PurgeErrorCount) \
+		XANALYTIC(IadInstallCacheDefragCount, CurrentInstallCacheAnalytics.DefragCount) \
+		XANALYTIC(IadInstallCacheDefragErrorCount, CurrentInstallCacheAnalytics.DefragErrorCount) \
+		XANALYTIC(IadInstallCacheJournalCommitCount, CurrentInstallCacheAnalytics.JournalCommitCount) \
+		XANALYTIC(IadInstallCacheJournalCommitErrorCount, CurrentInstallCacheAnalytics.JournalCommitErrorCount) \
+		XANALYTIC(IadInstallCacheMaxSize, CurrentInstallCacheAnalytics.MaxCacheSize) \
+		XANALYTIC(IadInstallCacheMaxUsageSize, CurrentInstallCacheAnalytics.MaxCacheUsageSize) \
+		XANALYTIC(IadInstallCacheMaxReferencedBlockSize, CurrentInstallCacheAnalytics.MaxReferencedBlockSize) \
+		XANALYTIC(IadInstallCacheMaxReferencedSize, CurrentInstallCacheAnalytics.MaxReferencedSize) \
+		XANALYTIC(IadInstallCacheMaxFragmentedSize, CurrentInstallCacheAnalytics.MaxFragmentedSize) \
+		XANALYTIC(IadInstallCacheOldestBlockAgeMinutes, OldestBlockAge.GetTotalMinutes()) \
+		XANALYTIC(IadInstallCacheReadCount, CurrentInstallCacheAnalytics.ReadCount) \
+		XANALYTIC(IadInstallCacheReadErrorCount, CurrentInstallCacheAnalytics.ReadErrorCount)
+
+	if (bReport)
+	{
+		#define REPORT_ANALYTIC(inName, inValue) , TEXT(#inName), inValue
+		AppendAnalyticsEventAttributeArray(OutAnalyticsArray
+			INSTALLER_ANALYTICS_LIST(REPORT_ANALYTIC)
+		);
+		#undef REPORT_ANALYTIC
+	}
+
+	if (bLog)
+	{
+		#define ANALYTIC_LOG_FORMAT(inName, inValue) "stat:\t" #inName " {"#inName"}\n"
+		#define ANALYTIC_LOG_VALUE(inName, inValue) , UE_LOGFMT_FIELD(#inName, inValue)
+
+		UE_LOGFMT_EX(LogIoStoreOnDemand, Log, "IadStats:\n" INSTALLER_ANALYTICS_LIST(ANALYTIC_LOG_FORMAT) INSTALLER_ANALYTICS_LIST(ANALYTIC_LOG_VALUE));
+
+		#undef ANALYTIC_LOG_FORMAT
+		#undef ANALYTIC_LOG_VALUE
+	}
+
+	#undef INSTALLER_ANALYTICS_LIST
+}
+
+void FOnDemandContentInstallerStats::ReportAnalytics(TArray<FAnalyticsEventAttribute>& OutAnalyticsArray)
+{
+	const bool bReportTrue = true;
+	const bool bLogFalse = false;
+	LogOrReportContentInstallerAnalytics(bReportTrue, bLogFalse, OutAnalyticsArray);
+}
+
+void FOnDemandContentInstallerStats::LogAnalytics()
+{
+	const bool bReportFalse = false;
+	const bool bLogTrue = true;
+	TArray<FAnalyticsEventAttribute> UnusedAnalyticsArray;
+	LogOrReportContentInstallerAnalytics(bReportFalse, bLogTrue, UnusedAnalyticsArray);
 }
 
 ////////////////////////////////////////////////////////////////////////////////
diff --git a/Engine/Source/Runtime/Experimental/IoStore/OnDemand/Private/Statistics.h b/Engine/Source/Runtime/Experimental/IoStore/OnDemand/Private/Statistics.h
--- a/Engine/Source/Runtime/Experimental/IoStore/OnDemand/Private/Statistics.h
+++ b/Engine/Source/Runtime/Experimental/IoStore/OnDemand/Private/Statistics.h
@@ -158,6 +158,7 @@ public:
 		double CacheHitRatio,
 		uint64 DurationCycles) IAS_STATISTICS_IMPL()
 	static void ReportAnalytics(TArray<FAnalyticsEventAttribute>& OutAnalyticsArray) IAS_STATISTICS_IMPL()
+	static void LogAnalytics() IAS_STATISTICS_IMPL()
 };
 
 ////////////////////////////////////////////////////////////////////////////////
