PR #13010: Add filtering to get both dynamic and sleeping particles within a field

#PR 13010
#Synced-CL 49661615

diff --git a/Engine/Source/Runtime/Experimental/Chaos/Private/PhysicsProxy/GeometryCollectionPhysicsProxy.cpp b/Engine/Source/Runtime/Experimental/Chaos/Private/PhysicsProxy/GeometryCollectionPhysicsProxy.cpp
--- a/Engine/Source/Runtime/Experimental/Chaos/Private/PhysicsProxy/GeometryCollectionPhysicsProxy.cpp
+++ b/Engine/Source/Runtime/Experimental/Chaos/Private/PhysicsProxy/GeometryCollectionPhysicsProxy.cpp
@@ -2320,6 +2320,16 @@ void FGeometryCollectionPhysicsProxy::GetFilteredParticleHandles(
 				}
 			}
 		}
+		else if (FilterType == EFieldFilterType::Field_Filter_DynamicAndSleeping)
+		{
+			for (FClusterHandle* ClusterHandle : SolverParticleHandles)
+			{
+				if (ClusterHandle && (ClusterHandle->ObjectState() == Chaos::EObjectStateType::Dynamic || ClusterHandle->ObjectState() == Chaos::EObjectStateType::Sleeping))
+				{
+					Handles.Add(ClusterHandle);
+				}
+			}
+		}
 		else if (FilterType == EFieldFilterType::Field_Filter_All)
 		{
 			for (FClusterHandle* ClusterHandle : SolverParticleHandles)
diff --git a/Engine/Source/Runtime/Experimental/Chaos/Private/PhysicsProxy/PerSolverFieldSystem.cpp b/Engine/Source/Runtime/Experimental/Chaos/Private/PhysicsProxy/PerSolverFieldSystem.cpp
--- a/Engine/Source/Runtime/Experimental/Chaos/Private/PhysicsProxy/PerSolverFieldSystem.cpp
+++ b/Engine/Source/Runtime/Experimental/Chaos/Private/PhysicsProxy/PerSolverFieldSystem.cpp
@@ -562,6 +562,23 @@ void FPerSolverFieldSystem::GetFilteredParticleHandles(
 			}
 		}
 	}
+	else if (FilterType == EFieldFilterType::Field_Filter_DynamicAndSleeping)
+	{
+		const Chaos::TParticleView<Chaos::TPBDRigidParticles<Chaos::FReal, 3>>& ParticleView =
+			SolverParticles.GetNonDisabledDynamicView();
+		Handles.Reserve(ParticleView.Num());
+
+		for (Chaos::TParticleIterator<Chaos::TPBDRigidParticles<Chaos::FReal, 3>> It = ParticleView.Begin(), ItEnd = ParticleView.End();
+			It != ItEnd; ++It)
+		{
+			if (ValidateParticle(ObjectType, It))
+			{
+				const Chaos::TTransientGeometryParticleHandle<Chaos::FReal, 3>* Handle = &(*It);
+				Handles.Add(GetHandleHelper(const_cast<Chaos::TTransientGeometryParticleHandle<Chaos::FReal, 3>*>(Handle)));
+				AddClusterChildren(Handles, RigidSolver, It->CastToClustered());
+			}
+		}
+	}
 	else if (FilterType == EFieldFilterType::Field_Filter_All)
 	{
 		const Chaos::TParticleView<Chaos::FGeometryParticles>& ParticleView =
diff --git a/Engine/Source/Runtime/Experimental/Chaos/Public/Field/FieldSystemTypes.h b/Engine/Source/Runtime/Experimental/Chaos/Public/Field/FieldSystemTypes.h
--- a/Engine/Source/Runtime/Experimental/Chaos/Public/Field/FieldSystemTypes.h
+++ b/Engine/Source/Runtime/Experimental/Chaos/Public/Field/FieldSystemTypes.h
@@ -90,6 +90,7 @@ enum EFieldFilterType : int
 	Field_Filter_All  UMETA(DisplayName = "All", ToolTip = "Apply the field to all the solver particles"),
 	Field_Filter_Sleeping  UMETA(DisplayName = "Sleeping", ToolTip = "Apply the field to all the sleeping particles"),
 	Field_Filter_Disabled  UMETA(DisplayName = "Disabled", ToolTip = "Apply the field to all the disabled particles"),
+	Field_Filter_DynamicAndSleeping  UMETA(DisplayName = "Dynamic & Sleeping", ToolTip = "Apply the field to all the dynamic and sleeping particles"),
 	//~~~
 	//256th entry
 	Field_Filter_Max      UMETA(Hidden)
