Dataflow  : add new evaluators, enable new evaluation for the fracture editor as a test
For now only the Tick based evaluator is being used by the fracture editor, other evaluator will be put in use at a later stage

#Synced-CL 49661615

diff --git a/Engine/Plugins/Experimental/ChaosEditor/Source/FractureEditor/Private/FractureEditorModeToolkit.cpp b/Engine/Plugins/Experimental/ChaosEditor/Source/FractureEditor/Private/FractureEditorModeToolkit.cpp
--- a/Engine/Plugins/Experimental/ChaosEditor/Source/FractureEditor/Private/FractureEditorModeToolkit.cpp
+++ b/Engine/Plugins/Experimental/ChaosEditor/Source/FractureEditor/Private/FractureEditorModeToolkit.cpp
@@ -8,6 +8,7 @@
 
 #include "Engine/Selection.h"
 #include "Dataflow/DataflowInstanceDetails.h"
+#include "Dataflow/DataflowObjectInterface.h"
 #include "FractureEditorMode.h"
 #include "Framework/MultiBox/MultiBoxBuilder.h"
 #include "Framework/MultiBox/MultiBoxDefs.h"
@@ -80,6 +81,7 @@
 #include "FractureSettings.h"
 #include "Toolkits/AssetEditorModeUILayer.h"
 #include "SPrimaryButton.h"
+#include "Widgets/Notifications/SProgressBar.h"
 
 #include "Algo/RemoveIf.h"
 
@@ -87,6 +89,13 @@
 
 #define LOCTEXT_NAMESPACE "FFractureEditorModeToolkit"
 
+namespace FractureEditorModeToolkit::Private
+{
+	bool bUseNewDataflowEvaluator = true;
+	FAutoConsoleVariableRef CVARDataflowEnableGraphEval(TEXT("p.FractureEditor.UseNewDataflowEvaluator"), bUseNewDataflowEvaluator,
+		TEXT("Enable new dataflow evaluator for the fracture editor. [def:true]"));
+}
+
 TArray<UClass*> FindFractureToolClasses()
 {
 	TArray<UClass*> Classes;
@@ -157,7 +166,7 @@ FFractureViewSettingsCustomization::CustomizeDetails(IDetailLayoutBuilder& Detai
 			})
 		]
 	];
-};
+};	
 
 FHistogramSettingsCustomization::FHistogramSettingsCustomization(FFractureEditorModeToolkit* InToolkit)
 	: Toolkit(InToolkit)
@@ -207,8 +216,14 @@ FFractureEditorModeToolkit::~FFractureEditorModeToolkit()
 		auto& LevelEditorModule = FModuleManager::LoadModuleChecked<FLevelEditorModule>(TEXT("LevelEditor"));
 		LevelEditorModule.OnMapChanged().RemoveAll(this);
 	}
-	UFractureModeSettings* Settings = GetMutableDefault<UFractureModeSettings>();
-	Settings->OnModified.Remove(ProjectSettingsModifiedHandle);
+	if (UFractureModeSettings* Settings = GetMutableDefault<UFractureModeSettings>())
+	{
+		Settings->OnModified.Remove(ProjectSettingsModifiedHandle);
+	}
+	if (OverridesDetailsView)
+	{
+		OverridesDetailsView->GetOnFinishedChangingPropertiesDelegate().RemoveAll(this);
+	}
 }
 
 void FFractureEditorModeToolkit::UpdateAssetLocationMode(TSharedPtr<FString> NewString)
@@ -285,10 +300,27 @@ void FFractureEditorModeToolkit::CreateVariableOverrideDetailView()
 		};
 	FOnGetDetailCustomizationInstance DataflowInstanceDetailsCustomizationInstance = FOnGetDetailCustomizationInstance::CreateLambda(MakeDataflowInstanceDetailCustomizationLambda);
 	OverridesDetailsView->GetDetailsView()->RegisterInstancedCustomPropertyLayout(FDataflowInstance::StaticStruct(), DataflowInstanceDetailsCustomizationInstance);
+	OnFinishedChangingVariableOverridesDelegateHandle = OverridesDetailsView->GetOnFinishedChangingPropertiesDelegate().AddSP(this, &FFractureEditorModeToolkit::OnPropertyValueChanged);
 
 	RefreshVariableOverrideDetailView(nullptr);
 }
 
+void FFractureEditorModeToolkit::OnPropertyValueChanged(const FPropertyChangedEvent& PropertyChangedEvent)
+{
+	if (UFractureSettings* FractureSettings = GetMutableDefault<UFractureSettings>())
+	{
+		if (TStrongObjectPtr<const UGeometryCollection> RestCollection = FractureSettings->RestCollection.Pin())
+		{
+			// we need a modifiable object
+			if (UGeometryCollection* EditableRestCollection = const_cast<UGeometryCollection*>(RestCollection.Get()))
+			{
+				// Recreate the context for now
+ 				DataflowContext = MakeShared<UE::Dataflow::FEngineContext>(EditableRestCollection);
+			}
+		}
+	}
+}
+
 void FFractureEditorModeToolkit::RefreshVariableOverrideDetailView(const UGeometryCollection* RestCollection)
 {
 	if (RestCollection == nullptr)
@@ -304,6 +336,12 @@ void FFractureEditorModeToolkit::RefreshVariableOverrideDetailView(const UGeomet
 		OverridesDetailsView->GetDetailsView()->SetObject(EditableRestCollection);
 		TSharedPtr<FStructOnScope> StructOnScope = Interface->GetDataflowInstance().MakeStructOnScope();
 		OverridesDetailsView->SetStructureData(StructOnScope);
+
+		// Recreate the context 
+		if (!DataflowContext.IsValid() || !IsValid(DataflowContext->Owner) || DataflowContext->Owner != EditableRestCollection)
+		{
+			DataflowContext = MakeShared<UE::Dataflow::FEngineContext>(EditableRestCollection);
+		}
 	}
 }
 
@@ -313,7 +351,29 @@ FReply FFractureEditorModeToolkit::OnDataflowOverridesUpdateAsset()
 	UFractureSettings* FractureSettings = GetMutableDefault<UFractureSettings>();
 	if (FractureSettings && FractureSettings->RestCollection.IsValid())
 	{
-		FractureSettings->RestCollection->GetDataflowInstance().UpdateOwnerAsset(/* bUpdateDependentAssets */true);
+		if (FractureEditorModeToolkit::Private::bUseNewDataflowEvaluator)
+		{
+			auto OnEvaluationCompleted = [DataflowContext = DataflowContext](UE::Dataflow::FContext&)
+				{
+					if (DataflowContext)
+					{
+						DataflowContext->UpdateAssetFromTerminalNodes(true);
+					}
+				};
+			DataflowEvaluator.Reset();
+			if (DataflowContext)
+			{
+				DataflowEvaluator = DataflowContext->EvaluateGraph<UE::Dataflow::FDataflowTickEvaluator>(OnEvaluationCompleted);
+				if (!DataflowEvaluator.IsValid())
+				{
+					UE_LOG(LogFractureTool, Error, TEXT("FractureMode: Failed to start evaluating Dataflow graph for geometry collection %s"), *FractureSettings->RestCollection->GetFName().ToString());
+				}
+			}
+		}
+		else
+		{
+			FractureSettings->RestCollection->GetDataflowInstance().UpdateOwnerAsset(/* bUpdateDependentAssets */true);
+		}
 	}
 	return FReply::Handled();
 }
@@ -383,6 +443,7 @@ void FFractureEditorModeToolkit::Init(const TSharedPtr<IToolkitHost>& InitToolki
 	OutlinerDetailsView->SetObject(GetMutableDefault<UOutlinerSettings>());
 
 	CreateVariableOverrideDetailView();
+	CreateDataflowProgressBar();
 
 	AssetLocationModes.Reset();
 	AssetLocationModes.Add(MakeShared<FString>(TEXT("Source Asset Folder")));
@@ -460,6 +521,11 @@ void FFractureEditorModeToolkit::Init(const TSharedPtr<IToolkitHost>& InitToolki
 		]
 		+ SVerticalBox::Slot()
 		.AutoHeight()
+		[
+			DataflowProgressBar.ToSharedRef()
+		]
+		+ SVerticalBox::Slot()
+		.AutoHeight()
 		[
 			SNew(SHorizontalBox)
 			+ SHorizontalBox::Slot()
@@ -2466,4 +2532,100 @@ void FFractureEditorModeToolkit::GetStatisticsSummary(FGeometryCollectionStatist
 	}
 }
 
+void FFractureEditorModeToolkit::CreateDataflowProgressBar()
+{
+	const FSlateBrush* OverlayBrush = FAppStyle::Get().GetBrush("EditorViewport.OverlayBrush");
+
+	auto GetBarVisibility = [this]() ->EVisibility
+		{
+			const bool bVisible = DataflowEvaluator ? DataflowEvaluator->IsRunning() : false;
+			return bVisible ? EVisibility::SelfHitTestInvisible : EVisibility::Collapsed;
+		};
+
+	auto GetProgressText = [this]() -> FText
+		{
+			if (DataflowEvaluator)
+			{
+				return FText::Format(
+					LOCTEXT("DataflowGraphEditor_ProgressText", "Total Progress : {0}% - {1}")
+					, DataflowEvaluator->GetProgress() * 100.f
+					, DataflowEvaluator->GetProgressMessage()
+				);
+			}
+			return FText::GetEmpty();
+		};
+
+	auto GetProgressPercent = [this]() -> float
+		{
+			return DataflowEvaluator ? DataflowEvaluator->GetProgress() : 0.f;
+		};
+
+	auto OnCancel = [this]() -> FReply
+		{
+			if (DataflowEvaluator)
+			{
+				DataflowEvaluator->Cancel();
+			}
+			return FReply::Handled();
+		};
+
+	SAssignNew(DataflowProgressBar, SHorizontalBox)
+		.Visibility_Lambda(GetBarVisibility)
+
+		+ SHorizontalBox::Slot()
+		.HAlign(HAlign_Left)
+		.VAlign(VAlign_Bottom)
+		.FillWidth(1.0)
+		.Padding(FMargin(0.0f, 0.0f, 0.f, 15.f))
+		[
+			SNew(SBorder)
+			.BorderImage(OverlayBrush)
+			[
+				SNew(SVerticalBox)
+
+				// Text 
+				+ SVerticalBox::Slot()
+				.AutoHeight()
+				.Padding(4)
+				.VAlign(VAlign_Center)
+				[
+					// Progress text
+					SNew(STextBlock)
+					.Text_Lambda(GetProgressText)
+				]
+
+				// Progress Bar + Cancel button
+				+ SVerticalBox::Slot()
+				.AutoHeight()
+				.Padding(4)
+				.VAlign(VAlign_Center)
+				[
+					SNew(SHorizontalBox)
+
+					// Progress bar 
+					+ SHorizontalBox::Slot()
+					.FillWidth(1.0)
+					.Padding(4)
+					[
+						SNew(SProgressBar)
+						.Percent_Lambda(GetProgressPercent)
+					]
+
+					// Cancel button
+					+ SHorizontalBox::Slot()
+					.AutoWidth()
+					.Padding(4)
+					[
+						SNew(SButton)
+						.Text(LOCTEXT("DataflowGraphEditor_CancelText", "Cancel"))
+						.ToolTipText(LOCTEXT("DataflowGraphEditor_CancelTooltip", "Cancel current evaluation"))
+						.HAlign(HAlign_Center)
+						.VAlign(VAlign_Center)
+						.OnClicked_Lambda(OnCancel)
+					]
+				]
+			]
+		];
+}
+
 #undef LOCTEXT_NAMESPACE
diff --git a/Engine/Plugins/Experimental/ChaosEditor/Source/FractureEditor/Public/FractureEditorModeToolkit.h b/Engine/Plugins/Experimental/ChaosEditor/Source/FractureEditor/Public/FractureEditorModeToolkit.h
--- a/Engine/Plugins/Experimental/ChaosEditor/Source/FractureEditor/Public/FractureEditorModeToolkit.h
+++ b/Engine/Plugins/Experimental/ChaosEditor/Source/FractureEditor/Public/FractureEditorModeToolkit.h
@@ -4,6 +4,7 @@
 
 #include "CoreMinimal.h"
 #include "Chaos/Real.h"
+#include "Dataflow/DataflowObjectInterface.h"
 #include "Widgets/DeclarativeSyntaxSupport.h"
 #include "Toolkits/BaseToolkit.h"
 #include "UnrealEdMisc.h"
@@ -247,11 +248,15 @@ private:
 	UE_API void UpdateOutlinerHeader();
 
 	UE_API void CreateVariableOverrideDetailView();
+	UE_API void OnPropertyValueChanged(const FPropertyChangedEvent& PropertyChangedEvent);
 	UE_API void RefreshVariableOverrideDetailView(const UGeometryCollection* RestCollection);
 
 	UE_API FReply OnDataflowOverridesUpdateAsset();
 	UE_API bool CanDataflowOverridesUpdateAsset() const;
 	UE_API bool GetDataflowOverridesUpdateAssetEnabled() const;
+
+	void CreateDataflowProgressBar();
+
 private:
 	TObjectPtr<UFractureModalTool> ActiveTool;
 
@@ -275,9 +280,15 @@ private:
 	TSharedPtr<SGeometryCollectionStatistics> StatisticsView;
 	TArray<TSharedPtr<FString>> AssetLocationModes;
 	TSharedPtr<STextComboBox> AssetLocationMode;
+	TSharedPtr<SWidget> DataflowProgressBar;
+
+	TSharedPtr<UE::Dataflow::FEngineContext> DataflowContext;
+	TSharedPtr<UE::Dataflow::IDataflowEvaluator> DataflowEvaluator;
 
 	// Simple cached statistics to allow use to quickly/heuristically check for stale geometry collection data
 	int64 OutlinerCachedBoneCount = 0, OutlinerCachedVertexCount = 0, OutlinerCachedHullCount = 0;
+
+	FDelegateHandle OnFinishedChangingVariableOverridesDelegateHandle;
 };
 
 #undef UE_API
diff --git a/Engine/Source/Runtime/Experimental/Dataflow/Core/Private/Dataflow/DataflowContextEvaluator.cpp b/Engine/Source/Runtime/Experimental/Dataflow/Core/Private/Dataflow/DataflowContextEvaluator.cpp
--- a/Engine/Source/Runtime/Experimental/Dataflow/Core/Private/Dataflow/DataflowContextEvaluator.cpp
+++ b/Engine/Source/Runtime/Experimental/Dataflow/Core/Private/Dataflow/DataflowContextEvaluator.cpp
@@ -1,7 +1,13 @@
 // Copyright Epic Games, Inc. All Rights Reserved.
 
 #include "Dataflow/DataflowContextEvaluator.h"
+#include "Dataflow/DataflowGraph.h"
 #include "Dataflow/DataflowNode.h"
+#include "Dataflow/DataflowCompiledGraph.h"
+
+#define LOCTEXT_NAMESPACE "DataflowEvaluator"
+
+DEFINE_LOG_CATEGORY_STATIC(LogDataflowEvaluator, Log, All);
 
 namespace UE::Dataflow
 {
@@ -29,10 +35,10 @@ namespace UE::Dataflow
 	{
 		if (IsScheduledOrRunning(Entry.Id))
 		{
-			UE_LOG(LogTemp, Warning, TEXT("FContextEvaluator::ScheduleEvaluation : skipped [%s]"), *Entry.ToString());
+			UE_LOG(LogDataflowEvaluator, Verbose, TEXT("FContextEvaluator::ScheduleEvaluation : skipped [%s]"), *Entry.ToString());
 			return;
 		}
-		UE_LOG(LogTemp, Warning, TEXT("FContextEvaluator::ScheduleEvaluation : [%s]"), *Entry.ToString());
+		UE_LOG(LogDataflowEvaluator, Verbose, TEXT("FContextEvaluator::ScheduleEvaluation : [%s]"), *Entry.ToString());
 		PendingEvaluationEntries.Add(Entry.Id, Entry);
 
 		// add upstream nodes
@@ -42,7 +48,7 @@ namespace UE::Dataflow
 			FindInvalidUpstreamNodes(*Node, InvalidUpstreamNodes);
 			for (const FDataflowNode* UpstreamNode: InvalidUpstreamNodes)
 			{
-				UE_LOG(LogTemp, Warning, TEXT("FContextEvaluator::ScheduleEvaluation :  [%s] -- Invalid Upstream Node [%s]"), *Entry.ToString(), *UpstreamNode->GetName().ToString());
+				UE_LOG(LogDataflowEvaluator, Verbose, TEXT("FContextEvaluator::ScheduleEvaluation :  [%s] -- Invalid Upstream Node [%s]"), *Entry.ToString(), *UpstreamNode->GetName().ToString());
 			}
 			for (const FDataflowNode* UpstreamNode : InvalidUpstreamNodes)
 			{
@@ -76,7 +82,7 @@ namespace UE::Dataflow
 			{
 				if (const FDataflowOutput* UpstreamOutput = Input->GetConnection())
 				{
-					UE_LOG(LogTemp, Warning, TEXT("FContextEvaluator::FindInvalidUpstreamOutputs :  [%s] input[%s] -> output [%s]"),
+					UE_LOG(LogDataflowEvaluator, Verbose, TEXT("FContextEvaluator::FindInvalidUpstreamOutputs :  [%s] input[%s] -> output [%s]"),
 						*Node.GetName().ToString(),
 						*Input->GetName().ToString(),
 						*UpstreamOutput->GetName().ToString()
@@ -152,10 +158,10 @@ namespace UE::Dataflow
 
 		if (NumScheduleTasks == 0)
 		{
-			UE_LOG(LogTemp, Warning, TEXT("FContextEvaluator::Process : No Task Scheduled NumPendingTasks=[%d]"), PendingEvaluationEntries.Num());
+			UE_LOG(LogDataflowEvaluator, Verbose, TEXT("FContextEvaluator::Process : No Task Scheduled NumPendingTasks=[%d]"), PendingEvaluationEntries.Num());
 			for (const TPair<FNodeId, FEvaluationEntry>& Entry : PendingEvaluationEntries)
 			{
-				UE_LOG(LogTemp, Warning, TEXT("FContextEvaluator::Process : \t -[%s]"), *Entry.Value.ToString());
+				UE_LOG(LogDataflowEvaluator, Verbose, TEXT("FContextEvaluator::Process : \t -[%s]"), *Entry.Value.ToString());
 			}
 		}
 
@@ -188,7 +194,7 @@ namespace UE::Dataflow
 
 			const bool bUseGameThread = ShouldRunOnGameThread(*Node);
 
-			UE_LOG(LogTemp, Warning, TEXT("FContextEvaluator::ScheduleTask : [%s] GameThread=[%d] previousTasks=[%d]"),
+			UE_LOG(LogDataflowEvaluator, Verbose, TEXT("FContextEvaluator::ScheduleTask : [%s] GameThread=[%d] previousTasks=[%d]"),
 				*Entry.ToString(),
 				(int32)bUseGameThread,
 				(int32)ExistingTasks.Num()
@@ -203,25 +209,9 @@ namespace UE::Dataflow
 						if (TSharedPtr<const FDataflowNode> Node = Entry.WeakNode.Pin())
 						{
 							Node->SetAsyncEvaluating(true);
-							if (Node->NumOutputs() == 0)
-							{
-								Node->Evaluate(*ContextPtr, nullptr);
-							}
-							else
-							{
-								const TArray<FDataflowOutput*> Outputs = Node->GetOutputs();
-								for (const FDataflowOutput* Output : Outputs)
-								{
-									// todo(ccaillaud) : should we check if the oputput is frozen?
-									// todo(ccaillaud) : shoudl we have an option to skip the non connected outputs ? 
-									if (!Output->HasValidData(*ContextPtr))
-									{
-										Node->Evaluate(*ContextPtr, Output);
-									}
-								}
-							}
+							Node->EvaluateNode(*ContextPtr);
 							Node->SetAsyncEvaluating(false);
-							UE_LOG(LogTemp, Warning, TEXT("FContextEvaluator::EndTask : [%s]"), *Entry.ToString());
+							UE_LOG(LogDataflowEvaluator, Verbose, TEXT("FContextEvaluator::EndTask : [%s]"), *Entry.ToString());
 						}
 					}
 				},
@@ -262,5 +252,324 @@ namespace UE::Dataflow
 		}
 		return NodeName.ToString();
 	}
+
+	////////////////////////////////////////////////////////////////////////////////////////////////
+
+	float FThreadSafeProgressTracker::GetProgress() const
+	{
+		FScopeLock Lock(&ProgressLock);
+		return Progress;
+	}
+
+	FText FThreadSafeProgressTracker::GetProgressMessage() const
+	{
+		FScopeLock Lock(&ProgressLock);
+		return FText::Format(LOCTEXT("DirectEvaluator_ProgressMsg_Format", "Evaluating {0}"), FText::FromName(NodeName));
+	}
+
+	void FThreadSafeProgressTracker::SetProgress(float InCurrent, float InTotal, FName InNodeName)
+	{
+		FScopeLock Lock(&ProgressLock);
+		Progress = (InTotal > 0)
+			? FMath::Clamp(InCurrent / InTotal, 0.0f, 1.0f)
+			: 0.0f;
+		NodeName = InNodeName;
+	}
+
+	void FThreadSafeProgressTracker::SetProgress(float InProgress, FName InNodeName)
+	{
+		FScopeLock Lock(&ProgressLock);
+		Progress = InProgress;
+		NodeName = InNodeName;
+	}
+
+	////////////////////////////////////////////////////////////////////////////////////////////////
+
+	FDataflowDirectEvaluator::FDataflowDirectEvaluator(FContext& InOwningContext)
+		: OwningContext(InOwningContext)
+	{
+	}
+
+	void FDataflowDirectEvaluator::Start(TSharedRef<const FGraph> InDataflowGraph, TSharedRef<const FCompiledGraph> InCompiledGraph, const FGuid& InNodeId, FOnPostEvaluationFunction OnPostEvaluation)
+	{
+		bIsCanceled = false;
+		bIsRunning = true;
+
+		int32 LastNodeIndex = INDEX_NONE;
+		FGuid NodeId;
+		do
+		{
+			NodeId = InCompiledGraph->GetNextNodeId(LastNodeIndex++);
+			if (NodeId.IsValid())
+			{
+				if (TSharedPtr<const FDataflowNode> Node = InDataflowGraph->FindBaseNode(NodeId))
+				{
+					ProgressTracker.SetProgress((float)LastNodeIndex, (float)InCompiledGraph->GetNumNodes(), Node->GetName());
+
+					if (ensure(!Node->EvaluateOnGameThreadOnly()))
+					{
+						Node->EvaluateNode(OwningContext);
+					}
+				}
+			}
+		} while (NodeId.IsValid() && !bIsCanceled);
+
+		bIsRunning = false;
+
+		// No more nodes to process
+		if (!bIsCanceled && OnPostEvaluation.IsSet())
+		{
+			OnPostEvaluation(OwningContext);
+		}
+	}
+
+	void FDataflowDirectEvaluator::Cancel()
+	{
+		bIsCanceled = true;
+	}
+
+	bool FDataflowDirectEvaluator::IsRunning() const
+	{
+		return bIsRunning;
+	}
+
+	float FDataflowDirectEvaluator::GetProgress() const
+	{
+		return ProgressTracker.GetProgress();
+	}
+
+	FText FDataflowDirectEvaluator::GetProgressMessage() const
+	{
+		return ProgressTracker.GetProgressMessage();
+	}
+
+	////////////////////////////////////////////////////////////////////////////////////////////////
+
+	FDataflowTickEvaluator::FDataflowTickEvaluator(FContext& InOwningContext)
+		: OwningContext(InOwningContext)
+	{
+	}
+
+	FDataflowTickEvaluator::~FDataflowTickEvaluator()
+	{
+		FTSTicker::GetCoreTicker().RemoveTicker(TickDelegateHandle);
+		TickDelegateHandle.Reset();
+	}
+
+	void FDataflowTickEvaluator::Reset()
+	{
+		check(IsInGameThread());
+
+		CompiledGraph.Reset();
+		DataflowGraph.Reset();
+		LastNodeIndex = INDEX_NONE;
+		NodeId = FGuid();
+		PostEvaluationFunction.Reset();
+		bIsRunning = false;
+
+		FTSTicker::GetCoreTicker().RemoveTicker(TickDelegateHandle);
+		TickDelegateHandle.Reset();
+	}
+
+	void FDataflowTickEvaluator::Start(TSharedRef<const FGraph> InDataflowGraph, TSharedRef<const FCompiledGraph> InCompiledGraph, const FGuid& InNodeId, FOnPostEvaluationFunction OnPostEvaluation)
+	{
+		Reset();
+
+		DataflowGraph = InDataflowGraph;
+		CompiledGraph = InCompiledGraph;
+		NodeId = InNodeId;
+		PostEvaluationFunction = OnPostEvaluation;
+		LastNodeIndex = INDEX_NONE;
+
+		TickDelegateHandle = FTSTicker::GetCoreTicker().AddTicker(FTickerDelegate::CreateRaw(this, &FDataflowTickEvaluator::Tick));
+
+		bIsCanceled = false;
+		bIsRunning = true;
+	}
+
+	void FDataflowTickEvaluator::Cancel()
+	{
+		bIsCanceled = true;
+	}
+
+	bool FDataflowTickEvaluator::IsRunning() const
+	{
+		return bIsRunning;
+	}
+
+	float FDataflowTickEvaluator::GetProgress() const
+	{
+		return ProgressTracker.GetProgress();
+	}
+
+	FText FDataflowTickEvaluator::GetProgressMessage() const
+	{
+		return ProgressTracker.GetProgressMessage();
+	}
+
+	void FDataflowTickEvaluator::OnEvaluationCompleted()
+	{
+		// Save the function because reset reset it 
+		// and we need to call reset first before calling the call to the post evaluation function 
+		FOnPostEvaluationFunction FunctionToCall = PostEvaluationFunction;
+		Reset();
+		if (FunctionToCall.IsSet())
+		{
+			FunctionToCall(OwningContext);
+		}		
+	}
+
+	void FDataflowTickEvaluator::OnEvaluationCanceled()
+	{
+		Reset();
+	}
+
+	bool FDataflowTickEvaluator::Tick(float DeltaTime)
+	{
+		constexpr double TickTimeBudget = 5.0;
+
+		UE_LOG(LogDataflowEvaluator, Verbose, TEXT("FDataflowTickEvaluator::Tick -- START"));
+
+		const double StartTimeMs = FPlatformTime::ToMilliseconds64(FPlatformTime::Cycles64());
+		double CurrentTimeMs = StartTimeMs;
+
+		while (!bIsCanceled && CompiledGraph && DataflowGraph)
+		{
+			const FGuid NextNodeId = CompiledGraph->GetNextNodeId(LastNodeIndex);
+			if (!NextNodeId.IsValid())
+			{
+				OnEvaluationCompleted();
+				break;
+			}
+
+			if (TSharedPtr<const FDataflowNode> Node = DataflowGraph->FindBaseNode(NextNodeId))
+			{
+				ProgressTracker.SetProgress((float)LastNodeIndex, (float)CompiledGraph->GetNumNodes(), Node->GetName());
+
+				const double StartEvalTimeMs = FPlatformTime::ToMilliseconds64(FPlatformTime::Cycles64());
+
+				Node->EvaluateNode(OwningContext);
+
+				const double EndEvalTimeMs = FPlatformTime::ToMilliseconds64(FPlatformTime::Cycles64());
+
+				UE_LOG(LogDataflowEvaluator, Verbose, TEXT("FDataflowTickEvaluator::Tick --     Node [%s] : [%.3f ms]"), *Node->GetName().ToString(), (float)(EndEvalTimeMs - StartEvalTimeMs));
+			}
+			LastNodeIndex++;
+
+			// timeout check
+			CurrentTimeMs = FPlatformTime::ToMilliseconds64(FPlatformTime::Cycles64());
+			if ((CurrentTimeMs - StartTimeMs) > TickTimeBudget)
+			{
+				break;
+			}
+		}
+
+		UE_LOG(LogDataflowEvaluator, Verbose, TEXT("FDataflowTickEvaluator::Tick -- END : [%.3f ms]"), (float)(CurrentTimeMs - StartTimeMs));
+
+		if (bIsCanceled)
+		{
+			OnEvaluationCanceled();
+			return false;
+		}
+
+		return true;
+	}
+
+	////////////////////////////////////////////////////////////////////////////////////////////////
+
+	FDataflowTaskGraphEvaluator::FDataflowTaskGraphEvaluator(FContext& InOwningContext)
+		: OwningContext(InOwningContext)
+		, CancellationToken(MakeShared<UE::Tasks::FCancellationToken>())
+	{
+	}
+
+	void FDataflowTaskGraphEvaluator::Start(TSharedRef<const FGraph> InDataflowGraph, TSharedRef<const FCompiledGraph> InCompiledGraph, const FGuid& InNodeId, FOnPostEvaluationFunction OnPostEvaluation)
+	{
+		check(!IsInGameThread());
+
+		// build the taskgraph from the compiled graph
+		const int32 NumTasks = InCompiledGraph->GetNumTasks();
+
+		CancellationToken = MakeShared<UE::Tasks::FCancellationToken>();
+		TArray<UE::Tasks::FTask> AllTasks;
+		AllTasks.SetNum(NumTasks);
+
+		TArray<FGuid> NodeIds;
+		TArray<int32> ParentTaskIds;
+		TArray<TWeakPtr<const FDataflowNode>> NodeWeakPtrs;
+		TArray<UE::Tasks::FTask> PreRequisites;
+
+		for (int32 TaskIndex = 0; TaskIndex < NumTasks; ++TaskIndex)
+		{
+			InCompiledGraph->GetTasksNodesAndDependencies(TaskIndex, NodeIds, ParentTaskIds);
+
+			NodeWeakPtrs.Reset();
+			for (const FGuid& NodeId : NodeIds)
+			{
+				if (TSharedPtr<const FDataflowNode> Node = InDataflowGraph->FindBaseNode(NodeId))
+				{
+					NodeWeakPtrs.Add(Node);
+				}
+			}
+
+			PreRequisites.Reset();
+			for (const int32& ParentTaskId : ParentTaskIds)
+			{
+				if (AllTasks.IsValidIndex(ParentTaskId) && AllTasks[ParentTaskId].IsValid())
+				{
+					PreRequisites.Add(AllTasks[ParentTaskId]);
+				}
+			}
+			// TODO : maybe we shaould have the context as a shared Ptr
+			auto TaskBody = [NodeWeakPtrs = MoveTemp(NodeWeakPtrs), CancellationToken = CancellationToken, &OwningContext = OwningContext]()
+				{
+					if (!CancellationToken->IsCanceled())
+					{
+						for (TWeakPtr<const FDataflowNode> NodeWeakPtr : NodeWeakPtrs)
+						{
+							if (TSharedPtr<const FDataflowNode> NodePtr = NodeWeakPtr.Pin())
+							{
+								NodePtr->EvaluateNode(OwningContext);
+							}
+							if (CancellationToken->IsCanceled())
+							{
+								break;
+							}
+						}
+					}
+				};
+			AllTasks[TaskIndex] = UE::Tasks::Launch(TEXT("DataflowEvaluationTask"), MoveTemp(TaskBody), PreRequisites);
+		}
+
+		auto CompletionTaskBody = [OnPostEvaluation, &OwningContext = OwningContext]()
+			{
+				if (OnPostEvaluation.IsSet())
+				{
+					OnPostEvaluation(OwningContext);
+				}
+			};
+		CompletionTask = UE::Tasks::Launch(TEXT("DataflowCompletionTask"), MoveTemp(CompletionTaskBody), AllTasks);
+	}
+
+	void FDataflowTaskGraphEvaluator::Cancel()
+	{
+		CancellationToken->Cancel();
+	}
+
+	bool FDataflowTaskGraphEvaluator::IsRunning() const
+	{
+		return !CompletionTask.IsCompleted();
+	}
+
+	float FDataflowTaskGraphEvaluator::GetProgress() const
+	{
+		return ProgressTracker.GetProgress();
+	}
+
+	FText FDataflowTaskGraphEvaluator::GetProgressMessage() const
+	{
+		return ProgressTracker.GetProgressMessage();
+	}
 };
 
+#undef LOCTEXT_NAMESPACE
diff --git a/Engine/Source/Runtime/Experimental/Dataflow/Core/Private/Dataflow/DataflowNode.cpp b/Engine/Source/Runtime/Experimental/Dataflow/Core/Private/Dataflow/DataflowNode.cpp
--- a/Engine/Source/Runtime/Experimental/Dataflow/Core/Private/Dataflow/DataflowNode.cpp
+++ b/Engine/Source/Runtime/Experimental/Dataflow/Core/Private/Dataflow/DataflowNode.cpp
@@ -729,7 +729,7 @@ void FDataflowNode::EvaluateNode(UE::Dataflow::FContext& Context) const
 				if (const FDataflowOutput* const Output = OutputByKey.Value)
 				{
 					// todo(ccaillaud) : shoudl we have an option to skip the non connected outputs ? 
-					if (Context.HasData(Output->CacheKey(), GetTimestamp()))
+					if (!Context.HasData(Output->CacheKey(), GetTimestamp()))
 					{
 						Evaluate(Context, Output);
 					}
diff --git a/Engine/Source/Runtime/Experimental/Dataflow/Core/Public/Dataflow/DataflowContextEvaluator.h b/Engine/Source/Runtime/Experimental/Dataflow/Core/Public/Dataflow/DataflowContextEvaluator.h
--- a/Engine/Source/Runtime/Experimental/Dataflow/Core/Public/Dataflow/DataflowContextEvaluator.h
+++ b/Engine/Source/Runtime/Experimental/Dataflow/Core/Public/Dataflow/DataflowContextEvaluator.h
@@ -7,18 +7,23 @@
 #include "Misc/Guid.h"
 #include "Templates/Function.h"
 #include "Templates/SharedPointer.h"
+#include "Containers/Ticker.h"
+#include "Tasks/Task.h"
 
 struct FDataflowNode;
 struct FDataflowOutput;
 
 namespace UE::Dataflow
 {
+	struct FCompiledGraph;
+	class FGraph;
+
 	class FContext;
 
 	using FOnPostEvaluationFunction = TFunction<void(FContext&)>;
 
 	/*
-	* Asynchronously evaluate a dataflow nodes 
+	* Asynchronously evaluate a dataflow nodes
 	* This may be slower than executing the graph in one go synchronously but this offers the following advantage:
 	* - this can be cancelled at anytime ( only the in progress node will have finish evaluating while all the pending ones will be discarded )
 	* - nodes that need to run on the game thread will do so ( see FDataflowNode EvaluateOnGameThreadOnly method )
@@ -48,7 +53,8 @@ namespace UE::Dataflow
 
 		FContextEvaluator(FContext& InOwningContext)
 			: OwningContext(InOwningContext)
-		{}
+		{
+		}
 
 		void ScheduleNodeEvaluation(const FDataflowNode& Node, FOnPostEvaluationFunction OnPostEvaluation);
 		void ScheduleOutputEvaluation(const FDataflowOutput& Output, FOnPostEvaluationFunction OnPostEvaluation);
@@ -77,6 +83,135 @@ namespace UE::Dataflow
 		TMap<FNodeId, FGraphEventRef> RunningTasks;
 		TSet<FNodeId> CompletedTasks;
 	};
+
+	struct IDataflowEvaluator
+	{
+		virtual ~IDataflowEvaluator() = default;
+		virtual void Start(TSharedRef<const FGraph> InDataflowGraph, TSharedRef<const FCompiledGraph> InCompiledGraph, const FGuid& InNodeId, FOnPostEvaluationFunction OnPostEvaluation) = 0;
+		virtual void Cancel() = 0;
+
+		virtual bool IsRunning() const = 0;
+		virtual float GetProgress() const = 0;
+		virtual FText GetProgressMessage() const = 0;
+	};
+
+
+	/**
+	* Evaluator progress tracking
+	* Easy drop-in structure to set and get progress in a thread safe way 
+	*/
+	struct FThreadSafeProgressTracker
+	{
+	public:
+		float GetProgress() const;
+		FText GetProgressMessage() const;
+
+		void SetProgress(float InCurrent, float InTotal, FName InNodeName);
+		void SetProgress(float InProgress, FName InNodeName);
+	
+	private:
+		mutable FCriticalSection ProgressLock;
+		FName NodeName;
+		float Progress = 0.f;
+	};
+
+	/**
+	* This evaluator will evaluate the graph in one go when Start is called
+	* This is useful for running simulation graph for example that needs to execute in exactly one frame
+	* it is required to run from a worker thread and will assert if a node requires gamethread evaluation
+	* Note that it is safe to cancel or query if its running or its progress from another thread 
+	*/
+	struct FDataflowDirectEvaluator : public IDataflowEvaluator
+	{
+	public:
+		DATAFLOWCORE_API FDataflowDirectEvaluator(FContext& InOwningContext);
+		virtual ~FDataflowDirectEvaluator() = default;
+
+		//~ Begin IDataflowEvaluator interface
+		virtual void Start(TSharedRef<const FGraph> InDataflowGraph, TSharedRef<const FCompiledGraph> InCompiledGraph, const FGuid& InNodeId, FOnPostEvaluationFunction OnPostEvaluation) override;
+		virtual void Cancel() override;
+		virtual bool IsRunning() const override;
+		virtual float GetProgress() const override;
+		virtual FText GetProgressMessage() const override;
+		//~ End IDataflowEvaluator interface
+
+	private:
+		FContext& OwningContext;
+
+		std::atomic<bool> bIsCanceled = false;
+		std::atomic<bool> bIsRunning = false;
+
+		FThreadSafeProgressTracker ProgressTracker;
+	};
+
+	/**
+	* This evaluator will try evaluate the graph on the game thread tick 
+	* It will run as many nodes it can based on a time budget 
+	* Note that a node may take more than the budget and create hitches
+	*/
+	struct FDataflowTickEvaluator: public IDataflowEvaluator
+	{
+		DATAFLOWCORE_API FDataflowTickEvaluator(FContext& InOwningContext);
+		DATAFLOWCORE_API virtual ~FDataflowTickEvaluator();
+
+		//~ Begin IDataflowEvaluator interface
+		virtual void Start(TSharedRef<const FGraph> InDataflowGraph, TSharedRef<const FCompiledGraph> InCompiledGraph, const FGuid& InNodeId, FOnPostEvaluationFunction OnPostEvaluation) override;
+		virtual void Cancel() override;
+		virtual bool IsRunning() const override;
+		virtual float GetProgress() const override;
+		virtual FText GetProgressMessage() const override;
+		//~ End IDataflowEvaluator interface
+
+	private:
+		void Reset();
+		bool Tick(float DeltaTime);
+		void OnEvaluationCompleted();
+		void OnEvaluationCanceled();
+
+		FContext& OwningContext;
+		TSharedPtr<const FCompiledGraph> CompiledGraph;
+		TSharedPtr<const FGraph> DataflowGraph;
+		FOnPostEvaluationFunction PostEvaluationFunction;
+		FGuid NodeId;
+
+		FTSTicker::FDelegateHandle TickDelegateHandle;
+
+		int32 LastNodeIndex = INDEX_NONE;
+
+		std::atomic<bool> bIsCanceled = false;
+		std::atomic<bool> bIsRunning = false;
+
+		FThreadSafeProgressTracker ProgressTracker;
+	};
+
+	/**
+	* This evaluator will evaluate the graph in one go when Start is called
+	* This is useful for running simulation graph for example that needs to execute in exactly one frame
+	* it is required to run from a worker thread and will assert if a node requires gamethread evaluation
+	* Note that it is safe to cancel or query if its running or its progress from another thread
+	*/
+	struct FDataflowTaskGraphEvaluator : public IDataflowEvaluator
+	{
+	public:
+		DATAFLOWCORE_API FDataflowTaskGraphEvaluator(FContext& InOwningContext);
+		virtual ~FDataflowTaskGraphEvaluator() = default;
+
+		//~ Begin IDataflowEvaluator interface
+		virtual void Start(TSharedRef<const FGraph> InDataflowGraph, TSharedRef<const FCompiledGraph> InCompiledGraph, const FGuid& InNodeId, FOnPostEvaluationFunction OnPostEvaluation) override;
+		virtual void Cancel() override;
+		virtual bool IsRunning() const override;
+		virtual float GetProgress() const override;
+		virtual FText GetProgressMessage() const override;
+		//~ End IDataflowEvaluator interface
+
+	private:
+		FContext& OwningContext;
+
+		UE::Tasks::FTask CompletionTask;
+		TSharedRef<UE::Tasks::FCancellationToken> CancellationToken;
+
+		FThreadSafeProgressTracker ProgressTracker;
+	};
 }
 
 
diff --git a/Engine/Source/Runtime/Experimental/Dataflow/Engine/Private/Dataflow/DataflowObjectInterface.cpp b/Engine/Source/Runtime/Experimental/Dataflow/Engine/Private/Dataflow/DataflowObjectInterface.cpp
--- a/Engine/Source/Runtime/Experimental/Dataflow/Engine/Private/Dataflow/DataflowObjectInterface.cpp
+++ b/Engine/Source/Runtime/Experimental/Dataflow/Engine/Private/Dataflow/DataflowObjectInterface.cpp
@@ -1,9 +1,74 @@
-﻿// Copyright Epic Games, Inc. All Rights Reserved.
+// Copyright Epic Games, Inc. All Rights Reserved.
 
 #include "Dataflow/DataflowObjectInterface.h"
+#include "Dataflow/DataflowObject.h"
+#include "Dataflow/DataflowCompiler.h"
+#include "Dataflow/DataflowInstance.h"
 
 namespace UE::Dataflow
 {
 	template class TEngineContext<FContextSingle>;
 	template class TEngineContext<FContextThreaded>;
+
+	template <typename T>
+	bool TEngineContext<T>::EvaluateGraphInternal(TSharedRef<IDataflowEvaluator> Evaluator, FOnPostEvaluationFunction OnPostEvaluation)
+	{
+		if (UDataflow* DataflowObject = InstanceUtils::GetDataflowAssetFromObject(Owner))
+		{
+			if (TSharedPtr<FGraph> Graph = DataflowObject->GetDataflow())
+			{
+				UE::Dataflow::FCompilerParameters Params;
+				Params.DataflowObjectToCompile = DataflowObject;
+
+				UE::Dataflow::FCompilerResults Results;
+				if (UE::Dataflow::FCompiler::Compile(Params, Results))
+				{
+					if (TSharedPtr<const UE::Dataflow::FCompiledGraph> CompiledGraph = Results.GetCompiledGraph(DataflowObject))
+					{
+						Evaluator->Start(Graph.ToSharedRef(), CompiledGraph.ToSharedRef(), FGuid(), OnPostEvaluation);
+						return true;
+					}
+				}
+			}
+		}
+		return false;
+	}
+
+	template <typename T>
+	void TEngineContext<T>::UpdateAssetFromTerminalNodes(bool bUseAllTerminals)
+	{
+		if (UDataflow* DataflowObject = InstanceUtils::GetDataflowAssetFromObject(Owner))
+		{
+			if (TSharedPtr<FGraph> Graph = DataflowObject->GetDataflow())
+			{
+				TArray<TSharedPtr<FDataflowNode>> TerminalNodes;
+				if (bUseAllTerminals)
+				{
+					TerminalNodes = Graph->GetFilteredNodes(FDataflowTerminalNode::StaticType());
+				}
+				else
+				{
+					const FName TerminalNodeName = InstanceUtils::GetTerminalNodeNameFromObject(Owner);
+					if (const TSharedPtr<FDataflowNode> TerminalNode = Graph->FindFilteredNode(FDataflowTerminalNode::StaticType(), TerminalNodeName))
+					{
+						TerminalNodes.Add(TerminalNode);
+					}
+				}
+
+				for (TSharedPtr<FDataflowNode> Node : TerminalNodes)
+				{
+					if (const FDataflowTerminalNode* TerminalNode = Node ? Node->AsType<const FDataflowTerminalNode>() : nullptr)
+					{
+						// Note: If the node is deactivated and has any outputs, then these outputs might still need to be forwarded.
+						//       Therefore the Evaluate method has to be called for whichever value of bActive.
+						//       This however isn't the case of SetAssetValue() for which the active state needs to be checked before the call.
+						if (TerminalNode->IsActive() && Owner)
+						{
+							TerminalNode->SetAssetValue(Owner, *this);
+						}
+					}
+				}
+			}
+		}
+	}
 }
diff --git a/Engine/Source/Runtime/Experimental/Dataflow/Engine/Public/Dataflow/DataflowObjectInterface.h b/Engine/Source/Runtime/Experimental/Dataflow/Engine/Public/Dataflow/DataflowObjectInterface.h
--- a/Engine/Source/Runtime/Experimental/Dataflow/Engine/Public/Dataflow/DataflowObjectInterface.h
+++ b/Engine/Source/Runtime/Experimental/Dataflow/Engine/Public/Dataflow/DataflowObjectInterface.h
@@ -31,6 +31,22 @@ namespace UE::Dataflow
 
 		virtual void Serialize(FArchive& Ar) { Base::Serialize(Ar); }
 
+		template <typename T UE_REQUIRES(std::is_base_of_v<IDataflowEvaluator, T>)>
+		TSharedPtr<IDataflowEvaluator> EvaluateGraph(FOnPostEvaluationFunction OnPostEvaluation)
+		{
+			TSharedRef<IDataflowEvaluator> Evaluator = MakeShared<T>(*this);
+			if (EvaluateGraphInternal(Evaluator, OnPostEvaluation))
+			{
+				return Evaluator.ToSharedPtr();
+			}
+			return {};
+			
+		}
+
+		DATAFLOWENGINE_API void UpdateAssetFromTerminalNodes(bool bUseAllTerminals);
+
+	private:
+		DATAFLOWENGINE_API bool EvaluateGraphInternal(TSharedRef<IDataflowEvaluator> Evaluator, FOnPostEvaluationFunction OnPostEvaluation);
 	};
 
 	typedef TEngineContext<FContextSingle> FEngineContext;
