Dataflow : Fix macArm64 build
- Move template method back to the header and keep the implementation details in private functions

#Synced-CL 49661615

diff --git a/Engine/Source/Runtime/Experimental/Dataflow/Engine/Private/Dataflow/DataflowObjectInterface.cpp b/Engine/Source/Runtime/Experimental/Dataflow/Engine/Private/Dataflow/DataflowObjectInterface.cpp
--- a/Engine/Source/Runtime/Experimental/Dataflow/Engine/Private/Dataflow/DataflowObjectInterface.cpp
+++ b/Engine/Source/Runtime/Experimental/Dataflow/Engine/Private/Dataflow/DataflowObjectInterface.cpp
@@ -10,61 +10,62 @@ namespace UE::Dataflow
 	template class TEngineContext<FContextSingle>;
 	template class TEngineContext<FContextThreaded>;
 
-	template <typename T>
-	bool TEngineContext<T>::EvaluateGraphInternal(TSharedRef<IDataflowEvaluator> Evaluator, FOnPostEvaluationFunction OnPostEvaluation)
+	namespace Private
 	{
-		if (UDataflow* DataflowObject = InstanceUtils::GetDataflowAssetFromObject(Owner))
+		bool EvaluateGraphInternal(TSharedRef<IDataflowEvaluator> Evaluator, UObject* Owner, FOnPostEvaluationFunction OnPostEvaluation)
 		{
-			if (TSharedPtr<FGraph> Graph = DataflowObject->GetDataflow())
+			if (UDataflow* DataflowObject = InstanceUtils::GetDataflowAssetFromObject(Owner))
 			{
-				UE::Dataflow::FCompilerParameters Params;
-				Params.DataflowObjectToCompile = DataflowObject;
-
-				UE::Dataflow::FCompilerResults Results;
-				if (UE::Dataflow::FCompiler::Compile(Params, Results))
+				if (TSharedPtr<FGraph> Graph = DataflowObject->GetDataflow())
 				{
-					if (TSharedPtr<const UE::Dataflow::FCompiledGraph> CompiledGraph = Results.GetCompiledGraph(DataflowObject))
+					UE::Dataflow::FCompilerParameters Params;
+					Params.DataflowObjectToCompile = DataflowObject;
+
+					UE::Dataflow::FCompilerResults Results;
+					if (UE::Dataflow::FCompiler::Compile(Params, Results))
 					{
-						Evaluator->Start(Graph.ToSharedRef(), CompiledGraph.ToSharedRef(), FGuid(), OnPostEvaluation);
-						return true;
+						if (TSharedPtr<const UE::Dataflow::FCompiledGraph> CompiledGraph = Results.GetCompiledGraph(DataflowObject))
+						{
+							Evaluator->Start(Graph.ToSharedRef(), CompiledGraph.ToSharedRef(), FGuid(), OnPostEvaluation);
+							return true;
+						}
 					}
 				}
 			}
+			return false;
 		}
-		return false;
-	}
 
-	template <typename T>
-	void TEngineContext<T>::UpdateAssetFromTerminalNodes(bool bUseAllTerminals)
-	{
-		if (UDataflow* DataflowObject = InstanceUtils::GetDataflowAssetFromObject(Owner))
+		void UpdateAssetFromTerminalNodes(FContext& Context, UObject* Asset, bool bUseAllTerminals)
 		{
-			if (TSharedPtr<FGraph> Graph = DataflowObject->GetDataflow())
+			if (UDataflow* DataflowObject = InstanceUtils::GetDataflowAssetFromObject(Asset))
 			{
-				TArray<TSharedPtr<FDataflowNode>> TerminalNodes;
-				if (bUseAllTerminals)
-				{
-					TerminalNodes = Graph->GetFilteredNodes(FDataflowTerminalNode::StaticType());
-				}
-				else
+				if (TSharedPtr<FGraph> Graph = DataflowObject->GetDataflow())
 				{
-					const FName TerminalNodeName = InstanceUtils::GetTerminalNodeNameFromObject(Owner);
-					if (const TSharedPtr<FDataflowNode> TerminalNode = Graph->FindFilteredNode(FDataflowTerminalNode::StaticType(), TerminalNodeName))
+					TArray<TSharedPtr<FDataflowNode>> TerminalNodes;
+					if (bUseAllTerminals)
 					{
-						TerminalNodes.Add(TerminalNode);
+						TerminalNodes = Graph->GetFilteredNodes(FDataflowTerminalNode::StaticType());
+					}
+					else
+					{
+						const FName TerminalNodeName = InstanceUtils::GetTerminalNodeNameFromObject(Asset);
+						if (const TSharedPtr<FDataflowNode> TerminalNode = Graph->FindFilteredNode(FDataflowTerminalNode::StaticType(), TerminalNodeName))
+						{
+							TerminalNodes.Add(TerminalNode);
+						}
 					}
-				}
 
-				for (TSharedPtr<FDataflowNode> Node : TerminalNodes)
-				{
-					if (const FDataflowTerminalNode* TerminalNode = Node ? Node->AsType<const FDataflowTerminalNode>() : nullptr)
+					for (TSharedPtr<FDataflowNode> Node : TerminalNodes)
 					{
-						// Note: If the node is deactivated and has any outputs, then these outputs might still need to be forwarded.
-						//       Therefore the Evaluate method has to be called for whichever value of bActive.
-						//       This however isn't the case of SetAssetValue() for which the active state needs to be checked before the call.
-						if (TerminalNode->IsActive() && Owner)
+						if (const FDataflowTerminalNode* TerminalNode = Node ? Node->AsType<const FDataflowTerminalNode>() : nullptr)
 						{
-							TerminalNode->SetAssetValue(Owner, *this);
+							// Note: If the node is deactivated and has any outputs, then these outputs might still need to be forwarded.
+							//       Therefore the Evaluate method has to be called for whichever value of bActive.
+							//       This however isn't the case of SetAssetValue() for which the active state needs to be checked before the call.
+							if (TerminalNode->IsActive() && Asset)
+							{
+								TerminalNode->SetAssetValue(Asset, Context);
+							}
 						}
 					}
 				}
diff --git a/Engine/Source/Runtime/Experimental/Dataflow/Engine/Public/Dataflow/DataflowObjectInterface.h b/Engine/Source/Runtime/Experimental/Dataflow/Engine/Public/Dataflow/DataflowObjectInterface.h
--- a/Engine/Source/Runtime/Experimental/Dataflow/Engine/Public/Dataflow/DataflowObjectInterface.h
+++ b/Engine/Source/Runtime/Experimental/Dataflow/Engine/Public/Dataflow/DataflowObjectInterface.h
@@ -10,6 +10,12 @@ class FArchive;
 
 namespace UE::Dataflow
 {
+	namespace Private
+	{
+		DATAFLOWENGINE_API void UpdateAssetFromTerminalNodes(FContext& Context, UObject* Asset, bool bUseAllTerminals);
+		DATAFLOWENGINE_API bool EvaluateGraphInternal(TSharedRef<IDataflowEvaluator> Evaluator, UObject* Owner, FOnPostEvaluationFunction OnPostEvaluation);
+	}
+
 	template<class Base = FContextSingle>
 	class TEngineContext : public Base
 	{
@@ -35,7 +41,7 @@ namespace UE::Dataflow
 		TSharedPtr<IDataflowEvaluator> EvaluateGraph(FOnPostEvaluationFunction OnPostEvaluation)
 		{
 			TSharedRef<IDataflowEvaluator> Evaluator = MakeShared<T>(*this);
-			if (EvaluateGraphInternal(Evaluator, OnPostEvaluation))
+			if (Private::EvaluateGraphInternal(Evaluator, Owner, OnPostEvaluation))
 			{
 				return Evaluator.ToSharedPtr();
 			}
@@ -43,10 +49,10 @@ namespace UE::Dataflow
 			
 		}
 
-		DATAFLOWENGINE_API void UpdateAssetFromTerminalNodes(bool bUseAllTerminals);
-
-	private:
-		DATAFLOWENGINE_API bool EvaluateGraphInternal(TSharedRef<IDataflowEvaluator> Evaluator, FOnPostEvaluationFunction OnPostEvaluation);
+		void UpdateAssetFromTerminalNodes(bool bUseAllTerminals)
+		{
+			Private::UpdateAssetFromTerminalNodes(*this, Owner, bUseAllTerminals);
+		}
 	};
 
 	typedef TEngineContext<FContextSingle> FEngineContext;
