Update comments

#Synced-CL 49554261

diff --git a/Engine/Source/Runtime/Experimental/IoStore/OnDemandCore/Public/IO/IoStoreOnDemand.h b/Engine/Source/Runtime/Experimental/IoStore/OnDemandCore/Public/IO/IoStoreOnDemand.h
--- a/Engine/Source/Runtime/Experimental/IoStore/OnDemandCore/Public/IO/IoStoreOnDemand.h
+++ b/Engine/Source/Runtime/Experimental/IoStore/OnDemandCore/Public/IO/IoStoreOnDemand.h
@@ -696,12 +696,12 @@ public:
 	virtual void GetHttpStats(FOnDemandHttpStats& Out) const = 0;
 	/**
 	 * Reports the statistics for the current OnDemandBackend. It will record any events that have occurred from the last time that this was called.
-	 * This is a legacy method, prefer using StartAnalyticsRecording instead.
+	 * This is a legacy method, prefer using StartAnalyticsRecording if possible.
 	 */
 	virtual void ReportAnalytics(TArray<FAnalyticsEventAttribute>& OutAnalyticsArray) const = 0;
 	/** 
 	 * Create a new analytics interface for the current OnDemandBackend, can be used to record statistics for a set period of time.
-	 * This can return a nullptr if the analytics system is not enabled.
+	 * This can return a nullptr if the analytics system is not enabled. Currently only supports the streaming backend.
 	 */
 	virtual TUniquePtr<IAnalyticsRecording> StartAnalyticsRecording() const = 0;
 	/**
