Move midphases stats after Narrow phase
Fix race condition on the stats count that has to be computed just after the Narrow Phase

#Synced-CL 49725521

diff --git a/Engine/Source/Runtime/Experimental/Chaos/Public/Chaos/Collision/SpatialAccelerationBroadPhase.h b/Engine/Source/Runtime/Experimental/Chaos/Public/Chaos/Collision/SpatialAccelerationBroadPhase.h
--- a/Engine/Source/Runtime/Experimental/Chaos/Public/Chaos/Collision/SpatialAccelerationBroadPhase.h
+++ b/Engine/Source/Runtime/Experimental/Chaos/Public/Chaos/Collision/SpatialAccelerationBroadPhase.h
@@ -417,13 +417,6 @@ namespace Chaos
 			// Run some error checks in non-shipping builds
 			// NOTE: This must come after MidPhase assignment for now because that's where the filter is applied
 			CheckOverlapResults();
-
-			// Update stats
-			for (int32 ContextIndex = 0; ContextIndex < NumActiveBroadphaseContexts; ++ContextIndex)
-			{
-				NumBroadPhasePairs += BroadphaseContexts[ContextIndex].Overlaps.Num();
-				NumMidPhases += BroadphaseContexts[ContextIndex].MidPhases.Num();
-			}
 		}
 
 		/**
@@ -624,6 +617,13 @@ namespace Chaos
 				BroadphaseContexts[0].CollisionContext.Allocator->NewActiveConstraints = 
 					MoveTemp(BroadphaseContexts[ContextsWithConstraints[0]].CollisionContext.Allocator->NewActiveConstraints);
 			}
+
+			// Update stats at the end in single thread when all mid phases have been computed.
+			for (int32 ContextIndex = 0; ContextIndex < NumActiveBroadphaseContexts; ++ContextIndex)
+			{
+				NumBroadPhasePairs += BroadphaseContexts[ContextIndex].Overlaps.Num();
+				NumMidPhases += BroadphaseContexts[ContextIndex].MidPhases.Num();
+			}
 		}
 
 		/** @brief This function is the outer loop of collision detection. It loops over the
