IAD - Make cache block size configurable.

#Synced-CL 49554261

diff --git a/Engine/Source/Runtime/Experimental/IoStore/OnDemand/Private/OnDemandConfig.cpp b/Engine/Source/Runtime/Experimental/IoStore/OnDemand/Private/OnDemandConfig.cpp
--- a/Engine/Source/Runtime/Experimental/IoStore/OnDemand/Private/OnDemandConfig.cpp
+++ b/Engine/Source/Runtime/Experimental/IoStore/OnDemand/Private/OnDemandConfig.cpp
@@ -463,6 +463,42 @@ TIoStatusOr<FOnDemandInstallCacheConfig> TryParseInstallCacheConfig(const TCHAR*
 		OutConfig.JournalMaxSize = JournalMaxSize;
 	}
 
+	const int64 MinBlockSize = [CommandLine]
+	{
+		if (FString ParamValue; FParse::Value(CommandLine, TEXT("-Iad.FileCache.MinBlockSize="), ParamValue))
+		{
+			return Config::ParseSizeParam(ParamValue);
+		}
+		else if (FString ValueStr; GConfig->GetString(TEXT("OnDemandInstall"), TEXT("FileCache.MinBlockSize"), ValueStr, GEngineIni))
+		{
+			return Config::ParseSizeParam(ValueStr);
+		}
+		return int64(-1);
+	}();
+
+	if (MinBlockSize > 0)
+	{
+		OutConfig.MinBlockSize = MinBlockSize;
+	}
+
+	const int64 MaxBlockSize = [CommandLine]
+	{
+		if (FString ParamValue; FParse::Value(CommandLine, TEXT("-Iad.FileCache.MaxBlockSize="), ParamValue))
+		{
+			return Config::ParseSizeParam(ParamValue);
+		}
+		else if (FString ValueStr; GConfig->GetString(TEXT("OnDemandInstall"), TEXT("FileCache.MaxBlockSize"), ValueStr, GEngineIni))
+		{
+			return Config::ParseSizeParam(ValueStr);
+		}
+		return int64(-1);
+	}();
+
+	if (MaxBlockSize > 0)
+	{
+		OutConfig.MaxBlockSize = MaxBlockSize;
+	}
+
 	const double LastAccessGranularitySeconds = [CommandLine]
 	{
 		if (double ParamValue; FParse::Value(CommandLine, TEXT("-Iad.FileCache.LastAccessGranularitySeconds="), ParamValue))
diff --git a/Engine/Source/Runtime/Experimental/IoStore/OnDemand/Private/OnDemandInstallCache.cpp b/Engine/Source/Runtime/Experimental/IoStore/OnDemand/Private/OnDemandInstallCache.cpp
--- a/Engine/Source/Runtime/Experimental/IoStore/OnDemand/Private/OnDemandInstallCache.cpp
+++ b/Engine/Source/Runtime/Experimental/IoStore/OnDemand/Private/OnDemandInstallCache.cpp
@@ -702,11 +702,12 @@ struct FCas
 	using FBlockStates			= TMap<FCasBlockId, FCasBlockState>;
 	using FBlockFileSizes		= TMap<FCasBlockId, uint64>;
 
-	FCas(const FOnDemandInstallCacheConfig& Config);
+	FCas();
 
 	void						Lock()		{ Mutex.Lock(); }
 	void						Unlock()	{ Mutex.Unlock(); }
 
+	FResult						Configure(const FOnDemandInstallCacheConfig& Config);
 	FResult						Initialize(FStringView Directory, bool bDeleteExisting = false);
 	FCasLocation				FindChunk(const FIoHash& Hash) const;
 	FCasBlockId					CreateBlock();
@@ -747,17 +748,15 @@ struct FCas
 
 	// Returns the place holder timestamp to use if a timestamp is not found in the LastAccess table
 	static int64				GetTimestampForMissingLastAccess();
-	static uint32				GetMaxBlockSize()
-	{
-		const uint32 MaxBlockSize = 32 << 20; //TODO: Make configurable	
-		return MaxBlockSize;
-	}
 
-	static uint32				GetMinBlockSize()
+	uint32						GetMinBlockSize()
 	{
-		const uint32 MinBlockSize = 32 << 19; //TODO: Make configurable	
 		return MinBlockSize;
 	}
+	uint32						GetMaxBlockSize()
+	{
+		return MaxBlockSize;
+	}
 
 	FString						GetRootDirectory() const { return FString(RootDirectory); }
 
@@ -769,16 +768,48 @@ private:
 	FBlockStates			BlockIds;
 	FReadHandles			ReadHandles;
 	FEventRef				BlockReadsDoneEvent;
-	const int64				LastAccessGranularityTicks;
+	int64					LastAccessGranularityTicks = 0;
+	uint32					MinBlockSize = 32 << 19;
+	uint32					MaxBlockSize = 32 << 20;
 	mutable UE::FMutex		Mutex;
 };
 
 ///////////////////////////////////////////////////////////////////////////////
-FCas::FCas(const FOnDemandInstallCacheConfig& Config)
-	: LastAccessGranularityTicks(FTimespan::FromSeconds(Config.LastAccessGranularitySeconds).GetTicks())
+FCas::FCas()
 {
 }
 
+FResult FCas::Configure(const FOnDemandInstallCacheConfig& Config)
+{
+	using namespace UE::UnifiedError::IoStoreOnDemand;
+
+	const uint32 MaxSupportedBlockSize = (128 << 20);
+
+	if (Config.MinBlockSize > MaxSupportedBlockSize)
+	{
+		return MakeCasError<void>(
+			ECasErrorCode::InitializeFailed, EIoErrorCode::InvalidParameter, FString::Printf(TEXT("MinBlockSize must be less than %i"), MaxSupportedBlockSize));
+	}
+
+	if (Config.MaxBlockSize > MaxSupportedBlockSize)
+	{
+		return MakeCasError<void>(
+			ECasErrorCode::InitializeFailed, EIoErrorCode::InvalidParameter, FString::Printf(TEXT("MaxBlockSize must be less than %i"), MaxSupportedBlockSize));
+	}
+
+	if (Config.MinBlockSize > Config.MaxBlockSize)
+	{
+		return MakeCasError<void>(
+			ECasErrorCode::InitializeFailed, EIoErrorCode::InvalidParameter, FString::Printf(TEXT("MinBlockSize cannot be greater than MaxBlockSize")));
+	}
+
+	LastAccessGranularityTicks = FTimespan::FromSeconds(Config.LastAccessGranularitySeconds).GetTicks();
+	MinBlockSize = static_cast<uint32>(Config.MinBlockSize);
+	MaxBlockSize = static_cast<uint32>(Config.MaxBlockSize);
+
+	return MakeValue();
+}
+
 ///////////////////////////////////////////////////////////////////////////////
 FResult FCas::Initialize(FStringView Directory, bool bDeleteExisting)
 {
@@ -1338,7 +1369,7 @@ FResult FCas::Verify(TArray<FCasAddr>& OutRemovedChunksAddrs, TMap<FCasBlockId,
 			VerifyTimestamp(BlockState.LastAccessUtcTicks,			BlockId, TEXT("last access"));
 			VerifyTimestamp(BlockState.LastModificationUtcTicks,	BlockId, TEXT("modification"));
 
-			if (*FileSize > FCas::GetMaxBlockSize())
+			if (*FileSize > GetMaxBlockSize())
 			{
 				OutOverBudgetBlocks.Add(BlockId, *FileSize);
 			}
@@ -2076,14 +2107,20 @@ FOnDemandInstallCache::FOnDemandInstallCache(const FOnDemandInstallCacheConfig&
 	: IoStore(InIoStore)
 	, Governor(InGovernor)
 	, CacheDirectory(Config.RootDirectory)
-	, Cas(Config)
 	, MaxCacheSize(Config.DiskQuota)
 	, MaxJournalSize(Config.JournalMaxSize)
 {
 	using namespace UE::UnifiedError::IoStoreOnDemand;
 
-	UE_LOG(LogIoStoreOnDemand, Log, TEXT("Initializing install cache, MaxCacheSize=%.2lf MiB, MaxJournalSize=%.2lf KiB"),
-		ToMiB(MaxCacheSize), ToKiB(MaxJournalSize));
+	FResult InitResult = Cas.Configure(Config);
+	if (InitResult.HasError())
+	{
+		UE_LOGFMT(LogIoStoreOnDemand, Error, "Failed to configure install cache, error: {Error}", InitResult.GetError());
+		return;
+	}
+
+	UE_LOG(LogIoStoreOnDemand, Log, TEXT("Initializing install cache, MaxCacheSize=%.2lf MiB, MaxJournalSize=%.2lf KiB, MaxBlockSize=%.2lf MiB, MinBlockSize=%.2lf MiB"),
+		ToMiB(MaxCacheSize), ToKiB(MaxJournalSize), ToMiB(Cas.GetMaxBlockSize()), ToMiB(Cas.GetMinBlockSize()));
 
 	const uint64 MinDiskQuota = 2 * Cas.GetMaxBlockSize();
 	if (MaxCacheSize < MinDiskQuota)
@@ -2096,7 +2133,7 @@ FOnDemandInstallCache::FOnDemandInstallCache(const FOnDemandInstallCacheConfig&
 	MaxCacheSize -= Cas.GetMaxBlockSize();
 	UE_LOG(LogIoStoreOnDemand, Log, TEXT("Effective MaxCacheSize without defragmentation space is MaxCacheSize=%.2lf MiB"), ToMiB(MaxCacheSize));
 
-	FResult InitResult = Cas.Initialize(CacheDirectory);
+	InitResult = Cas.Initialize(CacheDirectory);
 	if (InitResult.HasError())
 	{
 		UE_LOGFMT(LogIoStoreOnDemand, Error, "Failed to initialize install cache, error: {Error}", InitResult.GetError());
@@ -3077,9 +3114,9 @@ FResult FOnDemandInstallCache::InitialVerify()
 		uint64 TotalBlockBytesOverBudget = 0;
 		for (const TPair<FCasBlockId, uint64>& Kv : OverBudgetBlocks)
 		{
-			if (ensure(Kv.Value > FCas::GetMaxBlockSize()))
+			if (ensure(Kv.Value > Cas.GetMaxBlockSize()))
 			{
-				TotalBlockBytesOverBudget += (Kv.Value - FCas::GetMaxBlockSize());
+				TotalBlockBytesOverBudget += (Kv.Value - Cas.GetMaxBlockSize());
 			}
 		}
 
diff --git a/Engine/Source/Runtime/Experimental/IoStore/OnDemand/Private/OnDemandInstallCache.h b/Engine/Source/Runtime/Experimental/IoStore/OnDemand/Private/OnDemandInstallCache.h
--- a/Engine/Source/Runtime/Experimental/IoStore/OnDemand/Private/OnDemandInstallCache.h
+++ b/Engine/Source/Runtime/Experimental/IoStore/OnDemand/Private/OnDemandInstallCache.h
@@ -51,6 +51,8 @@ struct FOnDemandInstallCacheConfig
 	FString RootDirectory;
 	uint64	DiskQuota		= 1ull << 30;
 	uint64	JournalMaxSize	= 2ull << 20;
+	uint64	MinBlockSize	= 32 << 19;
+	uint64	MaxBlockSize	= 32 << 20;
 	double	LastAccessGranularitySeconds = 60 * 60;
 	bool	bDropCache		= false;
 };
