Dataflow : fix debug draw component bounds computation to make sure the transform of boxes are accounted for

#Synced-CL 49754362

diff --git a/Engine/Source/Runtime/Experimental/Dataflow/Engine/Private/Dataflow/DataflowDebugDrawComponent.cpp b/Engine/Source/Runtime/Experimental/Dataflow/Engine/Private/Dataflow/DataflowDebugDrawComponent.cpp
--- a/Engine/Source/Runtime/Experimental/Dataflow/Engine/Private/Dataflow/DataflowDebugDrawComponent.cpp
+++ b/Engine/Source/Runtime/Experimental/Dataflow/Engine/Private/Dataflow/DataflowDebugDrawComponent.cpp
@@ -53,7 +53,7 @@ FBoxSphereBounds UDataflowDebugDrawComponent::CalcBounds(const FTransform& Local
 		}
 		for (const FDebugRenderSceneProxy::FDebugBox& Box : DebugSceneProxy->Boxes)
 		{
-			BoundsBuilder += Box.Box.ExpandBy(Box.Thickness);
+			BoundsBuilder += Box.Box.TransformBy(Box.Transform).ExpandBy(Box.Thickness);
 		}
 		for (const FDebugRenderSceneProxy::FSphere& Sphere : DebugSceneProxy->Spheres)
 		{
