Filter out sleeping constraints when adding them to the solver containers.

This was previously implemented but not working properly because FConstraintHandle::IsSleeping() is only implemented for collision constraints.

#Synced-CL 49649653

diff --git a/Engine/Source/Runtime/Experimental/Chaos/Private/Chaos/Character/CharacterGroundConstraintContainerSolver.cpp b/Engine/Source/Runtime/Experimental/Chaos/Private/Chaos/Character/CharacterGroundConstraintContainerSolver.cpp
--- a/Engine/Source/Runtime/Experimental/Chaos/Private/Chaos/Character/CharacterGroundConstraintContainerSolver.cpp
+++ b/Engine/Source/Runtime/Experimental/Chaos/Private/Chaos/Character/CharacterGroundConstraintContainerSolver.cpp
@@ -39,16 +39,11 @@ namespace Chaos
 		{
 			for (Private::FPBDIslandConstraint* IslandConstraint : IslandConstraints)
 			{
-				FConstraintHandle* Constraint = IslandConstraint->GetConstraint();
-				
 				// Filter out sleeping constraints in any partially sleeping island
-				// @todo(chaos): This is not working correctly since IsSleeping() always returns false for this constraint type.
-				// As a result, the constraint is added to the solver, data is gathered and scattered, and the constraint correction is computed.
-				// Nevertheless, the two connected particles are considered kinematic and, thus, do not "feel" the effect of the constraint.
-				// Hence, the solve remains correct, but we do more work than needed.
-				// In the future, we should store sleep state of for this constraint type as well.
-				if (!Constraint->IsSleeping())
+				if (!IslandConstraint->IsSleeping())
 				{
+					FConstraintHandle* Constraint = IslandConstraint->GetConstraint();
+
 					// We will only ever be given constraints from our container (asserts in non-shipping)
 					FCharacterGroundConstraintHandle* CharacterGroundConstraint = Constraint->AsUnsafe<FCharacterGroundConstraintHandle>();
 
diff --git a/Engine/Source/Runtime/Experimental/Chaos/Private/Chaos/Collision/PBDCollisionContainerSolver.cpp b/Engine/Source/Runtime/Experimental/Chaos/Private/Chaos/Collision/PBDCollisionContainerSolver.cpp
--- a/Engine/Source/Runtime/Experimental/Chaos/Private/Chaos/Collision/PBDCollisionContainerSolver.cpp
+++ b/Engine/Source/Runtime/Experimental/Chaos/Private/Chaos/Collision/PBDCollisionContainerSolver.cpp
@@ -558,11 +558,17 @@ namespace Chaos
 	{
 		for (Private::FPBDIslandConstraint* IslandConstraint : IslandConstraints)
 		{
-			FConstraintHandle* Constraint = IslandConstraint->GetConstraint();
+			// Validate the constraint sleep state matches between the FPBDIslandConstraint and FConstraintHandle
+			ensureMsgf(IslandConstraint->IsSleeping() == IslandConstraint->GetConstraint()->IsSleeping(),
+				TEXT("Mismatch in constraint sleep state. Particle 0: %s, Particle 1: %s."),
+				*(IslandConstraint->GetConstraint()->GetConstrainedParticles()[0]->GetDebugName()),
+				*(IslandConstraint->GetConstraint()->GetConstrainedParticles()[1]->GetDebugName()));
 
 			// Filter out sleeping constraints in any partially sleeping island
-			if (!Constraint->IsSleeping())
+			if (!IslandConstraint->IsSleeping())
 			{
+				FConstraintHandle* Constraint = IslandConstraint->GetConstraint();
+
 				// We will only ever be given constraints from our container (NOTE: AsUnsafe asserts in dev)
 				FPBDCollisionConstraint& CollisionConstraint = Constraint->AsUnsafe<FPBDCollisionConstraintHandle>()->GetContact();
 
diff --git a/Engine/Source/Runtime/Experimental/Chaos/Private/Chaos/Collision/PBDCollisionContainerSolverJacobi.cpp b/Engine/Source/Runtime/Experimental/Chaos/Private/Chaos/Collision/PBDCollisionContainerSolverJacobi.cpp
--- a/Engine/Source/Runtime/Experimental/Chaos/Private/Chaos/Collision/PBDCollisionContainerSolverJacobi.cpp
+++ b/Engine/Source/Runtime/Experimental/Chaos/Private/Chaos/Collision/PBDCollisionContainerSolverJacobi.cpp
@@ -306,11 +306,17 @@ namespace Chaos
 		{
 			for (Private::FPBDIslandConstraint* IslandConstraint : IslandConstraints)
 			{
-				FConstraintHandle* Constraint = IslandConstraint->GetConstraint();
+				// Validate the constraint sleep state matches between the FPBDIslandConstraint and FConstraintHandle
+				ensureMsgf(IslandConstraint->IsSleeping() == IslandConstraint->GetConstraint()->IsSleeping(),
+					TEXT("Mismatch in constraint sleep state. Particle 0: %s, Particle 1: %s."), 
+					*(IslandConstraint->GetConstraint()->GetConstrainedParticles()[0]->GetDebugName()),
+					*(IslandConstraint->GetConstraint()->GetConstrainedParticles()[1]->GetDebugName()));
 
 				// Filter out sleeping constraints in any partially sleeping island
-				if (!Constraint->IsSleeping())
+				if (!IslandConstraint->IsSleeping())
 				{
+					FConstraintHandle* Constraint = IslandConstraint->GetConstraint();
+
 					// We will only ever be given constraints from our container (asserts in non-shipping)
 					FPBDCollisionConstraint& CollisionConstraint = Constraint->AsUnsafe<FPBDCollisionConstraintHandle>()->GetContact();
 
diff --git a/Engine/Source/Runtime/Experimental/Chaos/Private/Chaos/Collision/PBDCollisionContainerSolverSimd.cpp b/Engine/Source/Runtime/Experimental/Chaos/Private/Chaos/Collision/PBDCollisionContainerSolverSimd.cpp
--- a/Engine/Source/Runtime/Experimental/Chaos/Private/Chaos/Collision/PBDCollisionContainerSolverSimd.cpp
+++ b/Engine/Source/Runtime/Experimental/Chaos/Private/Chaos/Collision/PBDCollisionContainerSolverSimd.cpp
@@ -361,6 +361,7 @@ namespace Chaos
 			}
 
 			// Make sure we have enough constraint rows for these constraints (space is pre-allocated)
+			// @todo(chaos): If there are many sleeping constraints, we will be overallocating here.
 			if (SimdData.SimdConstraints.Num() < IslandLaneNumConstraints + IslandConstraints.Num())
 			{
 				SimdData.SimdConstraints.SetNumZeroed(IslandLaneNumConstraints + IslandConstraints.Num(), EAllowShrinking::No);
@@ -369,12 +370,18 @@ namespace Chaos
 			// Add all the constraints in the island to the selected lane
 			int32 NumActiveConstraints = 0;
 			for (Private::FPBDIslandConstraint* IslandConstraint : IslandConstraints)
-			{				
-				FConstraintHandle* Constraint = IslandConstraint->GetConstraint();
+			{
+				// Validate the constraint sleep state matches between the FPBDIslandConstraint and FConstraintHandle
+				ensureMsgf(IslandConstraint->IsSleeping() == IslandConstraint->GetConstraint()->IsSleeping(),
+					TEXT("Mismatch in constraint sleep state. Particle 0: %s, Particle 1: %s."),
+					*(IslandConstraint->GetConstraint()->GetConstrainedParticles()[0]->GetDebugName()),
+					*(IslandConstraint->GetConstraint()->GetConstrainedParticles()[1]->GetDebugName()));
 
 				// Filter out sleeping constraints in any partially sleeping island
-				if (!Constraint->IsSleeping())
+				if (!IslandConstraint->IsSleeping())
 				{
+					FConstraintHandle* Constraint = IslandConstraint->GetConstraint();
+
 					// NOTE: We will only ever be given constraints from our container (asserts in non-shipping)
 					FPBDCollisionConstraint& CollisionConstraint = Constraint->AsUnsafe<FPBDCollisionConstraintHandle>()->GetContact();
 
diff --git a/Engine/Source/Runtime/Experimental/Chaos/Private/Chaos/Joint/PBDJointContainerSolver.cpp b/Engine/Source/Runtime/Experimental/Chaos/Private/Chaos/Joint/PBDJointContainerSolver.cpp
--- a/Engine/Source/Runtime/Experimental/Chaos/Private/Chaos/Joint/PBDJointContainerSolver.cpp
+++ b/Engine/Source/Runtime/Experimental/Chaos/Private/Chaos/Joint/PBDJointContainerSolver.cpp
@@ -149,16 +149,11 @@ namespace Chaos
 			int32 CurrentIndex = 0;
 			for (Private::FPBDIslandConstraint* IslandConstraint : IslandConstraints)
 			{
-				FConstraintHandle* Constraint = IslandConstraint->GetConstraint();
-
 				// Filter out sleeping constraints in any partially sleeping island
-				// @todo(chaos): This is not working correctly since IsSleeping() always returns false for this constraint type.
-				// As a result, the constraint is added to the solver, data is gathered and scattered, and the constraint correction is computed.
-				// Nevertheless, the two connected particles are considered kinematic and, thus, do not "feel" the effect of the constraint.
-				// Hence, the solve remains correct, but we do more work than needed.
-				// In the future, we should store sleep state of for this constraint type as well.
-				if (!Constraint->IsSleeping())
+				if (!IslandConstraint->IsSleeping())
 				{
+					FConstraintHandle* Constraint = IslandConstraint->GetConstraint();
+
 					// We will only ever be given constraints from our container (asserts in non-shipping)
 					const int32 ContainerConstraintIndex = Constraint->AsUnsafe<FPBDJointConstraintHandle>()->GetConstraintIndex();
 
diff --git a/Engine/Source/Runtime/Experimental/Chaos/Public/Chaos/Evolution/IndexedConstraintContainer.h b/Engine/Source/Runtime/Experimental/Chaos/Public/Chaos/Evolution/IndexedConstraintContainer.h
--- a/Engine/Source/Runtime/Experimental/Chaos/Public/Chaos/Evolution/IndexedConstraintContainer.h
+++ b/Engine/Source/Runtime/Experimental/Chaos/Public/Chaos/Evolution/IndexedConstraintContainer.h
@@ -155,16 +155,11 @@ namespace Chaos
 		{
 			for (const Private::FPBDIslandConstraint* IslandConstraint : IslandConstraints)
 			{
-				FConstraintHandle* Constraint = IslandConstraint->GetConstraint();
-
 				// Filter out sleeping constraints in any partially sleeping island
-				// @todo(chaos): This is not working correctly since IsSleeping() always returns false for this constraint type.
-				// As a result, the constraint is added to the solver, data is gathered and scattered, and the constraint correction is computed.
-				// Nevertheless, the two connected particles are considered kinematic and, thus, do not "feel" the effect of the constraint.
-				// Hence, the solve remains correct, but we do more work than needed.
-				// In the future, we should store sleep state of for this constraint type as well.
-				if (!Constraint->IsSleeping())
+				if (!IslandConstraint->IsSleeping())
 				{
+					FConstraintHandle* Constraint = IslandConstraint->GetConstraint();
+
 					// We will only ever be given constraints from our container (asserts in non-shipping)
 					const FIndexedConstraintHandle* IndexedConstraintHandle = Constraint->AsUnsafe<FIndexedConstraintHandle>();
 
diff --git a/Engine/Source/Runtime/Experimental/Chaos/Public/Chaos/Island/IslandManager.h b/Engine/Source/Runtime/Experimental/Chaos/Public/Chaos/Island/IslandManager.h
--- a/Engine/Source/Runtime/Experimental/Chaos/Public/Chaos/Island/IslandManager.h
+++ b/Engine/Source/Runtime/Experimental/Chaos/Public/Chaos/Island/IslandManager.h
@@ -174,6 +174,8 @@ namespace Chaos::Private
 		// An integer ID for an island (for debug only). @see FPBDIsland::GetIslandId()
 		CHAOS_API int32 GetIslandId() const;
 
+		bool IsSleeping() const { return Flags.bIsSleeping; }
+
 	private:
 		friend class FPBDIslandManager;
 		friend class FPBDIsland;
