Experiment to determine if a slow down is due to the number of times UE_LOG is called, or just the amount of data that we are logging.

#rnx

#Synced-CL 49717872

diff --git a/Engine/Source/Runtime/Experimental/IoStore/OnDemand/Private/Tool/ChunkPlugin.cpp b/Engine/Source/Runtime/Experimental/IoStore/OnDemand/Private/Tool/ChunkPlugin.cpp
--- a/Engine/Source/Runtime/Experimental/IoStore/OnDemand/Private/Tool/ChunkPlugin.cpp
+++ b/Engine/Source/Runtime/Experimental/IoStore/OnDemand/Private/Tool/ChunkPlugin.cpp
@@ -681,12 +681,25 @@ static int32 ChunkPluginCommandEntry(const FContext& Context)
 					UE_LOG(LogIoStoreOnDemand, Error, TEXT("Bad settings file '%s'"), *SettingsFile);
 					return -1;
 				}
-
+				
+#if 0
 				UE_LOG(LogIoStoreOnDemand, Display, TEXT("Found Package Set '%s'"), *Pair.Key);
 				for (const FString& Package : Packages)
 				{
-					UE_LOG(LogIoStoreOnDemand, Display, TEXT("-> \t'%s'"), *Package);
+					UE_LOG(LogIoStoreOnDemand, Display, TEXT("->\t'%s'"), *Package);
 				}
+#else
+				TUtf8StringBuilder<0> Output;
+				Output << "Found Package Set '" << Pair.Key << "'";
+				for (const FString& Package : Packages)
+				{
+					Output << LINE_TERMINATOR << "->\t'" << Package << "'";
+				}
+
+				UE_LOG(LogIoStoreOnDemand, Log, TEXT("Package Set contains %d entries"), Packages.Num());
+				UE_LOG(LogIoStoreOnDemand, Display, TEXT("%hs"), *Output);
+#endif
+				
 
 				Settings.PackageSets.Add(FString(Pair.Key), MoveTemp(Packages));
 			}
