IAD - split out flushes from defrags into seperate stats.

#Synced-CL 49554261

diff --git a/Engine/Source/Runtime/Experimental/IoStore/OnDemand/Private/OnDemandInstallCache.cpp b/Engine/Source/Runtime/Experimental/IoStore/OnDemand/Private/OnDemandInstallCache.cpp
--- a/Engine/Source/Runtime/Experimental/IoStore/OnDemand/Private/OnDemandInstallCache.cpp
+++ b/Engine/Source/Runtime/Experimental/IoStore/OnDemand/Private/OnDemandInstallCache.cpp
@@ -3842,6 +3842,8 @@ FResult FOnDemandInstallCache::FlushPendingChunksImpl(const FPendingChunks& Chun
 {
 	using namespace UE::UnifiedError::IoStoreOnDemand;
 
+	const bool bFromDefrag = (&PendingChunks != &Chunks);
+
 	uint64 TotalCasBytes		= 0;
 	uint64 TotalJournalBytes	= 0;
 	uint32 TotalOpCount			= 0;
@@ -3879,7 +3881,7 @@ FResult FOnDemandInstallCache::FlushPendingChunksImpl(const FPendingChunks& Chun
 			FInstallCacheErrorContext Ctx = MakeInstallCacheErrorContext();
 			OpenWriteResult.GetError().PushErrorContext(MoveTemp(Ctx));
 			FResult Result = MakeError(OpenWriteResult.StealError());
-			FOnDemandInstallCacheStats::OnFlush(Result, TotalCasBytes);
+			FOnDemandInstallCacheStats::OnFlush(Result, TotalCasBytes, bFromDefrag);
 			return Result;
 		}
 		
@@ -3919,7 +3921,7 @@ FResult FOnDemandInstallCache::FlushPendingChunksImpl(const FPendingChunks& Chun
 					ECasErrorCode::WriteBlockFailed, EIoErrorCode::WriteError, FString::Printf(TEXT("Failed to write %" INT64_FMT " bytes to CAS block %u"), BytesToWrite, CurrentBlockId.Id));
 				FInstallCacheErrorContext Ctx = MakeInstallCacheErrorContext();
 				Result.GetError().PushErrorContext(MoveTemp(Ctx));
-				FOnDemandInstallCacheStats::OnFlush(Result, TotalCasBytes);
+				FOnDemandInstallCacheStats::OnFlush(Result, TotalCasBytes, bFromDefrag);
 				return Result;
 			}
 			TotalOpCount++;
@@ -3931,7 +3933,7 @@ FResult FOnDemandInstallCache::FlushPendingChunksImpl(const FPendingChunks& Chun
 					ECasErrorCode::WriteBlockFailed, EIoErrorCode::FileFlushFailed, FString::Printf(TEXT("Failed to flush %" INT64_FMT " bytes to CAS block %u"), BytesToWrite, CurrentBlockId.Id));
 				FInstallCacheErrorContext Ctx = MakeInstallCacheErrorContext();
 				Result.GetError().PushErrorContext(MoveTemp(Ctx));
-				FOnDemandInstallCacheStats::OnFlush(Result, TotalCasBytes);
+				FOnDemandInstallCacheStats::OnFlush(Result, TotalCasBytes, bFromDefrag);
 				return Result;
 			}
 			TotalOpCount++;
@@ -3981,7 +3983,7 @@ FResult FOnDemandInstallCache::FlushPendingChunksImpl(const FPendingChunks& Chun
 		{
 			FInstallCacheErrorContext Ctx = MakeInstallCacheErrorContext();
 			Result.GetError().PushErrorContext(MoveTemp(Ctx));
-			FOnDemandInstallCacheStats::OnFlush(Result, TotalCasBytes);
+			FOnDemandInstallCacheStats::OnFlush(Result, TotalCasBytes, bFromDefrag);
 			return Result;
 		}
 
@@ -3995,7 +3997,7 @@ FResult FOnDemandInstallCache::FlushPendingChunksImpl(const FPendingChunks& Chun
 	}
 
 	FResult Result = MakeValue();
-	FOnDemandInstallCacheStats::OnFlush(Result, TotalCasBytes);
+	FOnDemandInstallCacheStats::OnFlush(Result, TotalCasBytes, bFromDefrag);
 	return Result;
 }
 
diff --git a/Engine/Source/Runtime/Experimental/IoStore/OnDemand/Private/Statistics.cpp b/Engine/Source/Runtime/Experimental/IoStore/OnDemand/Private/Statistics.cpp
--- a/Engine/Source/Runtime/Experimental/IoStore/OnDemand/Private/Statistics.cpp
+++ b/Engine/Source/Runtime/Experimental/IoStore/OnDemand/Private/Statistics.cpp
@@ -578,6 +578,9 @@ struct FInstallCacheAnalytics
 	uint64 FlushCount = 0;
 	uint64 FlushErrorCount = 0;
 	uint64 FlushedBytes = 0;
+	uint64 DefragFlushCount = 0;
+	uint64 DefragFlushErrorCount = 0;
+	uint64 DefragFlushedBytes = 0;
 	uint64 PurgedBytes = 0;
 	uint64 PurgeCount = 0;
 	uint64 PurgeErrorCount = 0;
@@ -1614,6 +1617,9 @@ static void LogOrReportContentInstallerAnalytics(bool bReport, bool bLog, TArray
 		XANALYTIC(IadInstallCacheFlushCount, CurrentInstallCacheAnalytics.FlushCount) \
 		XANALYTIC(IadInstallCacheFlushErrorCount, CurrentInstallCacheAnalytics.FlushErrorCount) \
 		XANALYTIC(IadInstallCacheFlushedBytes, CurrentInstallCacheAnalytics.FlushedBytes) \
+		XANALYTIC(IadInstallCacheDefragFlushCount, CurrentInstallCacheAnalytics.DefragFlushCount) \
+		XANALYTIC(IadInstallCacheDefragFlushErrorCount, CurrentInstallCacheAnalytics.DefragFlushErrorCount) \
+		XANALYTIC(IadInstallCacheDefragFlushedBytes, CurrentInstallCacheAnalytics.DefragFlushedBytes) \
 		XANALYTIC(IadInstallCachePurgeCount, CurrentInstallCacheAnalytics.PurgeCount) \
 		XANALYTIC(IadInstallCachePurgeErrorCount, CurrentInstallCacheAnalytics.PurgeErrorCount) \
 		XANALYTIC(IadInstallCacheDefragCount, CurrentInstallCacheAnalytics.DefragCount) \
@@ -1677,7 +1683,7 @@ void FOnDemandInstallCacheStats::OnStartupError(const FResult& Result)
 	}
 }
 
-void FOnDemandInstallCacheStats::OnFlush(const FResult& Result, int64 ByteCount)
+void FOnDemandInstallCacheStats::OnFlush(const FResult& Result, int64 ByteCount, bool bFromDefrag)
 {
 	if (Result.HasError() && UE::UnifiedError::IsCancellationError(Result.GetError()))
 	{
@@ -1685,12 +1691,25 @@ void FOnDemandInstallCacheStats::OnFlush(const FResult& Result, int64 ByteCount)
 	}
 
 	UE::TUniqueLock Lock(InstallCacheAnalytics);
-	if (Result.HasError())
+
+	if (bFromDefrag)
 	{
-		InstallCacheAnalytics.FlushErrorCount++;
+		if (Result.HasError())
+		{
+			InstallCacheAnalytics.DefragFlushErrorCount++;
+		}
+		InstallCacheAnalytics.DefragFlushCount++;
+		InstallCacheAnalytics.DefragFlushedBytes += ByteCount;
+	}
+	else
+	{
+		if (Result.HasError())
+		{
+			InstallCacheAnalytics.FlushErrorCount++;
+		}
+		InstallCacheAnalytics.FlushCount++;
+		InstallCacheAnalytics.FlushedBytes += ByteCount;
 	}
-	InstallCacheAnalytics.FlushCount++;
-	InstallCacheAnalytics.FlushedBytes += ByteCount;
 }
 
 void FOnDemandInstallCacheStats::OnJournalCommit(const FResult& Result, int64 ByteCount)
diff --git a/Engine/Source/Runtime/Experimental/IoStore/OnDemand/Private/Statistics.h b/Engine/Source/Runtime/Experimental/IoStore/OnDemand/Private/Statistics.h
--- a/Engine/Source/Runtime/Experimental/IoStore/OnDemand/Private/Statistics.h
+++ b/Engine/Source/Runtime/Experimental/IoStore/OnDemand/Private/Statistics.h
@@ -166,7 +166,7 @@ class FOnDemandInstallCacheStats
 {
 public:
 	static void OnStartupError(const FResult& Result) IAS_STATISTICS_IMPL()
-	static void OnFlush(const FResult& Result, int64 ByteCount) IAS_STATISTICS_IMPL()
+	static void OnFlush(const FResult& Result, int64 ByteCount, bool bFromDefrag) IAS_STATISTICS_IMPL()
 	static void OnJournalCommit(EIoErrorCode ErrorCode, int64 ByteCount) IAS_STATISTICS_IMPL()
 	static void OnJournalCommit(const FResult& Result, int64 ByteCount) IAS_STATISTICS_IMPL()
 	static void OnCasVerification(const FResult& Result, int32 RemovedChunkCount, uint64 BlockBytesOverBudget) IAS_STATISTICS_IMPL()
