IAD Stats - Fix defragmentation stats using request count instead of defrag count for averages.

#Synced-CL 49554261

diff --git a/Engine/Source/Runtime/Experimental/IoStore/OnDemand/Private/Statistics.cpp b/Engine/Source/Runtime/Experimental/IoStore/OnDemand/Private/Statistics.cpp
--- a/Engine/Source/Runtime/Experimental/IoStore/OnDemand/Private/Statistics.cpp
+++ b/Engine/Source/Runtime/Experimental/IoStore/OnDemand/Private/Statistics.cpp
@@ -1599,8 +1599,8 @@ static void LogOrReportContentInstallerAnalytics(bool bReport, bool bLog, TArray
 
 	const double AvgInstallDurationMs = CurrentInstallerAnalytics.InstallCount > 0 ? (double)CurrentInstallerAnalytics.TotalInstallDurationMs / (double)CurrentInstallerAnalytics.InstallCount : 0.0;
 	const double AvgCacheHitRatio = CurrentInstallerAnalytics.InstallCount > 0 ? (double)CurrentInstallerAnalytics.TotalCacheHitRatio / (double)CurrentInstallerAnalytics.InstallCount : 1.0;
-	const double AvgTotalFragmentationRatio = CurrentInstallerAnalytics.InstallCount > 0 ? (double)CurrentInstallerAnalytics.TotalFragmentionRatio / (double)CurrentInstallerAnalytics.InstallCount : 0.0;
-	const double AvgDefraggedBlockFragmentationRatio = CurrentInstallerAnalytics.InstallCount > 0 ? (double)CurrentInstallerAnalytics.DefraggedBlockFragmentionRatio / (double)CurrentInstallerAnalytics.InstallCount : 0.0;
+	const double AvgTotalFragmentationRatio = CurrentInstallCacheAnalytics.DefragCount > 0 ? (double)CurrentInstallerAnalytics.TotalFragmentionRatio / (double)CurrentInstallCacheAnalytics.DefragCount : 0.0;
+	const double AvgDefraggedBlockFragmentationRatio = CurrentInstallCacheAnalytics.DefragCount > 0 ? (double)CurrentInstallerAnalytics.DefraggedBlockFragmentionRatio / (double)CurrentInstallCacheAnalytics.DefragCount : 0.0;
 
 	#define INSTALLER_ANALYTICS_LIST(XANALYTIC) \
 		XANALYTIC(IadTotalInstallCount, CurrentInstallerAnalytics.InstallCount) \
