Chaos : fix impulse fields not properly waking up particles when they are asleep, causing them to not move

#Synced-CL 49743922

diff --git a/Engine/Source/Runtime/Experimental/Chaos/Public/PhysicsProxy/FieldSystemProxyHelper.h b/Engine/Source/Runtime/Experimental/Chaos/Public/PhysicsProxy/FieldSystemProxyHelper.h
--- a/Engine/Source/Runtime/Experimental/Chaos/Public/PhysicsProxy/FieldSystemProxyHelper.h
+++ b/Engine/Source/Runtime/Experimental/Chaos/Public/PhysicsProxy/FieldSystemProxyHelper.h
@@ -600,6 +600,8 @@ namespace Chaos
 	{
 		TFieldArrayView<FVector> ResultsView(FinalResults, 0, FinalResults.Num());
 
+		Chaos::FPBDRigidsSolver::FPBDRigidsEvolution* Evolution = RigidSolver->GetEvolution();
+
 		if (FieldCommand.PhysicsType == EFieldPhysicsType::Field_LinearVelocity)
 		{
 			SCOPE_CYCLE_COUNTER(STAT_ParamUpdateField_LinearVelocity);
@@ -607,10 +609,18 @@ namespace Chaos
 				static_cast<const FFieldNode<FVector>*>(FieldCommand.RootNode.Get())->Evaluate(FieldContext, ResultsView);
 				for (const FFieldContextIndex& Index : FieldContext.GetEvaluatedSamples())
 				{
-					Chaos::FPBDRigidParticleHandle* RigidHandle = ParticleHandles[Index.Sample]->CastToRigidParticle();
-					if (RigidHandle && RigidHandle->ObjectState() == Chaos::EObjectStateType::Dynamic)
+					const FVector DeltaV = ResultsView[Index.Result];
+					if (!DeltaV.IsZero())
 					{
-						RigidHandle->SetV(RigidHandle->GetV() + ResultsView[Index.Result]);
+						Chaos::FPBDRigidParticleHandle* RigidHandle = ParticleHandles[Index.Sample]->CastToRigidParticle();
+						if (RigidHandle && RigidHandle->IsDynamic())
+						{
+							RigidHandle->SetV(RigidHandle->GetV() + DeltaV);
+							if (Evolution && RigidHandle->IsSleeping())
+							{
+								Evolution->WakeParticle(RigidHandle);
+							}
+						}
 					}
 				}
 			}
@@ -622,11 +632,19 @@ namespace Chaos
 				static_cast<const FFieldNode<FVector>*>(FieldCommand.RootNode.Get())->Evaluate(FieldContext, ResultsView);
 				for (const FFieldContextIndex& Index : FieldContext.GetEvaluatedSamples())
 				{
-					Chaos::FPBDRigidParticleHandle* RigidHandle = ParticleHandles[Index.Sample]->CastToRigidParticle();
-					if (RigidHandle && RigidHandle->ObjectState() == Chaos::EObjectStateType::Dynamic)
+					const FVector DeltaImpulse = ResultsView[Index.Result];
+					if (!DeltaImpulse.IsZero())
 					{
-						const FVec3 CurrentImpulseVelocity = RigidHandle->LinearImpulseVelocity();
-						RigidHandle->SetLinearImpulseVelocity(CurrentImpulseVelocity + (ResultsView[Index.Result] * RigidHandle->InvM()));
+						Chaos::FPBDRigidParticleHandle* RigidHandle = ParticleHandles[Index.Sample]->CastToRigidParticle();
+						if (RigidHandle && RigidHandle->IsDynamic())
+						{
+							const FVec3 CurrentImpulseVelocity = RigidHandle->LinearImpulseVelocity();
+							RigidHandle->SetLinearImpulseVelocity(CurrentImpulseVelocity + (DeltaImpulse * RigidHandle->InvM()));
+							if (Evolution && RigidHandle->IsSleeping())
+							{
+								Evolution->WakeParticle(RigidHandle);
+							}
+						}
 					}
 				}
 			}
@@ -639,10 +657,18 @@ namespace Chaos
 
 				for (const FFieldContextIndex& Index : FieldContext.GetEvaluatedSamples())
 				{
-					Chaos::FPBDRigidParticleHandle* RigidHandle = ParticleHandles[Index.Sample]->CastToRigidParticle();
-					if (RigidHandle && RigidHandle->ObjectState() == Chaos::EObjectStateType::Dynamic)
+					const FVector DeltaW = ResultsView[Index.Result];
+					if (!DeltaW.IsZero())
 					{
-						RigidHandle->SetW(RigidHandle->GetW() + ResultsView[Index.Result]);
+						Chaos::FPBDRigidParticleHandle* RigidHandle = ParticleHandles[Index.Sample]->CastToRigidParticle();
+						if (RigidHandle && RigidHandle->IsDynamic())
+						{
+							RigidHandle->SetW(RigidHandle->GetW() + DeltaW);
+							if (Evolution && RigidHandle->IsSleeping())
+							{
+								Evolution->WakeParticle(RigidHandle);
+							}
+						}
 					}
 				}
 			}
