Dataflow : add a node to convert a numeric value to a dataflow image resolution enum

#Synced-CL 49754362

diff --git a/Engine/Source/Runtime/Experimental/Dataflow/Core/Private/Dataflow/DataflowImageNodes.cpp b/Engine/Source/Runtime/Experimental/Dataflow/Core/Private/Dataflow/DataflowImageNodes.cpp
--- a/Engine/Source/Runtime/Experimental/Dataflow/Core/Private/Dataflow/DataflowImageNodes.cpp
+++ b/Engine/Source/Runtime/Experimental/Dataflow/Core/Private/Dataflow/DataflowImageNodes.cpp
@@ -10,6 +10,7 @@ namespace UE::Dataflow
 {
 	void RegisterDataflowImageNodes()
 	{
+		DATAFLOW_NODE_REGISTER_CREATION_FACTORY(FDataflowConvertToResolutionNode);
 		DATAFLOW_NODE_REGISTER_CREATION_FACTORY(FDataflowImageFromColorNode);
 		DATAFLOW_NODE_REGISTER_CREATION_FACTORY(FDataflowImageSplitChannelsNode);
 		DATAFLOW_NODE_REGISTER_CREATION_FACTORY(FDataflowImageCombineChannelsNode);
@@ -18,6 +19,90 @@ namespace UE::Dataflow
 
 ////////////////////////////////////////////////////////////////////////////////////////////////////////////////
 
+FDataflowConvertToResolutionNode::FDataflowConvertToResolutionNode(const UE::Dataflow::FNodeParameters& InParam, FGuid InGuid)
+	:FDataflowNode(InParam, InGuid)
+{
+	RegisterInputConnection(&Value);
+	RegisterOutputConnection(&Resolution);
+}
+
+void FDataflowConvertToResolutionNode::Evaluate(UE::Dataflow::FContext& Context, const FDataflowOutput* Out) const
+{
+	using FResolutionPair = TPair<EDataflowImageResolution, EDataflowImageResolution>;
+	constexpr EDataflowImageResolution ResolutionsArray[] =
+	{
+		EDataflowImageResolution::Resolution16,
+		EDataflowImageResolution::Resolution32,
+		EDataflowImageResolution::Resolution64,
+		EDataflowImageResolution::Resolution128,
+		EDataflowImageResolution::Resolution256,
+		EDataflowImageResolution::Resolution512,
+		EDataflowImageResolution::Resolution1024,
+		EDataflowImageResolution::Resolution2048,
+		EDataflowImageResolution::Resolution4096,
+		EDataflowImageResolution::Resolution8192,
+	};
+	constexpr int32 NumResolutions = int32(sizeof(ResolutionsArray) / sizeof(EDataflowImageResolution));
+
+	auto GetPrevious = [&ResolutionsArray](int32 Value) -> EDataflowImageResolution
+		{
+			for (int32 Index = 1; Index < NumResolutions; ++Index)
+			{
+				if (Value < (int32)ResolutionsArray[Index])
+				{
+					return ResolutionsArray[Index - 1];
+				}
+			}
+			return EDataflowImageResolution::Resolution8192;
+		};
+
+	auto GetNext = [&ResolutionsArray](int32 Value) -> EDataflowImageResolution
+		{
+			for (int32 Index = (NumResolutions-2); Index >= 0 ; --Index)
+			{
+				if (Value > (int32)ResolutionsArray[Index])
+				{
+					return ResolutionsArray[Index + 1];
+				}
+			}
+			return EDataflowImageResolution::Resolution16;
+		};
+	
+	if (Out->IsA(&Resolution))
+	{
+		// clampo to reasonable number before converting to integer
+		const int32 InValue = FMath::Clamp(GetValue(Context, &Value), 0.0, (double)TNumericLimits<int32>::Max());
+
+		const EDataflowImageResolution PrevResolution = GetPrevious(InValue);
+		const EDataflowImageResolution NextResolution = GetNext(InValue);
+
+		EDataflowImageResolution OutResolution = EDataflowImageResolution::Resolution16;
+		if (Method == EDataflowConvertToResolutionMethod::UseClosest)
+		{
+			if ((InValue - (int32)PrevResolution) < ((int32)NextResolution - InValue))
+			{
+				OutResolution = PrevResolution;
+			}
+			else
+			{
+				OutResolution = NextResolution;
+			}
+		}
+		else if (Method == EDataflowConvertToResolutionMethod::UsePrevious)
+		{
+			OutResolution = PrevResolution;
+		}
+		else if (Method == EDataflowConvertToResolutionMethod::UseNext)
+		{
+			OutResolution = NextResolution;
+		}
+
+		SetValue(Context, MoveTemp(OutResolution), &Resolution);
+	}
+}
+
+////////////////////////////////////////////////////////////////////////////////////////////////////////////////
+
 FDataflowImageFromColorNode::FDataflowImageFromColorNode(const UE::Dataflow::FNodeParameters& InParam, FGuid InGuid)
 	:FDataflowNode(InParam, InGuid)
 {
diff --git a/Engine/Source/Runtime/Experimental/Dataflow/Core/Public/Dataflow/DataflowImageNodes.h b/Engine/Source/Runtime/Experimental/Dataflow/Core/Public/Dataflow/DataflowImageNodes.h
--- a/Engine/Source/Runtime/Experimental/Dataflow/Core/Public/Dataflow/DataflowImageNodes.h
+++ b/Engine/Source/Runtime/Experimental/Dataflow/Core/Public/Dataflow/DataflowImageNodes.h
@@ -9,6 +9,42 @@
 
 #define DATAFLOW_IMAGE_NODES_CATEGORY "Image"
 
+UENUM()
+enum class EDataflowConvertToResolutionMethod : int32
+{
+	UseClosest = 0 UMETA(DisplayName = "Closest resolution"),
+	UsePrevious = 1 UMETA(DisplayName = "Previous resolution"),
+	UseNext = 2 UMETA(DisplayName = "Next resolution"),
+};
+
+/**
+* Convert a numeric value into a resolution 
+*/
+USTRUCT()
+struct FDataflowConvertToResolutionNode : public FDataflowNode
+{
+	GENERATED_USTRUCT_BODY()
+	DATAFLOW_NODE_DEFINE_INTERNAL(FDataflowConvertToResolutionNode, "ConvertToResolution", DATAFLOW_IMAGE_NODES_CATEGORY, "")
+
+public:
+	FDataflowConvertToResolutionNode(const UE::Dataflow::FNodeParameters& InParam, FGuid InGuid = FGuid::NewGuid());
+
+private:
+	/** Arbitrary value to convert to a image resolution type */
+	UPROPERTY(EditAnywhere, Category = Options, meta = (DataflowInput));
+	FDataflowNumericTypes Value;
+
+	/** Method to use to match the resolution */
+	UPROPERTY(EditAnywhere, Category = Options);
+	EDataflowConvertToResolutionMethod Method;
+
+	/** Output resolution */
+	UPROPERTY(meta = (DataflowOutput));
+	EDataflowImageResolution Resolution = EDataflowImageResolution::Resolution16;
+
+	virtual void Evaluate(UE::Dataflow::FContext& Context, const FDataflowOutput* Out) const override;
+};
+
 /**
 * Create a RGBA image from a single color at a specific resolution
 * Outputs are single channel images
