The ChunkPlugin command now accepts '-ValidateChunks' on the commandline to opt in to more extensive validation checks.

#rnx

- Previously these additional checks were enabled by default. Now they are off by default and opted in to by the caller.

#Synced-CL 49717872

diff --git a/Engine/Source/Runtime/Experimental/IoStore/OnDemand/Private/Tool/ChunkPlugin.cpp b/Engine/Source/Runtime/Experimental/IoStore/OnDemand/Private/Tool/ChunkPlugin.cpp
--- a/Engine/Source/Runtime/Experimental/IoStore/OnDemand/Private/Tool/ChunkPlugin.cpp
+++ b/Engine/Source/Runtime/Experimental/IoStore/OnDemand/Private/Tool/ChunkPlugin.cpp
@@ -274,7 +274,7 @@ static bool ShouldValidateChunk(FName CompressionFormat, uint64 ChunkRawSize)
 }
 
 static int32 ChunkContainer(const FString& ContainerUTocPath, const FChunkPluginSettings& Settings, const TMap<FGuid, FAES::FAESKey>& EncryptionKeys, FChunkWriterInterface* ChunkWriter, const FString& ChunksRelativeFolder, bool bIgnoreContainerHeader,
-	FOnDemandToc& OutOnDemandToc, TMap<FString, TArray<FOnDemandTocTagSetPackageList>>* OutFoundTagSets, TArray<FString>& OutContainerFiles, FChunkPluginStats& OutStats)
+	bool bValidateChunks, FOnDemandToc& OutOnDemandToc, TMap<FString, TArray<FOnDemandTocTagSetPackageList>>* OutFoundTagSets, TArray<FString>& OutContainerFiles, FChunkPluginStats& OutStats)
 {
 	TRACE_CPUPROFILER_EVENT_SCOPE(ChunkPluginCommand::ChunkContainer);
 
@@ -441,7 +441,7 @@ static int32 ChunkContainer(const FString& ContainerUTocPath, const FChunkPlugin
 		}
 
 		// Additional validation to make sure that the chunks can be decrypted and decoded
-		if (ShouldValidateChunk(CompressionFormat, RawChunkSize))
+		if (bValidateChunks && ShouldValidateChunk(CompressionFormat, RawChunkSize))
 		{
 			FIoChunkDecodingParams DecodingParams;
 			DecodingParams.CompressionFormat = CompressionFormat;
@@ -597,6 +597,8 @@ static int32 ChunkPluginCommandEntry(const FContext& Context)
 
 	const bool bIncludeSigPak			= Context.Get<bool>(TEXT("-IncludeSigPak"), false);
 	const bool bDeleteContainerFiles	= !Context.Get<bool>(TEXT("-KeepContainerFiles"), false);
+	const bool bValidateChunks			= Context.Get<bool>(TEXT("-ValidateChunks"), false);
+
 	FString OutputFolder				= InOutputFolder;
 	FString ContainerFolder				= InputFolder;
 	FString IoStoreRelativeFolder		= "iostore";
@@ -622,6 +624,7 @@ static int32 ChunkPluginCommandEntry(const FContext& Context)
 	UE_LOG(LogIoStoreOnDemand, Display, TEXT("\tIncludeSigPak: %s"), bIncludeSigPak ? TEXT("true") : TEXT("false"));
 	UE_LOG(LogIoStoreOnDemand, Display, TEXT("\tDeleteContainerFiles: %s"), bDeleteContainerFiles ? TEXT("true") : TEXT("false"));
 	UE_LOG(LogIoStoreOnDemand, Display, TEXT("\tIgnoreContainerHeader: %s"), bIgnoreContainerHeader ? TEXT("true") : TEXT("false"));
+	UE_LOG(LogIoStoreOnDemand, Display, TEXT("\tValidateChunks: %s"), bValidateChunks ? TEXT("true") : TEXT("false"));
 	UE_LOG(LogIoStoreOnDemand, Display, TEXT("\tTOC flags: '%s'"), *LexToString(TocFlags));
 
 	IFileManager& FileMgr = IFileManager::Get();
@@ -759,7 +762,7 @@ static int32 ChunkPluginCommandEntry(const FContext& Context)
 		if (bSupportStreamingContainers && bIsStreamOnDemandContainer)
 		{
 			// Attempt to add to IAS output container
-			int32 Result = ChunkContainer(FullPath, Settings, EncryptionKeys, ChunkWriter.Get(), ChunksRelativeFolder, bIgnoreContainerHeader, StreamingOnDemandToc, nullptr /*OutFoundTagSets*/, ContainerFiles, Stats);
+			int32 Result = ChunkContainer(FullPath, Settings, EncryptionKeys, ChunkWriter.Get(), ChunksRelativeFolder, bIgnoreContainerHeader, bValidateChunks, StreamingOnDemandToc, nullptr /*OutFoundTagSets*/, ContainerFiles, Stats);
 			if (Result < 0)
 			{
 				return Result;
@@ -768,7 +771,7 @@ static int32 ChunkPluginCommandEntry(const FContext& Context)
 		else if(bSupportInstalledContainers && !bIsStreamOnDemandContainer)
 		{
 			// Attempt to add to IAD output container
-			int32 Result = ChunkContainer(FullPath, Settings, EncryptionKeys, ChunkWriter.Get(), ChunksRelativeFolder, bIgnoreContainerHeader, OnDemandToc, &FoundTagSets, ContainerFiles, Stats);
+			int32 Result = ChunkContainer(FullPath, Settings, EncryptionKeys, ChunkWriter.Get(), ChunksRelativeFolder, bIgnoreContainerHeader, bValidateChunks, OnDemandToc, &FoundTagSets, ContainerFiles, Stats);
 			if (Result < 0)
 			{
 				return Result;
@@ -981,6 +984,7 @@ static FCommand ChunkPluginCommand(
 		TArgument<bool>(TEXT("-IgnoreContainerHeader"),		TEXT("Ignore the container header?")),
 		TArgument<FStringView>(TEXT("-BucketPrefix"),		TEXT("Path to prefix to bucket objects")),
 		TArgument<int32>(TEXT("-MaxConcurrentUploads"),		TEXT("Number of simultaneous uploads")),
+		TArgument<bool>(TEXT("-ValidateChunks"),			TEXT("Runs additional validation on chunks after they are read from the container and before they are written out")),
 		S3Arguments,
 	}
 );
