Improved error checking when writing chunks.

#rnx

- Removed the standalone function::WriteChunk as it is unused.
- FDiskChunkWriter::WriteChunk now checks that writing to the file succeeded as well as opening the file handle.

#Synced-CL 49717872

diff --git a/Engine/Source/Runtime/Experimental/IoStore/OnDemand/Private/Tool/ChunkPlugin.cpp b/Engine/Source/Runtime/Experimental/IoStore/OnDemand/Private/Tool/ChunkPlugin.cpp
--- a/Engine/Source/Runtime/Experimental/IoStore/OnDemand/Private/Tool/ChunkPlugin.cpp
+++ b/Engine/Source/Runtime/Experimental/IoStore/OnDemand/Private/Tool/ChunkPlugin.cpp
@@ -148,14 +148,11 @@ class FDiskChunkWriter : public FChunkWriterInterface
 {
 public:
 	FDiskChunkWriter(const FString& InOutputFolder)
+		: OutputFolder(InOutputFolder)
 	{
-		OutputFolder = InOutputFolder;
-		
 	}
 
-	virtual ~FDiskChunkWriter()
-	{
-	}
+	virtual ~FDiskChunkWriter() = default;
 
 	virtual  FIoStatus WriteChunk(const FString& Directory, FIoBuffer Chunk, const FIoHash& Hash) override 
 	{
@@ -176,63 +173,34 @@ public:
 
 		Sb << TEXT("/") << HashString << TEXT(".iochunk");
 
-		if (TUniquePtr<FArchive> Ar(IFileManager::Get().CreateFileWriter(Sb.ToString())); Ar.IsValid())
+		TUniquePtr<FArchive> Ar(IFileManager::Get().CreateFileWriter(Sb.ToString()));
+
+		if (!Ar.IsValid())
 		{
-			UE_LOG(LogIoStoreOnDemand, Display, TEXT("Writing file '%s' (%.2lf KiB)"),
-				Sb.ToString(), double(Chunk.GetView().GetSize()) / 1024.0);
-			Ar->Serialize((void*)Chunk.GetView().GetData(), Chunk.GetView().GetSize());
+			return FIoStatusBuilder(EIoErrorCode::WriteError) << TEXT("Failed to open handle to write file '") << FString(Sb.ToString()) << TEXT("'");
+		}
+
+		Ar->Serialize((void*)Chunk.GetView().GetData(), Chunk.GetView().GetSize());
 
-			return EIoErrorCode::Ok;
+		if (!Ar->Close())
+		{
+			return FIoStatusBuilder(EIoErrorCode::WriteError) << TEXT("Failed to write to file '") << FString(Sb.ToString()) << TEXT("'");
 		}
 
-		return FIoStatusBuilder(EIoErrorCode::WriteError)
-			<< TEXT("Failed to write file '")
-			<< FString(Sb.ToString())
-			<< TEXT("'");
+		UE_LOG(LogIoStoreOnDemand, Display, TEXT("Writing file '%s' (%.2lf KiB)"), Sb.ToString(), double(Chunk.GetView().GetSize()) / 1024.0);
+		return EIoErrorCode::Ok;
 	}
 
 	bool Flush() override
 	{
 		return true;
 	}
+
 private:
 	FString OutputFolder;
 };
 
 ////////////////////////////////////////////////////////////////////////////////
-static FIoStatus WriteChunk(const FString& Directory, FMemoryView Chunk, const FIoHash& Hash)
-{
-	IFileManager& FileMgr = IFileManager::Get();
-	const FString HashString = LexToString(Hash);
-
-	TStringBuilder<256> Sb;
-	Sb << Directory << TEXT("/") << HashString.Left(2);
-
-	bool bTree = true;
-	if (FileMgr.MakeDirectory(Sb.ToString(), bTree) == false)
-	{
-		return FIoStatusBuilder(EIoErrorCode::WriteError)
-			<< TEXT("Failed to create directory '")
-			<< FString(Sb.ToString())
-			<< TEXT("'");
-	}
-
-	Sb << TEXT("/") << HashString << TEXT(".iochunk");
-
-	if (TUniquePtr<FArchive> Ar(IFileManager::Get().CreateFileWriter(Sb.ToString())); Ar.IsValid())
-	{
-		UE_LOG(LogIoStoreOnDemand, Display, TEXT("Writing file '%s' (%.2lf KiB)"),
-			Sb.ToString(), double(Chunk.GetSize()) / 1024.0);
-		Ar->Serialize((void*)Chunk.GetData(), Chunk.GetSize());
-
-		return EIoErrorCode::Ok;
-	}
-
-	return FIoStatusBuilder(EIoErrorCode::WriteError)
-		<< TEXT("Failed to write file '")
-		<< FString(Sb.ToString())
-		<< TEXT("'");
-}
 
 static bool ShouldValidateChunk(FName CompressionFormat, uint64 ChunkRawSize)
 {
