Another experiment to determine if a slow down is due to the number of times UE_LOG is called, or just the amount of data that we are logging.

#rnx

#Synced-CL 49774731

diff --git a/Engine/Source/Runtime/Experimental/IoStore/OnDemand/Private/Tool/ChunkPlugin.cpp b/Engine/Source/Runtime/Experimental/IoStore/OnDemand/Private/Tool/ChunkPlugin.cpp
--- a/Engine/Source/Runtime/Experimental/IoStore/OnDemand/Private/Tool/ChunkPlugin.cpp
+++ b/Engine/Source/Runtime/Experimental/IoStore/OnDemand/Private/Tool/ChunkPlugin.cpp
@@ -612,6 +612,8 @@ static int32 ChunkPluginCommandEntry(const FContext& Context)
 	FChunkPluginSettings Settings;
 	if (!SettingsFile.IsEmpty())
 	{
+		UE_LOG(LogIoStoreOnDemand, Display, TEXT("Loading settings file '%s'..."), *SettingsFile);
+
 		TUniquePtr<FArchive> Ar(FileMgr.CreateFileReader(*SettingsFile));
 		if (!Ar || Ar->IsError())
 		{
@@ -667,12 +669,12 @@ static int32 ChunkPluginCommandEntry(const FContext& Context)
 					Output << LINE_TERMINATOR << "->\t'" << Package << "'";
 				}
 
-				UE_LOG(LogIoStoreOnDemand, Log, TEXT("Package Set contains %d entries"), Packages.Num());
-				UE_LOG(LogIoStoreOnDemand, Display, TEXT("%hs"), *Output);
+				UE_LOGF(LogIoStoreOnDemand, Display, "TEST: Before Large Message");
+				UE_LOGF(LogIoStoreOnDemand, Display, "%hs", *Output);
+				UE_LOGF(LogIoStoreOnDemand, Display, "TEST: After Large Message");
 #endif
-				
-
 				Settings.PackageSets.Add(FString(Pair.Key), MoveTemp(Packages));
+				UE_LOGF(LogIoStoreOnDemand, Display, "Processed %d package set entries", Packages.Num());
 			}
 		}
 	}
