Fix potential crash in constraint destructor.
The constraint was potentially delted before being remove from the constraint graph.
Made sure the deletion could occurs only after it has been remove from the graph

#Synced-CL 49586837

diff --git a/Engine/Source/Runtime/Experimental/Chaos/Public/Chaos/PBDRigidsEvolution.h b/Engine/Source/Runtime/Experimental/Chaos/Public/Chaos/PBDRigidsEvolution.h
--- a/Engine/Source/Runtime/Experimental/Chaos/Public/Chaos/PBDRigidsEvolution.h
+++ b/Engine/Source/Runtime/Experimental/Chaos/Public/Chaos/PBDRigidsEvolution.h
@@ -671,6 +671,11 @@ public:
 	*/
 	void DisconnectConstraints(const TSet<FGeometryParticleHandle*>& RemovedParticles)
 	{
+		for (FGeometryParticleHandle* ParticleHandle : RemovedParticles)
+		{
+			RemoveConstraintsFromConstraintGraph(ParticleHandle->ParticleConstraints());
+		}
+
 		for (FPBDConstraintContainer* Container : ConstraintContainers)
 		{
 			Container->DisconnectConstraints(RemovedParticles);
@@ -678,7 +683,6 @@ public:
 
 		for (FGeometryParticleHandle* ParticleHandle : RemovedParticles)
 		{
-			RemoveConstraintsFromConstraintGraph(ParticleHandle->ParticleConstraints());
 			ParticleHandle->ParticleConstraints().Reset();
 		}
 	}
