[DataFlow] Export SetVariableOverrideAndNotify symbol

- This symbol is used in OverrideVariableArrayAndNotify, which is a template function that is inlined on the call site. If this function is used outside the dataflow module it will cause a linker error.

#rnx

#Synced-CL 49759259

diff --git a/Engine/Source/Runtime/Experimental/Dataflow/Engine/Public/Dataflow/DataflowInstance.h b/Engine/Source/Runtime/Experimental/Dataflow/Engine/Public/Dataflow/DataflowInstance.h
--- a/Engine/Source/Runtime/Experimental/Dataflow/Engine/Public/Dataflow/DataflowInstance.h
+++ b/Engine/Source/Runtime/Experimental/Dataflow/Engine/Public/Dataflow/DataflowInstance.h
@@ -174,7 +174,7 @@ private:
 	void RemoveOverridenVariablesNotInDataflowAsset();
 
 private:
-	void SetVariableOverrideAndNotify(const FGuid& PropertyID, bool bOverrideState);
+	DATAFLOWENGINE_API void SetVariableOverrideAndNotify(const FGuid& PropertyID, bool bOverrideState);
 
 	template<typename T, typename TWriteElementFunction>
 	bool OverrideVariableArrayAndNotify(FInstancedPropertyBag& InVariables, FName InVariableName, const TArray<T>& InValues, const TWriteElementFunction& InWriteElementFunction)
