Dataflow : fix vertex selection by box  node not properly handling scaled and rotated box transform

#Synced-CL 49797956

diff --git a/Engine/Source/Runtime/Experimental/Chaos/Private/GeometryCollection/Facades/CollectionTransformSelectionFacade.cpp b/Engine/Source/Runtime/Experimental/Chaos/Private/GeometryCollection/Facades/CollectionTransformSelectionFacade.cpp
--- a/Engine/Source/Runtime/Experimental/Chaos/Private/GeometryCollection/Facades/CollectionTransformSelectionFacade.cpp
+++ b/Engine/Source/Runtime/Experimental/Chaos/Private/GeometryCollection/Facades/CollectionTransformSelectionFacade.cpp
@@ -786,6 +786,10 @@ namespace GeometryCollection::Facades
 
 		const TArray<FTransform> CollectionSpaceTransforms = TransformFacade.ComputeCollectionSpaceTransforms();
 
+		const FVector BoxCenter = InBox.GetCenter();
+		const FVector BoxExtentsScaled = InBox.GetExtent() * InBoxTransform.GetScale3D();
+		const FBox ScaledBoxCenterAtZero(-BoxExtentsScaled, BoxExtentsScaled);
+
 		for (int32 VertexIdx = 0; VertexIdx < NumVertices; ++VertexIdx)
 		{
 			const int32 TransformIdx = MeshFacade.BoneMapAttribute.IsValidIndex(VertexIdx) ? MeshFacade.BoneMapAttribute[VertexIdx] : INDEX_NONE;
@@ -795,9 +799,9 @@ namespace GeometryCollection::Facades
 			{
 				const FVector VertexInBoneSpace = FVector(MeshFacade.VertexAttribute[VertexIdx]);
 				const FVector VertexInCollectionSpace = CollectionSpaceTransforms[TransformIdx].TransformPosition(VertexInBoneSpace);
-				const FVector VertexInBoxSpace = InBoxTransform.InverseTransformPosition(VertexInCollectionSpace);
+				const FVector VertexInBoxSpace = InBoxTransform.InverseTransformPositionNoScale(VertexInCollectionSpace - BoxCenter);
 
-				bIsSelected = InBox.IsInside(VertexInBoxSpace);
+				bIsSelected = ScaledBoxCenterAtZero.IsInside(VertexInBoxSpace);
 			}
 			OutVertexSelection.Add(bIsSelected);
 		}
