Mark up the ChunkPlugin command with trace scopes.

#rnx
#prelfight

#Synced-CL 49774731

diff --git a/Engine/Source/Runtime/Experimental/IoStore/OnDemand/Private/Tool/ChunkPlugin.cpp b/Engine/Source/Runtime/Experimental/IoStore/OnDemand/Private/Tool/ChunkPlugin.cpp
--- a/Engine/Source/Runtime/Experimental/IoStore/OnDemand/Private/Tool/ChunkPlugin.cpp
+++ b/Engine/Source/Runtime/Experimental/IoStore/OnDemand/Private/Tool/ChunkPlugin.cpp
@@ -110,6 +110,8 @@ public:
 
 	virtual FIoStatus WriteChunk(const FString& RelativeDirectory, FIoBuffer Chunk, const FIoHash& Hash) override
 	{
+		TRACE_CPUPROFILER_EVENT_SCOPE(FS3ChunkWriter::WriteChunk);
+
 		const FString HashString = LexToString(Hash);
 
 		TStringBuilder<256> Key;
@@ -129,6 +131,8 @@ public:
 
 	virtual bool Flush() override
 	{
+		TRACE_CPUPROFILER_EVENT_SCOPE(FS3ChunkWriter::Flush);
+
 		if (!UploadQueue.Flush())
 		{
 			return false;
@@ -156,6 +160,8 @@ public:
 
 	virtual  FIoStatus WriteChunk(const FString& Directory, FIoBuffer Chunk, const FIoHash& Hash) override 
 	{
+		TRACE_CPUPROFILER_EVENT_SCOPE(FDiskChunkWriter::WriteChunk);
+
 		IFileManager& FileMgr = IFileManager::Get();
 		const FString HashString = LexToString(Hash);
 
@@ -367,36 +373,44 @@ static int32 ChunkContainer(const FString& ContainerUTocPath, const FChunkPlugin
 
 		const uint32 BlockOffset = ContainerEntry.BlockSizes.Num();
 		const uint32 BlockCount = ReadResult.Blocks.Num();
-		const FIoHash ChunkHash = FIoHash::HashBuffer(ReadResult.IoBuffer.GetView());
+		auto HashBuffer = [](const FIoStoreCompressedReadResult& ReadResult)
+		{
+			TRACE_CPUPROFILER_EVENT_SCOPE(HashBuffer);
+			return FIoHash::HashBuffer(ReadResult.IoBuffer.GetView());
+		};
+		const FIoHash ChunkHash = HashBuffer(ReadResult);
 		const FString HashString = LexToString(ChunkHash);
 
 		FMemoryView EncodedBlocks = ReadResult.IoBuffer.GetView();
 		uint64 RawChunkSize = 0;
 		uint64 EncodedChunkSize = 0;
 		FName CompressionFormat = NAME_None;
-		for (const FIoStoreCompressedBlockInfo& BlockInfo : ReadResult.Blocks)
+
 		{
-			check(Align(BlockInfo.CompressedSize, FAES::AESBlockSize) == BlockInfo.AlignedSize);
-			const uint64 EncodedBlockSize = BlockInfo.AlignedSize;
-			ContainerEntry.BlockSizes.Add(uint32(BlockInfo.CompressedSize));
+			TRACE_CPUPROFILER_EVENT_SCOPE(HashBlocks);
+			for (const FIoStoreCompressedBlockInfo& BlockInfo : ReadResult.Blocks)
+			{
+				check(Align(BlockInfo.CompressedSize, FAES::AESBlockSize) == BlockInfo.AlignedSize);
+				const uint64 EncodedBlockSize = BlockInfo.AlignedSize;
+				ContainerEntry.BlockSizes.Add(uint32(BlockInfo.CompressedSize));
 
-			FMemoryView EncodedBlock = EncodedBlocks.Left(EncodedBlockSize);
-			EncodedBlocks += EncodedBlock.GetSize();
+				FMemoryView EncodedBlock = EncodedBlocks.Left(EncodedBlockSize);
+				EncodedBlocks += EncodedBlock.GetSize();
 
-			check(((char*)EncodedBlock.GetData() + EncodedBlock.GetSize()) <= ((char*)ReadResult.IoBuffer.GetView().GetData() + ReadResult.IoBuffer.GetView().GetSize()));
-			
-			FIoBlockHash Hash = FIoChunkEncoding::HashBlock(EncodedBlock);
-			ContainerEntry.BlockHashes.Add(Hash);
+				check(((char*)EncodedBlock.GetData() + EncodedBlock.GetSize()) <= ((char*)ReadResult.IoBuffer.GetView().GetData() + ReadResult.IoBuffer.GetView().GetSize()));
+				FIoBlockHash Hash = FIoChunkEncoding::HashBlock(EncodedBlock);
+				ContainerEntry.BlockHashes.Add(Hash);
 
-			EncodedChunkSize += EncodedBlockSize;
-			RawChunkSize += BlockInfo.UncompressedSize;
+				EncodedChunkSize += EncodedBlockSize;
+				RawChunkSize += BlockInfo.UncompressedSize;
 
-			if (BlockInfo.CompressionMethod != NAME_None)
-			{
-				CompressionFormat = BlockInfo.CompressionMethod;
-				if (OutOnDemandToc.Header.CompressionFormat.IsEmpty())
+				if (BlockInfo.CompressionMethod != NAME_None)
 				{
-					OutOnDemandToc.Header.CompressionFormat = BlockInfo.CompressionMethod.ToString();
+					CompressionFormat = BlockInfo.CompressionMethod;
+					if (OutOnDemandToc.Header.CompressionFormat.IsEmpty())
+					{
+						OutOnDemandToc.Header.CompressionFormat = BlockInfo.CompressionMethod.ToString();
+					}
 				}
 			}
 		}
@@ -411,6 +425,8 @@ static int32 ChunkContainer(const FString& ContainerUTocPath, const FChunkPlugin
 		// Additional validation to make sure that the chunks can be decrypted and decoded
 		if (bValidateChunks && ShouldValidateChunk(CompressionFormat, RawChunkSize))
 		{
+			TRACE_CPUPROFILER_EVENT_SCOPE(ValidateChunk);
+
 			FIoChunkDecodingParams DecodingParams;
 			DecodingParams.CompressionFormat = CompressionFormat;
 			DecodingParams.TotalRawSize = RawChunkSize;
