Constraints: ensure that any existing constraint prerequisites are removed if the tick function of the child handle cannot be retrieved

#rnx

#Synced-CL 49801493

diff --git a/Engine/Source/Runtime/Experimental/Animation/Constraints/Private/Transform/TransformConstraint.cpp b/Engine/Source/Runtime/Experimental/Animation/Constraints/Private/Transform/TransformConstraint.cpp
--- a/Engine/Source/Runtime/Experimental/Animation/Constraints/Private/Transform/TransformConstraint.cpp
+++ b/Engine/Source/Runtime/Experimental/Animation/Constraints/Private/Transform/TransformConstraint.cpp
@@ -667,6 +667,36 @@ void UTickableTransformConstraint::TeardownConstraint(UWorld* InWorld)
 		{
 			ChildTickFunction->RemovePrerequisite(this, *ConstraintTick);
 		}
+		else
+		{
+			// If ChildTickFunction is null, meaning that the handle no longer references any component,
+			// we must ensure that no component in this level still has it as a prerequisite.
+			FConstraintTickFunction* ConstraintTickPtr = ConstraintTick.Get();
+
+			static constexpr bool bEvenIfPendingKill = true;
+			for (const TObjectPtr<AActor>& Actor: Level->Actors)
+			{
+				if (Actor)
+				{
+					for (UActorComponent* Component: Actor->GetComponents())
+					{
+						if (Component)
+						{
+							TArray<FTickPrerequisite>& Prerequisites = Component->PrimaryComponentTick.GetPrerequisites();
+							for (int32 PrereqIndex = 0; PrereqIndex < Prerequisites.Num(); PrereqIndex++)
+							{
+								FTickPrerequisite& Prerequisite = Prerequisites[PrereqIndex];
+								UObject* PrereqObject = Prerequisite.PrerequisiteObject.Get(bEvenIfPendingKill);
+								if (PrereqObject == this && Prerequisite.Get() == ConstraintTickPtr)
+								{
+									Prerequisites.RemoveAtSwap(PrereqIndex--);
+								}
+							}
+						}
+					}
+				}
+			}
+		}
 	
 		if (FTickFunction* ParentTickFunction = GetParentHandleTickFunction())
 		{
