[NetPhysics] When we want to merge inputs but the input to merge doesn't exist, create it via interpolation and then merge it.

#Synced-CL 49675227

diff --git a/Engine/Source/Runtime/Experimental/Chaos/Public/RewindData.h b/Engine/Source/Runtime/Experimental/Chaos/Public/RewindData.h
--- a/Engine/Source/Runtime/Experimental/Chaos/Public/RewindData.h
+++ b/Engine/Source/Runtime/Experimental/Chaos/Public/RewindData.h
@@ -306,14 +306,41 @@ public :
 
 	FORCEINLINE virtual void MergeData(int32 FromFrame, void* ToData) override
 	{
-		const int32 ToFrame = static_cast<DataType*>(ToData)->LocalFrame;
+		if (!ToData)
+		{
+			return;
+		}
+
+		DataType& MergeToData = *static_cast<DataType*>(ToData);
+
+		const int32 ToFrame = MergeToData.LocalFrame;
 		for (; FromFrame < ToFrame; FromFrame++)
 		{
 			const int32 FromIndex = GetFrameIndex(FromFrame);
+
+			// If data is missing, create it via interpolation
+			if (FromFrame != DataHistory[FromIndex].LocalFrame)
+			{
+				const int32 MinFrameIndex = ClosestData(FromFrame, true);
+				if (MinFrameIndex != INDEX_NONE)
+				{
+					DataType& NewData = DataHistory[FromIndex];
+					const int32 DeltaFrame = MergeToData.LocalFrame - FromFrame;
+					NewData.LocalFrame = FromFrame;
+					NewData.ServerFrame = MergeToData.ServerFrame - DeltaFrame;
+					NewData.bDataAltered = true; // This history entry is considered altered since it was produced locally by interpolation
+
+					PRAGMA_DISABLE_DEPRECATION_WARNINGS
+					NewData.DoInterpolateData(DataHistory[MinFrameIndex], MergeToData);
+					PRAGMA_ENABLE_DEPRECATION_WARNINGS
+				}
+			}
+
+			// Perform MergeData
 			if (FromFrame == DataHistory[FromIndex].LocalFrame)
 			{
-				static_cast<DataType*>(ToData)->MergeData(DataHistory[FromIndex]);
-				static_cast<DataType*>(ToData)->bDataAltered = true; // This history entry is now altered and doesn't correspond to the source entry
+				MergeToData.MergeData(DataHistory[FromIndex]);
+				MergeToData.bDataAltered = true; // This data is now altered and doesn't correspond to the source data
 			}
 		}
 	}
