[NetPhysics] Stop warning for duplicate CVar registration.

#Synced-CL 49682755

#Synced-CL 49682755

diff --git a/Engine/Source/Runtime/Experimental/Chaos/Private/Chaos/Framework/PhysicsSolverBase.cpp b/Engine/Source/Runtime/Experimental/Chaos/Private/Chaos/Framework/PhysicsSolverBase.cpp
--- a/Engine/Source/Runtime/Experimental/Chaos/Private/Chaos/Framework/PhysicsSolverBase.cpp
+++ b/Engine/Source/Runtime/Experimental/Chaos/Private/Chaos/Framework/PhysicsSolverBase.cpp
@@ -35,6 +35,21 @@ namespace Chaos
 	FAutoConsoleVariableRef CVarResimulateOnSpecifiedStep(TEXT("p.Resim.ResimulateOnSpecifiedStep"), ResimulateOnSpecifiedStep
 		, TEXT("0 = Resimulation can trigger on any physics step. \n -1 or other negative values = Resimulation only happens on the last physics step (if multiple are queued up at the same time). \n N = Resimulation only happens on the specified step or last, whichever comes first (if multiple are queued up at the same time)"));
 
+	/** Check if we can enable debugging informations for network physics */
+	static int32 DebugNetworkPhysicsPrediction = 0;
+	static FAutoConsoleVariableRef CVarDebugNetworkPhysicsPrediction(TEXT("np2.DebugNetworkPhysicsPrediction"), DebugNetworkPhysicsPrediction, TEXT("Debugs network physics prediction"));
+	/*static*/ bool FPhysicsSolverBase::CanDebugNetworkPhysicsPrediction()
+	{
+		return !!DebugNetworkPhysicsPrediction;
+	}
+
+	/** Return the interpolation lerp in case the resim is off */
+	static float NetworkPhysicsPredictionInterpLerp = 0.1f;
+	static FAutoConsoleVariableRef CVarNetworkPhysicsPredictionInterpLerp(TEXT("np2.NetworkPhysicsPredictionInterpLerp"), NetworkPhysicsPredictionInterpLerp, TEXT("State lerp value in between the target state and the current one in case resim is disabled or if the pawn is not possessed (continuous correction)"));
+	/*static*/ float FPhysicsSolverBase::NetworkPhysicsInterpolationLerp()
+	{
+		return NetworkPhysicsPredictionInterpLerp;
+	}
 
 	extern int GSingleThreadedPhysics;
 	void FPhysicsSolverBase::ChangeBufferMode(EMultiBufferMode InBufferMode)
diff --git a/Engine/Source/Runtime/Experimental/Chaos/Public/Chaos/Framework/PhysicsSolverBase.h b/Engine/Source/Runtime/Experimental/Chaos/Public/Chaos/Framework/PhysicsSolverBase.h
--- a/Engine/Source/Runtime/Experimental/Chaos/Public/Chaos/Framework/PhysicsSolverBase.h
+++ b/Engine/Source/Runtime/Experimental/Chaos/Public/Chaos/Framework/PhysicsSolverBase.h
@@ -795,13 +795,7 @@ namespace Chaos
 		float GetAsyncInterpolationMultiplier() const { return AsyncMultiplier; }
 
 		/** Check if we can enable debugging informations for network physics */
-		static bool CanDebugNetworkPhysicsPrediction()
-		{
-			static int32 DebugNetworkPhysicsPrediction = 0;
-			static FAutoConsoleVariableRef CVarDebugNetworkPhysicsPrediction(TEXT("np2.DebugNetworkPhysicsPrediction"), DebugNetworkPhysicsPrediction, TEXT("Debugs network physics prediction"));
-
-			return !!DebugNetworkPhysicsPrediction;
-		}
+		CHAOS_API static bool CanDebugNetworkPhysicsPrediction();
 
 		/** Check if network physics is enables or not */
 		static bool IsNetworkPhysicsPredictionEnabled()
@@ -873,13 +867,7 @@ namespace Chaos
 		}
 
 		/** Return the interpolation lerp in case the resim is off */
-		static float NetworkPhysicsInterpolationLerp()
-		{
-			static float NetworkPhysicsPredictionInterpLerp = 0.1f;
-			static FAutoConsoleVariableRef CVarNetworkPhysicsPredictionInterpLerp(TEXT("np2.NetworkPhysicsPredictionInterpLerp"), NetworkPhysicsPredictionInterpLerp, TEXT("State lerp value in between the target state and the current one in case resim is disabled or if the pawn is not possessed (continuous correction)"));
-
-			return NetworkPhysicsPredictionInterpLerp;
-		}
+		static float NetworkPhysicsInterpolationLerp();
 
 		/** Get the standalone solver flag */
 		bool IsStandaloneSolver() const { return bIsStandaloneSolver;}
