Start simplifying the ChunkPlugin code.

#rnx

- Moved most of the input parameters of ChunkContainer to a new FChunkParams structure and the output parameters to FChunkOutput.
- FOnDemandToc will be fixed in a subsequent submit which will start to remove the duplicate code paths.

#Synced-CL 49815329

diff --git a/Engine/Source/Runtime/Experimental/IoStore/OnDemand/Private/Tool/ChunkPlugin.cpp b/Engine/Source/Runtime/Experimental/IoStore/OnDemand/Private/Tool/ChunkPlugin.cpp
--- a/Engine/Source/Runtime/Experimental/IoStore/OnDemand/Private/Tool/ChunkPlugin.cpp
+++ b/Engine/Source/Runtime/Experimental/IoStore/OnDemand/Private/Tool/ChunkPlugin.cpp
@@ -247,16 +247,35 @@ static bool ShouldValidateChunk(FName CompressionFormat, uint64 ChunkRawSize)
 	return true;
 }
 
-static int32 ChunkContainer(const FString& ContainerUTocPath, const FChunkPluginSettings& Settings, const TMap<FGuid, FAES::FAESKey>& EncryptionKeys, FChunkWriterInterface* ChunkWriter, const FString& ChunksRelativeFolder, bool bIgnoreContainerHeader,
-	bool bValidateChunks, FOnDemandToc& OutOnDemandToc, TMap<FString, TArray<FOnDemandTocTagSetPackageList>>* OutFoundTagSets, TArray<FString>& OutContainerFiles, FChunkPluginStats& OutStats)
+struct FChunkParams
+{
+	FChunkPluginSettings Settings;
+	TMap<FGuid, FAES::FAESKey> EncryptionKeys;
+
+	TUniquePtr<FChunkWriterInterface> ChunkWriter;
+	FString ChunksRelativeFolder;
+
+	bool bIgnoreContainerHeader;
+	bool bValidateChunks;
+};
+
+struct FChunkOutput
+{
+//	FOnDemandToc OnDemandToc;
+	TMap<FString, TArray<FOnDemandTocTagSetPackageList>> FoundTagSets;
+	TArray<FString> ContainerFiles;
+	FChunkPluginStats Stats;
+};
+
+static int32 ChunkContainer(const FString& ContainerUTocPath, EOnDemandTocFlags TocFlags, const FChunkParams& Input, FOnDemandToc& OutOnDemandToc, FChunkOutput& Output)
 {
 	TRACE_CPUPROFILER_EVENT_SCOPE(ChunkPluginCommand::ChunkContainer);
 
-	check(ChunkWriter != nullptr);
+	check(Input.ChunkWriter.IsValid());
 
 	FIoStoreReader ContainerFileReader;
 	{
-		FIoStatus Status = ContainerFileReader.Initialize(*FPaths::ChangeExtension(ContainerUTocPath, TEXT("")), EncryptionKeys);
+		FIoStatus Status = ContainerFileReader.Initialize(*FPaths::ChangeExtension(ContainerUTocPath, TEXT("")), Input.EncryptionKeys);
 		if (!Status.IsOk())
 		{
 			UE_LOG(LogIoStoreOnDemand, Error, TEXT("Failed to open container '%s' for reading due to: %s"), *ContainerUTocPath, *Status.ToString());
@@ -293,7 +312,7 @@ static int32 ChunkContainer(const FString& ContainerUTocPath, const FChunkPlugin
 		ContainerEntry.ContainerFlags |= static_cast<uint8>(EIoContainerFlags::Encrypted);
 	}
 
-	if (Settings.PackageSets.Num() > 0 && OutFoundTagSets != nullptr)
+	if (EnumHasAnyFlags(TocFlags, EOnDemandTocFlags::InstallOnDemand) && Input.Settings.PackageSets.Num() > 0)
 	{
 		TArray<FOnDemandTocTagSet> TagSets;
 
@@ -323,7 +342,7 @@ static int32 ChunkContainer(const FString& ContainerUTocPath, const FChunkPlugin
 				return -1;
 			}
 
-			for (const TPair<FString, TArray<FString>>& Pair : Settings.PackageSets)
+			for (const TPair<FString, TArray<FString>>& Pair : Input.Settings.PackageSets)
 			{
 				const FString& Tag = Pair.Key;
 
@@ -343,7 +362,7 @@ static int32 ChunkContainer(const FString& ContainerUTocPath, const FChunkPlugin
 
 				if (TagSetPackageList.PackageIndicies.Num() > 0)
 				{
-					OutFoundTagSets->FindOrAdd(Tag).Add(MoveTemp(TagSetPackageList));
+					Output.FoundTagSets.FindOrAdd(Tag).Add(MoveTemp(TagSetPackageList));
 				}
 			}
 		}
@@ -423,7 +442,7 @@ static int32 ChunkContainer(const FString& ContainerUTocPath, const FChunkPlugin
 		}
 
 		// Additional validation to make sure that the chunks can be decrypted and decoded
-		if (bValidateChunks && ShouldValidateChunk(CompressionFormat, RawChunkSize))
+		if (Input.bValidateChunks && ShouldValidateChunk(CompressionFormat, RawChunkSize))
 		{
 			TRACE_CPUPROFILER_EVENT_SCOPE(ValidateChunk);
 
@@ -435,7 +454,7 @@ static int32 ChunkContainer(const FString& ContainerUTocPath, const FChunkPlugin
 
 			if (bIsContainerEncrypted)
 			{
-				const FAES::FAESKey* EncryptionKey = EncryptionKeys.Find(ContainerFileReader.GetEncryptionKeyGuid());
+				const FAES::FAESKey* EncryptionKey = Input.EncryptionKeys.Find(ContainerFileReader.GetEncryptionKeyGuid());
 				if (EncryptionKey == nullptr)
 				{
 					UE_LOG(LogIoStoreOnDemand, Error, TEXT("Failed to find encryption key '%s', Container='%s"), *LexToString(ContainerFileReader.GetEncryptionKeyGuid()), *ContainerUTocPath);
@@ -453,7 +472,7 @@ static int32 ChunkContainer(const FString& ContainerUTocPath, const FChunkPlugin
 			}
 		}
 
-		FIoStatus WriteStatus = ChunkWriter->WriteChunk(ChunksRelativeFolder, ReadResult.IoBuffer, ChunkHash);
+		FIoStatus WriteStatus = Input.ChunkWriter->WriteChunk(Input.ChunksRelativeFolder, ReadResult.IoBuffer, ChunkHash);
 		if (WriteStatus.IsOk() == false)
 		{
 			UE_LOG(LogIoStoreOnDemand, Error, TEXT("%s"), *WriteStatus.ToString());
@@ -468,10 +487,10 @@ static int32 ChunkContainer(const FString& ContainerUTocPath, const FChunkPlugin
 		TocEntry.BlockOffset = BlockOffset;
 		TocEntry.BlockCount = BlockCount;
 
-		OutStats.TotalIADChunksSizeKb += (double)EncodedChunkSize / 1024.0;
+		Output.Stats.TotalIADChunksSizeKb += (double)EncodedChunkSize / 1024.0;
 	}
 
-	if (bIgnoreContainerHeader)
+	if (Input.bIgnoreContainerHeader)
 	{
 		UE_LOG(LogIoStoreOnDemand, Display, TEXT("Ignoring container header for '%s'"), *ContainerUTocPath)
 	}
@@ -496,8 +515,8 @@ static int32 ChunkContainer(const FString& ContainerUTocPath, const FChunkPlugin
 		}
 	}
 
-	OutContainerFiles.Add(ContainerUTocPath);
-	ContainerFileReader.GetContainerFilePaths(OutContainerFiles);
+	Output.ContainerFiles.Add(ContainerUTocPath);						// .utoc
+	ContainerFileReader.GetContainerFilePaths(Output.ContainerFiles);	// .ucas
 
 	return 0;
 }
@@ -611,21 +630,22 @@ static int32 ChunkPluginCommandEntry(const FContext& Context)
 	UE_LOG(LogIoStoreOnDemand, Display, TEXT("\tValidateChunks: %s"), bValidateChunks ? TEXT("true") : TEXT("false"));
 	UE_LOG(LogIoStoreOnDemand, Display, TEXT("\tTOC flags: '%s'"), *LexToString(TocFlags));
 
+	FChunkParams Input;
+	Input.ChunksRelativeFolder = ChunksRelativeFolder;
+	Input.bIgnoreContainerHeader = bIgnoreContainerHeader;
+	Input.bValidateChunks = bValidateChunks;
+
 	IFileManager& FileMgr = IFileManager::Get();
 
-	TUniquePtr<FChunkWriterInterface> ChunkWriter;
-	bool bUseS3 = false;
 	if (S3Params.IsValid())
 	{
-		ChunkWriter = MakeUnique<FS3ChunkWriter>(S3Params);
+		Input.ChunkWriter = MakeUnique<FS3ChunkWriter>(S3Params);
 	}
 	else
 	{
-		ChunkWriter = MakeUnique<FDiskChunkWriter>(OutputFolder);
+		Input.ChunkWriter = MakeUnique<FDiskChunkWriter>(OutputFolder);
 	}
-	check(ChunkWriter != nullptr)
 
-	FChunkPluginSettings Settings;
 	if (!SettingsFile.IsEmpty())
 	{
 		UE_LOG(LogIoStoreOnDemand, Display, TEXT("Loading settings file '%s'..."), *SettingsFile);
@@ -689,20 +709,18 @@ static int32 ChunkPluginCommandEntry(const FContext& Context)
 				UE_LOGF(LogIoStoreOnDemand, Display, "%hs", *Output);
 				UE_LOGF(LogIoStoreOnDemand, Display, "TEST: After Large Message");
 #endif
-				Settings.PackageSets.Add(FString(Pair.Key), MoveTemp(Packages));
+				Input.Settings.PackageSets.Add(FString(Pair.Key), MoveTemp(Packages));
 				UE_LOGF(LogIoStoreOnDemand, Display, "Processed %d package set entries", Packages.Num());
 			}
 		}
 	}
 
-	FChunkPluginStats Stats;
-
-	TMap<FGuid, FAES::FAESKey> EncryptionKeys;
-
-	FKeyChain KeyChain = Common::LoadCryptoKeys(Context);
-	for (const TPair<FGuid, FNamedAESKey>& KeyPair : KeyChain.GetEncryptionKeys())
 	{
-		EncryptionKeys.Add(KeyPair.Key, KeyPair.Value.Key);
+		FKeyChain KeyChain = Common::LoadCryptoKeys(Context);
+		for (const TPair<FGuid, FNamedAESKey>& KeyPair : KeyChain.GetEncryptionKeys())
+		{
+			Input.EncryptionKeys.Add(KeyPair.Key, KeyPair.Value.Key);
+		}
 	}
 
 	if (FileMgr.DirectoryExists(*ContainerFolder) == false)
@@ -727,16 +745,14 @@ static int32 ChunkPluginCommandEntry(const FContext& Context)
 	OnDemandToc.Header.ChunksDirectory = TEXT(""); // TODO?
 	OnDemandToc.Containers.Reserve(ContainerFilenames.Num());
 
-	TMap<FString, TArray<FOnDemandTocTagSetPackageList>> FoundTagSets;
-	FoundTagSets.Reserve(Settings.PackageSets.Num());
-
 	FOnDemandToc StreamingOnDemandToc;
 	StreamingOnDemandToc.Header.Flags = static_cast<uint32>(EOnDemandTocFlags::StreamOnDemand);
 	StreamingOnDemandToc.Header.HostGroupName = HostGroupName;
 	StreamingOnDemandToc.Header.ChunksDirectory = TEXT(""); // TODO?
 	StreamingOnDemandToc.Containers.Reserve(ContainerFilenames.Num());
 
-	TArray<FString> ContainerFiles;
+	FChunkOutput Output;
+	Output.FoundTagSets.Reserve(Input.Settings.PackageSets.Num());
 
 	for (const FString& Filename : ContainerFilenames)
 	{
@@ -748,7 +764,7 @@ static int32 ChunkPluginCommandEntry(const FContext& Context)
 		if (bSupportStreamingContainers && bIsStreamOnDemandContainer)
 		{
 			// Attempt to add to IAS output container
-			int32 Result = ChunkContainer(FullPath, Settings, EncryptionKeys, ChunkWriter.Get(), ChunksRelativeFolder, bIgnoreContainerHeader, bValidateChunks, StreamingOnDemandToc, nullptr /*OutFoundTagSets*/, ContainerFiles, Stats);
+			int32 Result = ChunkContainer(FullPath, EOnDemandTocFlags::StreamOnDemand, Input, StreamingOnDemandToc, Output);
 			if (Result < 0)
 			{
 				return Result;
@@ -757,7 +773,7 @@ static int32 ChunkPluginCommandEntry(const FContext& Context)
 		else if(bSupportInstalledContainers && !bIsStreamOnDemandContainer)
 		{
 			// Attempt to add to IAD output container
-			int32 Result = ChunkContainer(FullPath, Settings, EncryptionKeys, ChunkWriter.Get(), ChunksRelativeFolder, bIgnoreContainerHeader, bValidateChunks, OnDemandToc, &FoundTagSets, ContainerFiles, Stats);
+			int32 Result = ChunkContainer(FullPath, EOnDemandTocFlags::InstallOnDemand, Input, OnDemandToc, Output);
 			if (Result < 0)
 			{
 				return Result;
@@ -765,10 +781,12 @@ static int32 ChunkPluginCommandEntry(const FContext& Context)
 		}
 	}
 
+	check(bSupportInstalledContainers || Output.FoundTagSets.IsEmpty());
+
 	if (bSupportInstalledContainers)
 	{
-		OnDemandToc.TagSets.Reserve(FoundTagSets.Num());
-		for (TPair<FString, TArray<FOnDemandTocTagSetPackageList>>& Pair : FoundTagSets)
+		OnDemandToc.TagSets.Reserve(Output.FoundTagSets.Num());
+		for (TPair<FString, TArray<FOnDemandTocTagSetPackageList>>& Pair : Output.FoundTagSets)
 		{
 			FOnDemandTocTagSet& TocTagSet = OnDemandToc.TagSets.Emplace_GetRef();
 			TocTagSet.Tag = Pair.Key;
@@ -782,7 +800,7 @@ static int32 ChunkPluginCommandEntry(const FContext& Context)
 
 		UE_LOG(LogIoStoreOnDemand, Display, TEXT("Deleting chunked  containers..."));
 
-		for (const FString& Path : ContainerFiles)
+		for (const FString& Path : Output.ContainerFiles)
 		{
 			if (FileMgr.FileExists(*Path))
 			{
@@ -836,7 +854,7 @@ static int32 ChunkPluginCommandEntry(const FContext& Context)
 
 			FIoBuffer Buffer(FIoBuffer::Clone, Chunk);
 
-			FIoStatus WriteStatus = ChunkWriter->WriteChunk(ChunksRelativeFolder, Buffer, ChunkHash);
+			FIoStatus WriteStatus = Input.ChunkWriter->WriteChunk(ChunksRelativeFolder, Buffer, ChunkHash);
 
 			if (WriteStatus.IsOk() == false)
 			{
@@ -893,7 +911,7 @@ static int32 ChunkPluginCommandEntry(const FContext& Context)
 
 		UE_LOG(LogIoStoreOnDemand, Display, TEXT("Replacing .utoc for chunked  containers..."));
 
-		for (const FString& SourcePath : ContainerFiles)
+		for (const FString& SourcePath : Output.ContainerFiles)
 		{
 			if (!SourcePath.EndsWith(TEXT(".utoc")))
 			{
@@ -922,7 +940,7 @@ static int32 ChunkPluginCommandEntry(const FContext& Context)
 	}
 
 
-	if (!ChunkWriter->Flush())
+	if (!Input.ChunkWriter->Flush())
 	{
 		UE_LOG(LogIoStoreOnDemand, Error, TEXT("Writer error: Failed to upload chunk(s)"));
 		return -1; 
@@ -937,7 +955,7 @@ static int32 ChunkPluginCommandEntry(const FContext& Context)
 			return -1;
 		}
 		TSharedRef<TJsonWriter<UTF8CHAR>> JsonWriter = TJsonWriterFactory<UTF8CHAR>::Create(FileWriter.Get());
-		Stats.ToJson(JsonWriter, false);
+		Output.Stats.ToJson(JsonWriter, false);
 	}
 
 	return 0;
