[NetPhysics] Add missing CHAOS_API on recent code change.

#Synced-CL 49840788

diff --git a/Engine/Source/Runtime/Experimental/Chaos/Public/Chaos/Framework/PhysicsSolverBase.h b/Engine/Source/Runtime/Experimental/Chaos/Public/Chaos/Framework/PhysicsSolverBase.h
--- a/Engine/Source/Runtime/Experimental/Chaos/Public/Chaos/Framework/PhysicsSolverBase.h
+++ b/Engine/Source/Runtime/Experimental/Chaos/Public/Chaos/Framework/PhysicsSolverBase.h
@@ -867,7 +867,7 @@ namespace Chaos
 		}
 
 		/** Return the interpolation lerp in case the resim is off */
-		static float NetworkPhysicsInterpolationLerp();
+		CHAOS_API static float NetworkPhysicsInterpolationLerp();
 
 		/** Get the standalone solver flag */
 		bool IsStandaloneSolver() const { return bIsStandaloneSolver;}
