Hopefully the last experiment, revert to the original logging but try with Log verbosity. This will prove that the problem is with the stdout pipe.

#rnx

#Synced-CL 49815329

diff --git a/Engine/Source/Runtime/Experimental/IoStore/OnDemand/Private/Tool/ChunkPlugin.cpp b/Engine/Source/Runtime/Experimental/IoStore/OnDemand/Private/Tool/ChunkPlugin.cpp
--- a/Engine/Source/Runtime/Experimental/IoStore/OnDemand/Private/Tool/ChunkPlugin.cpp
+++ b/Engine/Source/Runtime/Experimental/IoStore/OnDemand/Private/Tool/ChunkPlugin.cpp
@@ -691,12 +691,16 @@ static int32 ChunkPluginCommandEntry(const FContext& Context)
 					return -1;
 				}
 				
-#if 0
-				UE_LOG(LogIoStoreOnDemand, Display, TEXT("Found Package Set '%s'"), *Pair.Key);
+#if 1
+				UE_LOGF(LogIoStoreOnDemand, Display, "TEST: Before Large Message");
+
+				UE_LOG(LogIoStoreOnDemand, Log, TEXT("Found Package Set '%s'"), *Pair.Key);
 				for (const FString& Package : Packages)
 				{
-					UE_LOG(LogIoStoreOnDemand, Display, TEXT("->\t'%s'"), *Package);
+					UE_LOG(LogIoStoreOnDemand, Log, TEXT("->\t'%s'"), *Package);
 				}
+
+				UE_LOGF(LogIoStoreOnDemand, Display, "TEST: After Large Message");
 #else
 				TUtf8StringBuilder<0> Output;
 				Output << "Found Package Set '" << Pair.Key << "'";
