GeometryCollection 1 of 2 | Custom Depth/Stencil bug fix

Refresh custom renderer after setting depth/stencil to ensure all transforms are properly set in the ISM Pool

#Synced-CL 49824598

#Synced-CL 49824598

#Synced-CL 49824598

diff --git a/Engine/Source/Runtime/Experimental/GeometryCollectionEngine/Private/GeometryCollection/GeometryCollectionBlueprintLibrary.cpp b/Engine/Source/Runtime/Experimental/GeometryCollectionEngine/Private/GeometryCollection/GeometryCollectionBlueprintLibrary.cpp
--- a/Engine/Source/Runtime/Experimental/GeometryCollectionEngine/Private/GeometryCollection/GeometryCollectionBlueprintLibrary.cpp
+++ b/Engine/Source/Runtime/Experimental/GeometryCollectionEngine/Private/GeometryCollection/GeometryCollectionBlueprintLibrary.cpp
@@ -41,7 +41,8 @@ void UGeometryCollectionBlueprintLibrary::SetCustomDepthStencil(UGeometryCollect
 {
 	if (IGeometryCollectionCustomDataInterface* CustomRenderer = GetCustomRenderer(GeometryCollectionComponent))
 	{
-		CustomRenderer->SetRenderCustomDepthStencil(GeometryCollectionComponent, bEnable, CustomStencilValue);
+		CustomRenderer->SetRenderCustomDepthStencil(bEnable, CustomStencilValue);
+		GeometryCollectionComponent->RefreshCustomRenderer();
 	}
 }
 
diff --git a/Engine/Source/Runtime/Experimental/GeometryCollectionEngine/Public/GeometryCollection/GeometryCollectionExternalRenderInterface.h b/Engine/Source/Runtime/Experimental/GeometryCollectionEngine/Public/GeometryCollection/GeometryCollectionExternalRenderInterface.h
--- a/Engine/Source/Runtime/Experimental/GeometryCollectionEngine/Public/GeometryCollection/GeometryCollectionExternalRenderInterface.h
+++ b/Engine/Source/Runtime/Experimental/GeometryCollectionEngine/Public/GeometryCollection/GeometryCollectionExternalRenderInterface.h
@@ -93,5 +93,5 @@ public:
 	virtual void SetCustomInstanceData(FName CustomDataName, float CustomDataValue) = 0;
 
 	/** Enable or disable custom depth writing and, if enabled, the value written to the stencil buffer. The value will be set for all instances the renderer is spawning. */
-	virtual void SetRenderCustomDepthStencil(UGeometryCollectionComponent* GeometryCollectionComponent, bool bEnable, int32 CustomStencilValue) = 0;
+	virtual void SetRenderCustomDepthStencil(bool bEnable, int32 CustomStencilValue) = 0;
 };
