Dataflow : add nodes to load a texture directly or via a path

#Synced-CL 49874723

diff --git a/Engine/Source/Runtime/Experimental/Dataflow/Engine/Private/Dataflow/DataflowTextureAssetNodes.cpp b/Engine/Source/Runtime/Experimental/Dataflow/Engine/Private/Dataflow/DataflowTextureAssetNodes.cpp
--- a/Engine/Source/Runtime/Experimental/Dataflow/Engine/Private/Dataflow/DataflowTextureAssetNodes.cpp
+++ b/Engine/Source/Runtime/Experimental/Dataflow/Engine/Private/Dataflow/DataflowTextureAssetNodes.cpp
@@ -6,13 +6,20 @@
 
 #include UE_INLINE_GENERATED_CPP_BY_NAME(DataflowTextureAssetNodes)
 
+
+#define LOCTEXT_NAMESPACE "DataflowTextureAssetNodes"
+
 namespace UE::Dataflow
 {
 	void RegisterTextureAssetNodes()
 	{
 		DATAFLOW_NODE_REGISTER_CREATION_FACTORY(FDataflowTextureTerminalNode);
+		DATAFLOW_NODE_REGISTER_CREATION_FACTORY(FDataflowGetTextureAssetNode);
+		DATAFLOW_NODE_REGISTER_CREATION_FACTORY(FDataflowGetTextureAssetFromPathNode);
 		DATAFLOW_NODE_REGISTER_CREATION_FACTORY(FDataflowTextureToImageNode);
 		DATAFLOW_NODE_REGISTER_CREATION_FACTORY(FDataflowImageToTextureNode);
+
+		DATAFLOW_NODE_REGISTER_GETTER_FOR_ASSET(UTexture2D, FDataflowGetTextureAssetNode);
 	}
 }
 
@@ -76,6 +83,74 @@ void FDataflowTextureTerminalNode::SetAssetValue(TObjectPtr<UObject> Asset, UE::
 	}
 }
 
+////////////////////////////////////////////////////////////////////////////////////////////
+
+FDataflowGetTextureAssetNode::FDataflowGetTextureAssetNode(const UE::Dataflow::FNodeParameters& InParam, FGuid InGuid)
+	: FDataflowNode(InParam, InGuid)
+{
+	RegisterOutputConnection(&TextureAsset);
+}
+
+void FDataflowGetTextureAssetNode::Evaluate(UE::Dataflow::FContext& Context, const FDataflowOutput* Out) const
+{
+	if (Out->IsA(&TextureAsset))
+	{
+		SetValue(Context, TextureAsset, &TextureAsset);
+	}
+}
+
+bool FDataflowGetTextureAssetNode::SupportsAssetProperty(UObject* Asset) const
+{
+	return (Cast<UTexture2D>(Asset) != nullptr);
+}
+
+void FDataflowGetTextureAssetNode::SetAssetProperty(UObject* Asset)
+{
+	if (UTexture2D* InMaterialAsset = Cast<UTexture2D>(Asset))
+	{
+		TextureAsset = InMaterialAsset;
+	}
+}
+
+
+////////////////////////////////////////////////////////////////////////////////////////////
+
+FDataflowGetTextureAssetFromPathNode::FDataflowGetTextureAssetFromPathNode(const UE::Dataflow::FNodeParameters& InParam, FGuid InGuid)
+	: FDataflowNode(InParam, InGuid)
+{
+	RegisterInputConnection(&TexturePath);
+	RegisterOutputConnection(&TextureAsset);
+}
+
+void FDataflowGetTextureAssetFromPathNode::Evaluate(UE::Dataflow::FContext& Context, const FDataflowOutput* Out) const
+{
+	if (Out->IsA(&TextureAsset))
+	{
+		const FString& InPath = GetValue(Context, &TexturePath);
+		TObjectPtr<UTexture2D> OutTextureAsset = LoadObject<UTexture2D>(nullptr, InPath);
+		if (OutTextureAsset == nullptr)
+		{
+			Context.Error(FText::Format(LOCTEXT("TextureFromPath_NullAsset", "Failed to load asset {0}, returning null object"), FText::FromString(InPath)), this, Out);
+		}
+
+		SetValue(Context, OutTextureAsset, &TextureAsset);
+	}
+}
+
+bool FDataflowGetTextureAssetFromPathNode::SupportsAssetProperty(UObject* Asset) const
+{
+	return (Cast<UTexture2D>(Asset) != nullptr);
+}
+
+void FDataflowGetTextureAssetFromPathNode::SetAssetProperty(UObject* Asset)
+{
+	if (UTexture2D* InMaterialAsset = Cast<UTexture2D>(Asset))
+	{
+		TexturePath = InMaterialAsset->GetPathName();
+	}
+}
+
+////////////////////////////////////////////////////////////////////////////////////////////
 
 FDataflowTextureToImageNode::FDataflowTextureToImageNode(const UE::Dataflow::FNodeParameters& InParam, FGuid InGuid)
 	: FDataflowNode(InParam, InGuid)
@@ -127,6 +202,7 @@ void FDataflowTextureToImageNode::Evaluate(UE::Dataflow::FContext& Context, cons
 	}
 }
 
+////////////////////////////////////////////////////////////////////////////////////////////
 
 FDataflowImageToTextureNode::FDataflowImageToTextureNode(const UE::Dataflow::FNodeParameters& InParam, FGuid InGuid)
 	: FDataflowNode(InParam, InGuid)
@@ -161,3 +237,5 @@ void FDataflowImageToTextureNode::Evaluate(UE::Dataflow::FContext& Context, cons
 		return;
 	}
 }
+
+#undef LOCTEXT_NAMESPACE
diff --git a/Engine/Source/Runtime/Experimental/Dataflow/Engine/Public/Dataflow/DataflowTextureAssetNodes.h b/Engine/Source/Runtime/Experimental/Dataflow/Engine/Public/Dataflow/DataflowTextureAssetNodes.h
--- a/Engine/Source/Runtime/Experimental/Dataflow/Engine/Public/Dataflow/DataflowTextureAssetNodes.h
+++ b/Engine/Source/Runtime/Experimental/Dataflow/Engine/Public/Dataflow/DataflowTextureAssetNodes.h
@@ -32,6 +32,52 @@ public:
 	FDataflowTextureTerminalNode(const UE::Dataflow::FNodeParameters& InParam, FGuid InGuid = FGuid::NewGuid());
 };
 
+/*
+* Get a texture asset
+*/
+USTRUCT(meta = (DataflowTerminal))
+struct FDataflowGetTextureAssetNode : public FDataflowNode
+{
+	GENERATED_USTRUCT_BODY()
+	DATAFLOW_NODE_DEFINE_INTERNAL(FDataflowGetTextureAssetNode, "GetTextureAsset", "Image", "")
+
+private:
+	UPROPERTY(EditAnywhere, Category = Asset, meta = (DataflowOutput))
+	TObjectPtr<UTexture2D> TextureAsset = nullptr;
+
+	virtual void Evaluate(UE::Dataflow::FContext& Context, const FDataflowOutput* Out) const override;
+
+public:
+	FDataflowGetTextureAssetNode(const UE::Dataflow::FNodeParameters& InParam, FGuid InGuid = FGuid::NewGuid());
+
+	virtual bool SupportsAssetProperty(UObject* Asset) const override;
+	virtual void SetAssetProperty(UObject* Asset) override;
+};
+
+/*
+* Get a texture asset from a path
+*/
+USTRUCT(meta = (DataflowTerminal))
+struct FDataflowGetTextureAssetFromPathNode : public FDataflowNode
+{
+	GENERATED_USTRUCT_BODY()
+	DATAFLOW_NODE_DEFINE_INTERNAL(FDataflowGetTextureAssetFromPathNode, "GetTextureAssetFromPath", "Image", "Path String")
+
+private:
+	UPROPERTY(EditAnywhere, Category = Asset, meta = (DataflowInput))
+	FString TexturePath;
+
+	UPROPERTY(meta = (DataflowOutput))
+	TObjectPtr<UTexture2D> TextureAsset = nullptr;
+
+	virtual void Evaluate(UE::Dataflow::FContext& Context, const FDataflowOutput* Out) const override;
+
+public:
+	FDataflowGetTextureAssetFromPathNode(const UE::Dataflow::FNodeParameters& InParam, FGuid InGuid = FGuid::NewGuid());
+
+	virtual bool SupportsAssetProperty(UObject* Asset) const override;
+	virtual void SetAssetProperty(UObject* Asset) override;
+};
 
 /**
 * Import a texture asset as an image. Texture must have CPU availability.
