Changed on-demand container TOC format to support zero-copy serialization.

- changed the TOC to have a storage part and and a view part to support memory mapping
- added ResaveToc IasTool command to enable resaving TOC's to the latest format
- added ConvertToc IasTool command to support converting a TOC to Json or CSV
- changed the format to inline compression block size and hash if the chunk only has a single compresson block, this
makes the TOC about 20% smaller for IAD containers with a lot of package data chunks
- the runtime will be updated to use the new TOC in a follow-up change

#rnx

#Synced-CL 49823132

diff --git a/Engine/Source/Runtime/Experimental/IoStore/OnDemand/Private/Tool/ConvertToc.cpp b/Engine/Source/Runtime/Experimental/IoStore/OnDemand/Private/Tool/ConvertToc.cpp
new file mode 100644
--- /dev/null
+++ b/Engine/Source/Runtime/Experimental/IoStore/OnDemand/Private/Tool/ConvertToc.cpp
@@ -0,0 +1,299 @@
+// Copyright Epic Games, Inc. All Rights Reserved.
+
+#if defined(UE_WITH_IAS_TOOL)
+
+#include "Command.h"
+
+#include "HAL/FileManager.h"
+#include "IO/IoChunkId.h"
+#include "IO/IoContainerMeta.h"
+#include "IO/IoStoreOnDemand.h"
+#include "IO/OnDemandToc.h"
+#include "IO/PackageId.h"
+#include "IO/Serialization/OnDemandContainerToc.h"
+#include "Misc/FileHelper.h"
+#include "Misc/Paths.h"
+#include "Misc/StringBuilder.h"
+#include "Serialization/CompactBinarySerialization.h"
+#include "Serialization/CompactBinaryWriter.h"
+
+namespace UE::IoStore::Tool
+{
+
+////////////////////////////////////////////////////////////////////////////////
+static FIoChunkId GetPackageChunkId(const FIoChunkId& ChunkId, FPackageId& OutPackageId)
+{
+	FMemory::Memcpy(&OutPackageId, &ChunkId, sizeof(FPackageId));
+
+	if (ChunkId.GetChunkType() == EIoChunkType::ExportBundleData)
+	{
+		return ChunkId;
+	}
+
+	return CreatePackageDataChunkId(OutPackageId);
+}
+
+////////////////////////////////////////////////////////////////////////////////
+static FIoChunkId GetPackageChunkId(const FIoChunkId& ChunkId)
+{
+	FPackageId Unused;
+	return GetPackageChunkId(ChunkId, Unused);
+}
+
+////////////////////////////////////////////////////////////////////////////////
+class FContainerMeta
+{
+public:
+	FUtf8StringView			GetFilename(
+								const FIoChunkId& ChunkId,
+								FUtf8StringBuilderBase& OutFilename,
+								FUtf8StringView& OutContainerName);
+	FUtf8StringView			GetFilename(
+								const FIoChunkId& ChunkId,
+								FUtf8StringBuilderBase& OutFilename);
+	static FContainerMeta	Load(const FString& Directory);
+
+private:
+	TArray<FIoContainerMetaReader> Readers;
+};
+
+////////////////////////////////////////////////////////////////////////////////
+FUtf8StringView FContainerMeta::GetFilename(
+	const FIoChunkId& ChunkId,
+	FUtf8StringBuilderBase& OutFilename,
+	FUtf8StringView& OutContainerName)
+{
+	OutFilename.Reset();
+	for (FIoContainerMetaReader& Reader : Readers)
+	{
+		FUtf8StringView Filename = Reader.GetFilename(ChunkId, OutFilename, OutContainerName);
+		if (!Filename.IsEmpty())
+		{
+			return Filename;
+		}
+	}
+
+	return FUtf8StringView();
+}
+
+FUtf8StringView FContainerMeta::GetFilename(
+	const FIoChunkId& ChunkId,
+	FUtf8StringBuilderBase& OutFilename)
+{
+	FUtf8StringView Unused;
+	return GetFilename(ChunkId, OutFilename, Unused);
+}
+
+FContainerMeta FContainerMeta::Load(const FString& Directory)
+{
+	FContainerMeta	OutMeta;
+	TArray<FString> MetaFiles;
+
+	IFileManager& Ifm = IFileManager::Get();
+	Ifm.FindFiles(MetaFiles, *Directory, TEXT(".umeta"));
+
+	for (const FString& Filename : MetaFiles)
+	{
+		const FString	FullPath = Directory / Filename;
+		const bool		bMemoryMap = false;
+
+		if (TIoStatusOr<FIoContainerMetaReader> Status = FIoContainerMetaReader::Load(FullPath, bMemoryMap); Status.IsOk())
+		{
+			OutMeta.Readers.Add(Status.ConsumeValueOrDie());
+		}
+		else
+		{
+			UE_LOG(LogIoStoreOnDemand, Warning, TEXT("Failed to load container meta '%s'"), *FullPath);
+		}
+	}
+
+	return OutMeta;
+}
+
+////////////////////////////////////////////////////////////////////////////////
+static int32 ConvertTocCommandEntry(const FContext& Context)
+{
+	using namespace UE::IoStore::Serialization;
+	using namespace UE::IoStore::Serialization::V1;
+
+	GPrintLogCategory = GPrintLogVerbosity = false;
+	GPrintLogTimes = ELogTimes::None;
+
+	FString FilePath = FString(Context.Get<FStringView>(TEXT("-Path")));
+	if (FilePath.IsEmpty())
+	{
+		UE_LOG(LogIoStoreOnDemand, Error, TEXT("Invalid filename"));
+		return -1;
+	}
+
+	FString OutFilePath = FString(Context.Get<FStringView>(TEXT("-Out")));
+	if (OutFilePath.IsEmpty())
+	{
+		UE_LOG(LogIoStoreOnDemand, Error, TEXT("Invalid JSON output filename"));
+		return -1;
+	}
+
+	const FString Ext = FPaths::GetExtension(OutFilePath).ToLower();
+	if (Ext != TEXT("json") && Ext != TEXT("csv"))
+	{
+		UE_LOG(LogIoStoreOnDemand, Error, TEXT("Invalid output file extension"));
+		return -1;
+	}
+
+	IFileManager& Ifm = IFileManager::Get();
+
+	TIoStatusOr<FOnDemandTocReader> MaybeReader = FOnDemandTocReader::Read(FilePath);
+	if (!MaybeReader.IsOk())
+	{
+		UE_LOG(LogIoStoreOnDemand, Error, TEXT("Failed to read on-demand TOC '%s', reason: %s"),
+			*FilePath, *MaybeReader.Status().ToString());
+		return -1;
+	}
+
+	FOnDemandTocReader	Reader = MaybeReader.ConsumeValueOrDie();
+	FContainerMeta		ContainerMeta = FContainerMeta::Load(FPaths::GetPath(FilePath));
+
+	if (Ext == TEXT("json"))
+	{
+		FCbWriter Writer;
+		Writer.BeginObject();
+		{
+			Writer.BeginObject(UTF8TEXTVIEW("Header"));
+			Writer.AddString(UTF8TEXTVIEW("CreationTime"), Reader.CreationTime().ToString());
+			Writer.AddString(UTF8TEXTVIEW("BuildVersion"), Reader.BuildVersion());
+			Writer.AddString(UTF8TEXTVIEW("TargetPlatform"), Reader.TargetPlatform());
+			Writer.AddString(UTF8TEXTVIEW("ChunksDirectory"), Reader.ChunksDirectory());
+			Writer.EndObject();
+
+			Writer.BeginArray(UTF8TEXTVIEW("Containers"));
+			for (const FOnDemandContainerEntry& ContainerEntry : Reader.Containers())
+			{
+				FOnDemandTocStorage Storage;
+				TIoStatusOr<FOnDemandContainerTocView> MaybeContainerView =
+					Reader.ReadContainer(ContainerEntry, Storage, EOnDemandTocReaderOptions::None);
+
+				if (MaybeContainerView.IsOk() == false)
+				{
+					UE_LOG(LogIoStoreOnDemand, Error, TEXT("Failed to read on-demand container TOC '%s', reason: %s"),
+						*FilePath, *MaybeContainerView.Status().ToString());
+					return -1;
+				}
+
+				FOnDemandContainerTocView ContainerView = MaybeContainerView.ConsumeValueOrDie();
+
+				Writer.BeginObject();
+				{
+					Writer.AddString(UTF8TEXTVIEW("ContainerName"), ContainerView.Header.ContainerName());
+					Writer.AddString(UTF8TEXTVIEW("ContainerId"), LexToString(ContainerView.Header.ContainerId()));
+					Writer.AddInteger(UTF8TEXTVIEW("ContainerFlags"), uint32(ContainerView.Header.ContainerFlags()));
+					Writer.AddString(UTF8TEXTVIEW("EncryptionKeyGuid"), ContainerView.Header.EncryptionKeyGuid().ToString());
+					Writer.AddInteger(UTF8TEXTVIEW("BlockSize"), ContainerView.Header.BlockSize());
+
+					TUtf8StringBuilder<512> Sb;
+					Writer.BeginArray(UTF8TEXTVIEW("ChunkEntries"));
+					for (int32 Idx = 0; const FIoChunkId& ChunkId : ContainerView.ChunkIds)
+					{
+						FUtf8StringView Filename = ContainerMeta.GetFilename(GetPackageChunkId(ChunkId), Sb);
+						const FOnDemandChunkEntry& ChunkEntry = ContainerView.ChunkEntries[Idx++];
+						Writer.BeginObject();
+						{
+							Writer.AddString(UTF8TEXTVIEW("Filename"), Filename);
+							Writer.AddString(UTF8TEXTVIEW("ChunkId"), LexToString(ChunkId));
+							Writer.AddString(UTF8TEXTVIEW("Hash"), LexToString(ChunkEntry.Hash));
+							Writer.AddInteger(UTF8TEXTVIEW("RawSize"), ChunkEntry.RawSize);
+							Writer.AddInteger(UTF8TEXTVIEW("EncodedSize"), ChunkEntry.EncodedSize);
+							Writer.AddInteger(UTF8TEXTVIEW("BlockOffset"), ChunkEntry.BlockInfo.Offset());
+							Writer.AddInteger(UTF8TEXTVIEW("BlockCount"), ChunkEntry.BlockInfo.Count());
+							if (ChunkEntry.BlockInfo.Count() == 1)
+							{
+								Writer.AddInteger(UTF8TEXTVIEW("BlockSize"), ChunkEntry.BlockInfo.Size());
+								Writer.AddInteger(UTF8TEXTVIEW("BlockHash"), ChunkEntry.BlockInfo.Hash());
+							}
+						}
+						Writer.EndObject();
+					}
+					Writer.EndArray();
+				}
+				Writer.EndObject();
+			}
+			Writer.EndArray();
+		}
+		Writer.EndObject();
+
+		TStringBuilder<4096> Json;
+		CompactBinaryToJson(Writer.Save().AsObject(), Json);
+
+		if (!FFileHelper::SaveStringToFile(Json.ToView(), *OutFilePath, FFileHelper::EEncodingOptions::ForceUTF8))
+		{
+			UE_LOG(LogIoStoreOnDemand, Error, TEXT("Failed to save JSON '%s'"), *OutFilePath);
+			return -1;
+		}
+	}
+	else
+	{
+		TUniquePtr<FArchive> Ar(Ifm.CreateFileWriter(*OutFilePath));
+		if (!Ar.IsValid())
+		{
+			UE_LOG(LogIoStoreOnDemand, Error, TEXT("Failed to open '%s' for writing"), *OutFilePath);
+			return -1;
+		}
+
+		Ar->Logf(TEXT("Platform, ContainerName, ChunkId, PackageId, Hash, Size, CompressedSize, Filename"));
+		for (const FOnDemandContainerEntry& ContainerEntry : Reader.Containers())
+		{
+			FOnDemandTocStorage Storage;
+			TIoStatusOr<FOnDemandContainerTocView> MaybeContainerView =
+				Reader.ReadContainer(ContainerEntry, Storage, EOnDemandTocReaderOptions::None);
+
+			if (MaybeContainerView.IsOk() == false)
+			{
+				UE_LOG(LogIoStoreOnDemand, Error, TEXT("Failed to read on-demand container TOC '%s', reason: %s"),
+					*FilePath, *MaybeContainerView.Status().ToString());
+				return -1;
+			}
+
+			FOnDemandContainerTocView ContainerView = MaybeContainerView.ConsumeValueOrDie();
+
+			TUtf8StringBuilder<512> Sb;
+			for (int32 Idx = 0; const FIoChunkId& ChunkId : ContainerView.ChunkIds)
+			{
+				FPackageId					PackageId;
+				FUtf8StringView				Filename = ContainerMeta.GetFilename(GetPackageChunkId(ChunkId, PackageId), Sb);
+				const FOnDemandChunkEntry&	ChunkEntry = ContainerView.ChunkEntries[Idx++];
+				FString						ContainerFullName = FString(ContainerView.Header.ContainerName());
+
+				FString PakChunk, Platform;
+				ContainerFullName.Split(TEXT("-"), &PakChunk, &Platform, ESearchCase::CaseSensitive, ESearchDir::FromEnd);
+
+				Ar->Logf(TEXT("%s, %s, %s, %s, %s, %u, %u, %s"),
+					*Platform,
+					*PakChunk,
+					*LexToString(ChunkId),
+					*LexToString(PackageId),
+					*LexToString(ChunkEntry.Hash),
+					ChunkEntry.RawSize,
+					ChunkEntry.EncodedSize,
+					*FString(Filename));
+			}
+		}
+	}
+
+	UE_LOG(LogIoStoreOnDemand, Display, TEXT("Saved '%s'"), *OutFilePath);
+
+	return 0;
+}
+
+////////////////////////////////////////////////////////////////////////////////
+static FCommand ConvertTocCommand(
+	ConvertTocCommandEntry,
+	TEXT("ConvertToc"),
+	TEXT("Converts an on-demand TOC file to JSON or CSV"),
+	{
+		TArgument<FStringView>(TEXT("-Path"), TEXT("Path to a .uondemandtoc file")),
+		TArgument<FStringView>(TEXT("-Out"), TEXT("Output file path"))
+	}
+);
+
+} // namespace UE::IoStore::Tool
+
+#endif // UE_WITH_IAS_TOOL
diff --git a/Engine/Source/Runtime/Experimental/IoStore/OnDemand/Private/Tool/ResaveToc.cpp b/Engine/Source/Runtime/Experimental/IoStore/OnDemand/Private/Tool/ResaveToc.cpp
new file mode 100644
--- /dev/null
+++ b/Engine/Source/Runtime/Experimental/IoStore/OnDemand/Private/Tool/ResaveToc.cpp
@@ -0,0 +1,95 @@
+// Copyright Epic Games, Inc. All Rights Reserved.
+
+#if defined(UE_WITH_IAS_TOOL)
+
+#include "Command.h"
+
+#include "HAL/FileManager.h"
+#include "IO/IoStoreOnDemand.h"
+#include "IO/OnDemandToc.h"
+#include "IO/Serialization/OnDemandContainerToc.h"
+#include "Misc/Paths.h"
+
+namespace UE::IoStore::Tool
+{
+
+////////////////////////////////////////////////////////////////////////////////
+static int32 ResaveTocEntry(const FContext& Context)
+{
+	GPrintLogCategory = GPrintLogVerbosity = false;
+	GPrintLogTimes = ELogTimes::None;
+
+	FString FilePath = FString(Context.Get<FStringView>(TEXT("-Path")));
+	if (FilePath.IsEmpty())
+	{
+		UE_LOG(LogIoStoreOnDemand, Error, TEXT("Invalid filename"));
+		return -1;
+	}
+
+	IFileManager& Ifm = IFileManager::Get();
+	
+	int64 LegacyTocSize = 0;
+	FOnDemandToc LegacyToc;
+	{
+		TUniquePtr<FArchive> Ar(Ifm.CreateFileReader(*FilePath));
+		if (Ar.IsValid() == false)
+		{
+			UE_LOG(LogIoStoreOnDemand, Error, TEXT("Failed to open '%s'"), *FilePath);
+			return -1;
+		}
+		*Ar << LegacyToc;
+		LegacyTocSize = int64(Ar->Tell());
+
+		if (Ar->IsError())
+		{
+			UE_LOG(LogIoStoreOnDemand, Error, TEXT("Failed to serialize legacy TOC '%s'"), *FilePath);
+			return -1;
+		}
+	}
+
+	UE_LOG(LogIoStoreOnDemand, Display, TEXT("Loaded legacy TOC '%s' (%.2lf KB)"),
+		*FilePath, double(LegacyTocSize) / 1024.0);
+
+	FString TmpPath = FPaths::ChangeExtension(FilePath, TEXT("tmp"));
+	if (Ifm.Move(*TmpPath, *FilePath) == false)
+	{
+		UE_LOG(LogIoStoreOnDemand, Error, TEXT("Failed to move file '%s' -> '%s'"), *FilePath, *TmpPath);
+		return -1;
+	}
+
+	UE::IoStore::Serialization::FOnDemandTocWriter Writer;
+	if (TIoStatusOr<uint64> Status = Writer.Write(LegacyToc, FilePath); !Status.IsOk())
+	{
+		Ifm.Delete(*FilePath, false, true, true);
+		if (Ifm.Move(*FilePath, *TmpPath) == false)
+		{
+			UE_LOG(LogIoStoreOnDemand, Error, TEXT("Failed to move file '%s' -> '%s'"), *TmpPath, *FilePath);
+		}
+
+		UE_LOG(LogIoStoreOnDemand, Error, TEXT("Failed to resave file '%s', reason: %s"),
+			*FilePath, *Status.Status().ToString());
+		return -1;
+	}
+	else
+	{
+		Ifm.Delete(*TmpPath, false, true, true);
+		UE_LOG(LogIoStoreOnDemand, Display, TEXT("Successfully resaved '%s' (%.2lf KB)"),
+			*FilePath, double(Status.ConsumeValueOrDie()) / 1024.0);
+	}
+
+	return 0;
+}
+
+////////////////////////////////////////////////////////////////////////////////
+static FCommand ResaveTocCommand(
+	ResaveTocEntry,
+	TEXT("ResaveToc"),
+	TEXT("Resaves a serialized on-demand TOC file to the latest format"),
+	{
+		TArgument<FStringView>(TEXT("-Path"), TEXT("Path to a .uondemandtoc file")),
+	}
+);
+
+} // namespace UE::IoStore::Tool
+
+#endif // UE_WITH_IAS_TOOL
diff --git a/Engine/Source/Runtime/Experimental/IoStore/OnDemandCore/Internal/IO/Serialization/OnDemandContainerToc.h b/Engine/Source/Runtime/Experimental/IoStore/OnDemandCore/Internal/IO/Serialization/OnDemandContainerToc.h
new file mode 100644
--- /dev/null
+++ b/Engine/Source/Runtime/Experimental/IoStore/OnDemandCore/Internal/IO/Serialization/OnDemandContainerToc.h
@@ -0,0 +1,336 @@
+// Copyright Epic Games, Inc. All Rights Reserved.
+
+#pragma once
+
+#include "Containers/ArrayView.h"
+#include "Containers/StringFwd.h"
+#include "IO/IoBuffer.h"
+#include "IO/IoContainerId.h"
+#include "IO/IoHash.h"
+#include "IO/IoStatus.h"
+#include "Misc/AES.h"
+#include "Misc/EnumClassFlags.h"
+#include "Misc/DateTime.h"
+#include "Misc/Guid.h"
+#include "Misc/TVariant.h"
+
+#define UE_API IOSTOREONDEMANDCORE_API
+
+class FArchive;
+class IMappedFileHandle;
+class IMappedFileRegion;
+class FIoChunkId;
+using FIoBlockHash = uint32;
+ 
+enum class EIoContainerFlags : uint8;
+
+namespace UE::IoStore { struct FOnDemandToc; }
+
+namespace UE::IoStore::Serialization
+{
+
+///////////////////////////////////////////////////////////////////////////////
+struct FOnDemandTocSignature
+{
+	static constexpr inline ANSICHAR ExpectedSignature[] = "UE ON-DEMAND TOC";
+
+	FOnDemandTocSignature()
+	{
+		FMemory::Memcpy(Signature, ExpectedSignature, sizeof(Signature));
+	}
+
+	bool IsValid() const
+	{
+		return FMemory::Memcmp(Signature, ExpectedSignature, sizeof(Signature)) == 0;
+	}
+
+	ANSICHAR Signature[16];
+};
+static_assert(sizeof(FOnDemandTocSignature) == 16);
+
+///////////////////////////////////////////////////////////////////////////////
+enum class EOnDemandTocMajorVersion : uint16
+{
+	Invalid			= 0,
+	One				= 1,
+
+	LatestPlusOne,
+	Latest			= (LatestPlusOne - 1)
+};
+
+///////////////////////////////////////////////////////////////////////////////
+enum class EOnDemandTocMinorVersion : uint16
+{
+	Invalid			= 0,
+	MemoryMapped	= 1,
+
+	LatestPlusOne,
+	Latest			= (LatestPlusOne - 1)
+};
+
+///////////////////////////////////////////////////////////////////////////////
+struct alignas(4) FOnDemandTocVersion
+{
+	bool IsValid() const { return Major > 0 && Major != ~uint16(0) && Minor > 0 && Minor != ~uint16(0); }
+
+	uint16 Major = uint16(EOnDemandTocMajorVersion::Latest);
+	uint16 Minor = uint16(EOnDemandTocMinorVersion::Latest);
+};
+static_assert(sizeof(FOnDemandTocVersion) == 4);
+
+////////////////////////////////////////////////////////////////////////////////
+struct FOnDemandTocFooter
+{
+	FOnDemandTocSignature Signature;
+};
+static_assert(sizeof(FOnDemandTocFooter) == 16);
+
+////////////////////////////////////////////////////////////////////////////////
+class FOnDemandTocStorage
+{
+	struct FMappedFileStorage
+	{
+		TUniquePtr<IMappedFileHandle> MappedFile;
+		TUniquePtr<IMappedFileRegion> MappedRegion;
+	};
+
+	struct FMemoryStorage
+	{
+		FIoBuffer Buffer;
+	};
+
+	using FStorage = TVariant<FEmptyVariantState, FMemoryStorage, FMappedFileStorage>;
+
+public:
+	UE_API FOnDemandTocStorage();
+	UE_API FOnDemandTocStorage(FIoBuffer HeaderData, FIoBuffer ContainerData);
+	UE_API FOnDemandTocStorage(FIoBuffer HeaderData, TUniquePtr<IMappedFileHandle>&& File, TUniquePtr<IMappedFileRegion>&& Region);
+	UE_API FOnDemandTocStorage(FOnDemandTocStorage&& Other);
+	UE_API ~FOnDemandTocStorage();
+
+	UE_API FOnDemandTocStorage& operator=(FOnDemandTocStorage&&);
+
+	UE_API FMemoryView GetView() const;
+	bool IsEmpty() const { return HeaderData.GetSize() == 0; }
+
+private:
+	FOnDemandTocStorage(const FOnDemandTocStorage&) = delete;
+	FOnDemandTocStorage& operator=(const FOnDemandTocStorage&) = delete;
+
+	FIoBuffer	HeaderData;
+	FStorage	ContainerData;
+};
+
+////////////////////////////////////////////////////////////////////////////////
+class FOnDemandTocWriter
+{
+public:
+	UE_API FOnDemandTocWriter();
+	UE_API FOnDemandTocWriter(FOnDemandTocWriter&&);
+	UE_API ~FOnDemandTocWriter();
+
+	UE_API FOnDemandTocWriter& operator=(FOnDemandTocWriter&&);
+
+	UE_API TIoStatusOr<uint64> Write(const UE::IoStore::FOnDemandToc& Toc, FArchive& Ar);
+	UE_API TIoStatusOr<uint64> Write(const UE::IoStore::FOnDemandToc& Toc, const FString& Filename);
+
+private:
+	FOnDemandTocWriter(const FOnDemandTocWriter&) = delete;
+	FOnDemandTocWriter& operator=(const FOnDemandTocWriter&) = delete;
+
+	class FImpl;
+	TUniquePtr<FImpl> Impl;
+};
+
+////////////////////////////////////////////////////////////////////////////////
+namespace V1
+{
+
+struct FOnDemandStringEntry
+{
+	uint32 Offset = 0;
+	uint32 Len = 0;
+};
+static_assert(sizeof(FOnDemandStringEntry) == 8);
+
+struct FOnDemandTocHeader
+{
+	FOnDemandTocSignature	Signature;
+	FOnDemandTocVersion		Version;
+	uint32					Pad = 0;
+	int64					EpochTimestamp = 0;
+	FOnDemandStringEntry	BuildVersion;
+	FOnDemandStringEntry	TargetPlatform;
+	FOnDemandStringEntry	ChunksDirectory;
+	FOnDemandStringEntry	HostGroupName;
+	FOnDemandStringEntry	CompressionFormat;
+	uint32					StringTableLen = 0;
+	uint32					ContainerCount = 0;
+	uint8					Pad2[48] = {0};
+};
+static_assert(sizeof(FOnDemandTocHeader) == 128);
+
+struct FOnDemandContainerEntry
+{
+	FGuid					EncryptionKeyGuid;
+	FIoContainerId			ContainerId;
+	FOnDemandStringEntry	ContainerName;
+	uint32					ContainerHeaderSize = 0;
+	uint32					DataOffset = 0;
+	uint32					DataSize = 0;
+	uint32					ChunkCount = 0;
+	uint32					BlockCount = 0;
+	uint32					BlockSize = 0;
+	uint32					TagSetCount = 0;
+	uint32					TagSetIndicesCount = 0;
+	FIoHash					UTocHash = FIoHash::Zero;
+	uint32					ContainerFlags = 0; // EIoContainerFlags
+	int8					Pad[40] = {0};
+};
+static_assert(sizeof(FOnDemandContainerEntry) == 128);
+
+struct FOnDemandChunkBlockInfo
+{
+	FOnDemandChunkBlockInfo()
+	{
+		FMemory::Memzero(this, sizeof(FOnDemandChunkBlockInfo));
+	}
+
+	uint32 Offset()	const { return bHasOffset ? OffsetOrSize : ~uint32(0); }
+	uint32 Count()	const { return bHasOffset ? CountOrHash : OffsetOrSize > 0 ? 1 : 0; }
+	uint32 Size()	const { return bHasOffset ? 0 : OffsetOrSize; }
+	uint32 Hash()	const { return bHasOffset ? 0 : CountOrHash; }
+
+	static FOnDemandChunkBlockInfo FromOffsetAndCount(uint32 BlockOffset, uint32 BlockCount)
+	{
+		FOnDemandChunkBlockInfo Nfo;
+		Nfo.OffsetOrSize	= BlockOffset;
+		Nfo.bHasOffset		= 1;
+		Nfo.CountOrHash		= BlockCount;
+		return Nfo;
+	}
+
+	static FOnDemandChunkBlockInfo FromSizeAndHash(uint32 BlockSize, uint32 BlockHash)
+	{
+		FOnDemandChunkBlockInfo Nfo;
+		Nfo.OffsetOrSize	= BlockSize;
+		Nfo.bHasOffset		= 0;
+		Nfo.CountOrHash		= BlockHash;
+		return Nfo;
+	}
+
+private:
+	uint32 OffsetOrSize: 31;
+	uint32 bHasOffset: 1;
+	uint32 CountOrHash;
+};
+static_assert(sizeof(FOnDemandChunkBlockInfo) == 8);
+
+struct FOnDemandChunkEntry
+{
+	uint32	DiskSize() const { return Align(EncodedSize, FAES::AESBlockSize); }
+
+	FIoHash						Hash = FIoHash::Zero;
+	uint32						RawSize = 0;
+	uint32						EncodedSize = 0;
+	FOnDemandChunkBlockInfo		BlockInfo;
+};
+static_assert(sizeof(FOnDemandChunkEntry) == 36);
+
+struct FOnDemandTagSetEntry
+{
+	FOnDemandStringEntry	Tag;
+	uint32					Offset = 0;
+	uint32					Count = 0;
+};
+static_assert(sizeof(FOnDemandTagSetEntry) == 16);
+
+////////////////////////////////////////////////////////////////////////////////
+struct FOnDemandContainerTocHeaderView
+{
+	FOnDemandContainerTocHeaderView() = default;
+	FOnDemandContainerTocHeaderView(
+		const FOnDemandTocHeader* InHeader,
+		const UTF8CHAR* InStringTable,
+		const FOnDemandContainerEntry* InContainerEntry)
+			: Header(InHeader)
+			, StringTable(InStringTable)
+			, ContainerEntry(InContainerEntry)
+	{ }
+
+	FUtf8StringView			ContainerName() const { return GetString(ContainerEntry->ContainerName); }
+	FIoContainerId			ContainerId() const { return ContainerEntry->ContainerId; }
+	EIoContainerFlags		ContainerFlags() const;
+	const FGuid&			EncryptionKeyGuid() const { return ContainerEntry->EncryptionKeyGuid; }
+	uint32					BlockSize() const { return ContainerEntry->BlockSize; }
+	UE_API FUtf8StringView	GetString(FOnDemandStringEntry StringEntry) const;
+
+	bool					IsValid() const { return Header != nullptr && StringTable != nullptr && ContainerEntry != nullptr; }
+							operator bool() const { return IsValid(); }
+
+private:
+	const FOnDemandTocHeader*		Header = nullptr;
+	const UTF8CHAR*					StringTable = nullptr;
+	const FOnDemandContainerEntry*	ContainerEntry = nullptr;
+};
+
+///////////////////////////////////////////////////////////////////////////////
+struct FOnDemandContainerTocView
+{
+	FOnDemandContainerTocHeaderView			Header;
+	TConstArrayView<FIoChunkId>				ChunkIds;
+	TConstArrayView<FOnDemandChunkEntry>	ChunkEntries;
+	TConstArrayView<uint32>					BlockSizes;
+	TConstArrayView<FIoBlockHash>			BlockHashes;
+	TConstArrayView<FOnDemandTagSetEntry>	TagSets;
+	TConstArrayView<uint32>					TagSetIndices;
+	FIoBuffer								ContainerHeaderChunk; // FIoContainerHeader
+};
+
+////////////////////////////////////////////////////////////////////////////////
+enum class EOnDemandTocReaderOptions : uint32
+{
+	None		= 0,
+	MemoryMap	= (1 << 0)
+};
+ENUM_CLASS_FLAGS(EOnDemandTocReaderOptions);
+
+////////////////////////////////////////////////////////////////////////////////
+class FOnDemandTocReader
+{
+public:
+	FOnDemandTocReader() = default;
+	FOnDemandTocReader(FOnDemandTocReader&&) = default;
+	FOnDemandTocReader& operator=(FOnDemandTocReader&&) = default;
+
+	FDateTime				CreationTime() const { return FDateTime::FromUnixTimestamp(Header->EpochTimestamp); }
+	FUtf8StringView			BuildVersion() const { return GetString(Header->BuildVersion); }
+	FUtf8StringView			TargetPlatform() const { return GetString(Header->TargetPlatform); }
+	FUtf8StringView			ChunksDirectory() const { return GetString(Header->ChunksDirectory); }
+	UE_API FUtf8StringView	GetString(FOnDemandStringEntry StringEntry) const;
+
+	TConstArrayView<FOnDemandContainerEntry> Containers() const { return ContainerEntries; }
+
+	UE_API TIoStatusOr<FOnDemandContainerTocView> ReadContainer(
+		const FOnDemandContainerEntry& Container,
+		FOnDemandTocStorage& Storage,
+		EOnDemandTocReaderOptions Options);
+	UE_API static TIoStatusOr<FOnDemandTocReader> Read(const FString& Filename);
+
+private:
+	FOnDemandTocReader(const FString& Filename, FIoBuffer Buffer);
+	FOnDemandTocReader(const FOnDemandTocReader&) = delete;
+	FOnDemandTocReader& operator=(const FOnDemandTocReader&) = delete;
+
+	FString										Filename;
+	FIoBuffer									Buffer;
+	const FOnDemandTocHeader*					Header = nullptr;
+	const UTF8CHAR*								StringTable = nullptr;
+	TConstArrayView<FOnDemandContainerEntry>	ContainerEntries;
+};
+
+} // V1
+
+} // UE::IoStore::Serialization
+
+#undef UE_API
diff --git a/Engine/Source/Runtime/Experimental/IoStore/OnDemandCore/Private/Serialization/OnDemandContainerToc.cpp b/Engine/Source/Runtime/Experimental/IoStore/OnDemandCore/Private/Serialization/OnDemandContainerToc.cpp
new file mode 100644
--- /dev/null
+++ b/Engine/Source/Runtime/Experimental/IoStore/OnDemandCore/Private/Serialization/OnDemandContainerToc.cpp
@@ -0,0 +1,546 @@
+// Copyright Epic Games, Inc. All Rights Reserved.
+
+#include "IO/Serialization/OnDemandContainerToc.h"
+
+#include "Algo/Find.h"
+#include "Async/MappedFileHandle.h"
+#include "Containers/ArrayView.h"
+#include "Containers/ChunkedArray.h"
+#include "Containers/Utf8String.h"
+#include "GenericPlatform/GenericPlatformFile.h"
+#include "HAL/FileManager.h"
+#include "HAL/Platform.h"
+#include "HAL/PlatformFileManager.h"
+#include "IO/IoStatus.h"
+#include "IO/OnDemandToc.h"
+#include "Memory/MemoryView.h"
+#include "Misc/DateTime.h"
+#include "Serialization/Archive.h"
+#include "Serialization/LargeMemoryWriter.h"
+#include "Templates/UniquePtr.h"
+#include "Templates/UnrealTemplate.h"
+
+namespace UE::IoStore::Serialization
+{
+
+////////////////////////////////////////////////////////////////////////////////
+FOnDemandTocStorage::FOnDemandTocStorage()
+{
+}
+
+FOnDemandTocStorage::FOnDemandTocStorage(
+	FIoBuffer InHeaderData,
+	FIoBuffer InContainerData)
+		: HeaderData(InHeaderData)
+{
+	ContainerData.Emplace<FMemoryStorage>(FMemoryStorage{ .Buffer = InContainerData });
+}
+
+FOnDemandTocStorage::FOnDemandTocStorage(
+	FIoBuffer InHeaderData,
+	TUniquePtr<IMappedFileHandle>&& File,
+	TUniquePtr<IMappedFileRegion>&& Region)
+		: HeaderData(InHeaderData)
+{
+	ContainerData.Emplace<FMappedFileStorage>(FMappedFileStorage
+	{
+		.MappedFile		= MoveTemp(File),
+		.MappedRegion	= MoveTemp(Region)
+	});
+}
+
+FOnDemandTocStorage::FOnDemandTocStorage(FOnDemandTocStorage&& Other)
+	: HeaderData(MoveTemp(Other.HeaderData))
+	, ContainerData(MoveTemp(Other.ContainerData))
+{
+}
+
+FOnDemandTocStorage::~FOnDemandTocStorage()
+{
+}
+
+FOnDemandTocStorage& FOnDemandTocStorage::operator=(FOnDemandTocStorage&& Other)
+{
+	HeaderData = MoveTemp(Other.HeaderData);
+	ContainerData = MoveTemp(Other.ContainerData);
+	return *this;
+}
+
+FMemoryView FOnDemandTocStorage::GetView() const
+{
+	if (const FMemoryStorage* MemoryStorage = ContainerData.TryGet<FMemoryStorage>())
+	{
+		return MemoryStorage->Buffer.GetView();
+	}
+	else if (const FMappedFileStorage* FileStorage = ContainerData.TryGet<FMappedFileStorage>())
+	{
+		return MakeMemoryView(FileStorage->MappedRegion->GetMappedPtr(), FileStorage->MappedRegion->GetMappedSize());
+	}
+
+	return FMemoryView();
+}
+
+////////////////////////////////////////////////////////////////////////////////
+class FOnDemandTocWriter::FImpl
+{
+	using FOnDemandStringEntry = V1::FOnDemandStringEntry;
+
+public:
+	TIoStatusOr<uint64> Write(const UE::IoStore::FOnDemandToc& Toc, FArchive& Ar)
+	{
+		using namespace UE::IoStore::Serialization::V1;
+
+		TArray<FOnDemandContainerEntry> ContainerEntries;
+		FLargeMemoryWriter				ContainerDataAr;
+		const int64						ArStart = Ar.Tell();
+
+		for (int32 ContainerIndex = 0; const UE::IoStore::FOnDemandTocContainerEntry& LegacyContainer : Toc.Containers)
+		{
+			TConstArrayView<uint32>			LegacyContainerBlockSizes = LegacyContainer.BlockSizes;
+			TConstArrayView<FIoBlockHash>	LegacyContainerBlockHashes = LegacyContainer.BlockHashes;
+			TArray<FIoChunkId>				ChunkIds;
+			TArray<FOnDemandChunkEntry>		ChunkEntries;
+			TArray<uint32>					BlockSizes;
+			TArray<FIoBlockHash>			BlockHashes;
+			TArray<FOnDemandTagSetEntry>	TagSets;
+			TArray<uint32>					TagSetIndices;
+
+			TArray<int32> EntryIndices;
+			EntryIndices.Reserve(LegacyContainer.Entries.Num());
+			for (int32 Idx = 0; Idx < LegacyContainer.Entries.Num(); ++Idx)
+			{
+				EntryIndices.Add(Idx);
+			}
+
+			EntryIndices.Sort([&LegacyContainer](int32 LHS, int32 RHS)
+			{
+				return LegacyContainer.Entries[LHS].ChunkId < LegacyContainer.Entries[RHS].ChunkId;
+			});
+
+			for (int32 EntryIndex : EntryIndices)
+			{
+				const FOnDemandTocEntry&		LegacyEntry = LegacyContainer.Entries[EntryIndex];
+				TConstArrayView<uint32>			LegacyBlockSizes = LegacyContainerBlockSizes.Mid(LegacyEntry.BlockOffset, LegacyEntry.BlockCount);
+				TConstArrayView<FIoBlockHash>	LegacyBlockHashes = LegacyContainerBlockHashes.Mid(LegacyEntry.BlockOffset, LegacyEntry.BlockCount);
+
+				ChunkIds.Add(LegacyEntry.ChunkId);
+				
+				FOnDemandChunkBlockInfo ChunkBlockInfo =
+					FOnDemandChunkBlockInfo::FromOffsetAndCount(LegacyEntry.BlockOffset, LegacyEntry.BlockCount);
+
+				// If the chunk only has a single compression block, the size and block hash are encoded into the chunk entry
+				if (LegacyBlockSizes.Num() > 1)
+				{
+					BlockSizes.Append(LegacyBlockSizes);
+					BlockHashes.Append(LegacyBlockHashes);
+				}
+				else
+				{
+					check(!LegacyBlockSizes.IsEmpty() && !LegacyBlockHashes.IsEmpty());
+					ChunkBlockInfo = FOnDemandChunkBlockInfo::FromSizeAndHash(LegacyBlockSizes[0], LegacyBlockHashes[0]);
+				}
+
+				ChunkEntries.Add(FOnDemandChunkEntry
+				{
+					.Hash			= LegacyEntry.Hash,
+					.RawSize		= IntCastChecked<uint32>(LegacyEntry.RawSize),
+					.EncodedSize	= IntCastChecked<uint32>(LegacyEntry.EncodedSize),
+					.BlockInfo		= ChunkBlockInfo
+				});
+			}
+
+			for (const FOnDemandTocTagSet& LegacyTagSet : Toc.TagSets)
+			{
+				for (const FOnDemandTocTagSetPackageList& PkgList : LegacyTagSet.Packages)
+				{
+					if (PkgList.ContainerIndex == ContainerIndex)
+					{
+						FOnDemandTagSetEntry& TagSetEntry	= TagSets.AddDefaulted_GetRef();
+						TagSetEntry.Tag						= WriteString(LegacyTagSet.Tag);
+						TagSetEntry.Offset					= TagSetIndices.Num();
+						TagSetEntry.Count					= PkgList.PackageIndicies.Num();
+
+						TagSetIndices.Append(PkgList.PackageIndicies);
+					}
+				}
+			}
+
+			FOnDemandContainerEntry& ContainerEntry = ContainerEntries.AddDefaulted_GetRef();
+			ContainerEntry.ContainerId			= LegacyContainer.ContainerId;
+			ContainerEntry.ContainerName		= WriteString(LegacyContainer.ContainerName);
+			ContainerEntry.ContainerHeaderSize	= LegacyContainer.Header.Num();
+			ContainerEntry.ChunkCount			= LegacyContainer.Entries.Num();
+			ContainerEntry.BlockCount			= BlockSizes.Num();
+			ContainerEntry.BlockSize			= Toc.Header.BlockSize;
+			ContainerEntry.TagSetCount			= TagSets.Num();
+			ContainerEntry.TagSetIndicesCount	= TagSetIndices.Num();
+			ContainerEntry.UTocHash				= LegacyContainer.UTocHash;
+			ContainerEntry.ContainerFlags 		= LegacyContainer.ContainerFlags;
+
+			if (LegacyContainer.EncryptionKeyGuid.IsEmpty() == false)
+			{
+				ensure(FGuid::Parse(LegacyContainer.EncryptionKeyGuid, ContainerEntry.EncryptionKeyGuid));
+			}
+
+			ContainerEntry.DataOffset = IntCastChecked<uint32>(ContainerDataAr.Tell());
+			{
+				ContainerDataAr.Serialize(ChunkIds.GetData(), ChunkIds.GetTypeSize() * ChunkIds.Num());
+				ContainerDataAr.Serialize(ChunkEntries.GetData(), ChunkEntries.GetTypeSize() * ChunkEntries.Num());
+				ContainerDataAr.Serialize((void*)BlockSizes.GetData(), BlockSizes.GetTypeSize() * BlockSizes.Num());
+				ContainerDataAr.Serialize((void*)BlockHashes.GetData(), BlockHashes.GetTypeSize() * BlockHashes.Num());
+				ContainerDataAr.Serialize(TagSets.GetData(), TagSets.GetTypeSize() * TagSets.Num());
+				ContainerDataAr.Serialize(TagSetIndices.GetData(), TagSetIndices.GetTypeSize() * TagSetIndices.Num());
+				ContainerDataAr.Serialize((void*)LegacyContainer.Header.GetData(), LegacyContainer.Header.Num());
+			}
+			ContainerEntry.DataSize = IntCastChecked<uint32>(ContainerDataAr.Tell()) - ContainerEntry.DataOffset;
+
+			ContainerIndex++;
+		}
+
+		V1::FOnDemandTocHeader Header;
+		Header.EpochTimestamp		= FDateTime::Now().ToUnixTimestamp();
+		Header.BuildVersion			= WriteString(Toc.Meta.BuildVersion);
+		Header.TargetPlatform		= WriteString(Toc.Meta.TargetPlatform);
+		Header.ChunksDirectory		= WriteString(Toc.Header.ChunksDirectory);
+		Header.HostGroupName		= WriteString(Toc.Header.HostGroupName);
+		Header.CompressionFormat	= WriteString(Toc.Header.CompressionFormat);
+		Header.ContainerCount		= Toc.Containers.Num();
+		Header.StringTableLen		= StringTableLen;
+
+		Ar.Serialize(&Header, sizeof(FOnDemandTocHeader));
+		{
+			const int64 StringTableOffset = Ar.Tell();
+			for (const FUtf8String& String : StringTable)
+			{
+				Ar.Serialize((void*)*String, sizeof(UTF8CHAR) * String.Len());
+			}
+			const int64 StringTableSize = Ar.Tell() - StringTableOffset;
+
+			if (sizeof(UTF8CHAR) * Header.StringTableLen != StringTableSize)
+			{
+				return FIoStatus(EIoErrorCode::WriteError, TEXT("Unexpected string table size"));
+			}
+		}
+		Ar.Serialize((void*)ContainerEntries.GetData(), ContainerEntries.GetTypeSize() * ContainerEntries.Num());
+		Ar.Serialize(ContainerDataAr.GetData(), ContainerDataAr.TotalSize());
+
+		FOnDemandTocFooter Footer;
+		Ar.Serialize(&Footer, sizeof(FOnDemandTocFooter));
+
+		return uint64(Ar.Tell() - ArStart);
+	}
+
+private:
+	FOnDemandStringEntry WriteString(const FStringView& String)
+	{
+		if (String.IsEmpty())
+		{
+			return FOnDemandStringEntry{};
+		}
+
+		FUtf8String Utf8String(String);
+		if (const FOnDemandStringEntry* ExistingEntry = StringToEntry.Find(Utf8String))
+		{
+			return *ExistingEntry;
+		}
+		else
+		{
+			const FOnDemandStringEntry NewEntry = FOnDemandStringEntry
+			{
+				.Offset = uint32(StringTableLen),
+				.Len	= uint32(Utf8String.Len())
+			};
+
+			StringTableLen += Utf8String.Len();
+			const int32 Index = StringTable.AddElement(MoveTemp(Utf8String));
+			StringToEntry.Add(StringTable[Index], NewEntry);
+
+			return NewEntry;
+		}
+	}
+
+	using FLookup = TMap<FUtf8StringView, FOnDemandStringEntry>;
+
+	FLookup						StringToEntry;
+	TChunkedArray<FUtf8String>	StringTable;
+	int32						StringTableLen = 0;
+};
+
+////////////////////////////////////////////////////////////////////////////////
+FOnDemandTocWriter::FOnDemandTocWriter()
+	: Impl(MakeUnique<FImpl>())
+{
+}
+
+FOnDemandTocWriter::FOnDemandTocWriter(FOnDemandTocWriter&& Other)
+	: Impl(MoveTemp(Other.Impl))
+{
+}
+
+FOnDemandTocWriter::~FOnDemandTocWriter()
+{
+}
+
+FOnDemandTocWriter& FOnDemandTocWriter::operator=(FOnDemandTocWriter&& Other)
+{
+	Impl = MoveTemp(Other.Impl);
+	return *this;
+}
+
+TIoStatusOr<uint64> FOnDemandTocWriter::Write(const UE::IoStore::FOnDemandToc& Toc, FArchive& Ar)
+{
+	return Impl->Write(Toc, Ar);
+}
+
+TIoStatusOr<uint64> FOnDemandTocWriter::Write(const UE::IoStore::FOnDemandToc& Toc, const FString& Filename)
+{
+	TUniquePtr<FArchive> Ar(IFileManager::Get().CreateFileWriter(*Filename));
+	if (Ar.IsValid() == false)
+	{
+		return FIoStatusBuilder(EIoErrorCode::FileOpenFailed)
+			<< TEXT("Failed to open '") << Filename << TEXT("'");
+	}
+
+	return Impl->Write(Toc, *Ar);
+}
+
+////////////////////////////////////////////////////////////////////////////////
+namespace V1
+{
+
+////////////////////////////////////////////////////////////////////////////////
+EIoContainerFlags FOnDemandContainerTocHeaderView::ContainerFlags() const
+{
+	return static_cast<EIoContainerFlags>(ContainerEntry->ContainerFlags);
+}
+
+FUtf8StringView	FOnDemandContainerTocHeaderView::GetString(FOnDemandStringEntry StringEntry) const
+{
+	return FUtf8StringView(StringTable + StringEntry.Offset, StringEntry.Len);
+}
+
+////////////////////////////////////////////////////////////////////////////////
+FOnDemandTocReader::FOnDemandTocReader(const FString& InFilename, FIoBuffer InBuffer)
+	: Filename(InFilename)
+	, Buffer(InBuffer)
+{
+	Header = reinterpret_cast<const FOnDemandTocHeader*>(Buffer.GetData());
+	StringTable = reinterpret_cast<const UTF8CHAR*>(Header + 1);
+	ContainerEntries = MakeArrayView(
+		reinterpret_cast<const FOnDemandContainerEntry*>(StringTable + Header->StringTableLen), Header->ContainerCount);
+}
+
+FUtf8StringView FOnDemandTocReader::GetString(FOnDemandStringEntry StringEntry) const
+{
+	return FUtf8StringView(StringTable + StringEntry.Offset, StringEntry.Len);
+}
+
+TIoStatusOr<FOnDemandContainerTocView> FOnDemandTocReader::ReadContainer(
+	const FOnDemandContainerEntry& Container,
+	FOnDemandTocStorage& Storage,
+	EOnDemandTocReaderOptions Options)
+{
+	const FOnDemandContainerEntry* Existing = Algo::FindByPredicate(
+		ContainerEntries, [&Container](const FOnDemandContainerEntry& Entry) { return &Entry == &Container; });
+
+	if (Existing == nullptr)
+	{
+		return FIoStatus(EIoErrorCode::NotFound);
+	}
+
+	IPlatformFile& Ipf = FPlatformFileManager::Get().GetPlatformFile();
+
+	const int64 ContainerDataOffset = sizeof(FOnDemandTocHeader)
+		+ (Header->StringTableLen * sizeof(UTF8CHAR))
+		+ (Header->ContainerCount * sizeof(FOnDemandContainerEntry))
+		+ Container.DataOffset;
+
+	const int64 ContainerDataSize = Container.DataSize - Container.ContainerHeaderSize;
+	if (ContainerDataSize <= 0)
+	{
+		return FIoStatusBuilder(EIoErrorCode::CorruptToc)
+			<< TEXT("Unexpected container data size in '")
+			<< Filename << TEXT("'");
+	}
+
+	const bool bMemoryMap = FPlatformProperties::SupportsMemoryMappedFiles()
+		&& EnumHasAnyFlags(Options, EOnDemandTocReaderOptions::MemoryMap);
+
+	if (bMemoryMap)
+	{
+		FOpenMappedResult OpenResult = Ipf.OpenMappedEx(*Filename);
+		if (OpenResult.HasError() == false)
+		{
+			TUniquePtr<IMappedFileHandle> MappedFile = OpenResult.StealValue();
+			TUniquePtr<IMappedFileRegion> Region(MappedFile->MapRegion(ContainerDataOffset, ContainerDataSize));
+
+			if (Region.IsValid())
+			{
+				Storage = FOnDemandTocStorage(Buffer, MoveTemp(MappedFile), MoveTemp(Region));
+			}
+		}
+	}
+
+	FFileOpenResult OpenResult = Ipf.OpenRead(*Filename, IPlatformFile::EOpenReadFlags::None);
+	if (OpenResult.HasError())
+	{
+		FFileSystemError& Error = OpenResult.GetError();
+		return FIoStatusBuilder(EIoErrorCode::FileOpenFailed)
+			<< TEXT("Failed to open '") << Filename << TEXT("' reason: ")
+			<< Error.GetMessage();
+	}
+
+	TUniquePtr<IFileHandle> FileHandle = OpenResult.StealValue();
+	if (Storage.IsEmpty())
+	{
+		if (FileHandle->Seek(ContainerDataOffset) == false)
+		{
+			return FIoStatusBuilder(EIoErrorCode::FileSeekFailed)
+				<< TEXT("Failed to seek to container data in '") << Filename << TEXT("'");
+		}
+
+		FIoBuffer ContainerData(ContainerDataSize);
+		if (FileHandle->Read(ContainerData.GetData(), ContainerData.GetSize()) == false)
+		{
+			return FIoStatusBuilder(EIoErrorCode::ReadError)
+				<< TEXT("Failed to read container data in '") << Filename << TEXT("'");
+		}
+
+		Storage = FOnDemandTocStorage(Buffer, ContainerData);
+	}
+
+	const int64 ContainerHeaderOffset = ContainerDataOffset + ContainerDataSize;
+	if (FileHandle->Seek(ContainerHeaderOffset) == false)
+	{
+		return FIoStatusBuilder(EIoErrorCode::FileSeekFailed)
+			<< TEXT("Failed to seek to container header in '") << Filename << TEXT("'");
+	}
+
+	// The container header is currently copied into a separate buffer
+	FIoBuffer ContainerHeaderChunk(Container.ContainerHeaderSize);
+	if (FileHandle->Read(ContainerHeaderChunk.GetData(), ContainerHeaderChunk.GetSize()) == false)
+	{
+		return FIoStatusBuilder(EIoErrorCode::ReadError)
+			<< TEXT("Failed to read container header in '") << Filename << TEXT("'");
+	}
+
+	FMemoryView ContainerData = Storage.GetView();
+	FOnDemandContainerTocView View;
+	View.Header = FOnDemandContainerTocHeaderView(Header, StringTable, &Container);
+
+	View.ChunkIds = MakeArrayView(reinterpret_cast<const FIoChunkId*>(ContainerData.GetData()), Container.ChunkCount);
+	ContainerData += View.ChunkIds.NumBytes();
+
+	View.ChunkEntries = MakeArrayView(reinterpret_cast<const FOnDemandChunkEntry*>(ContainerData.GetData()), Container.ChunkCount);
+	ContainerData += View.ChunkEntries.NumBytes();
+
+	View.BlockSizes = MakeArrayView(reinterpret_cast<const uint32*>(ContainerData.GetData()), Container.BlockCount);
+	ContainerData += View.BlockSizes.NumBytes();
+
+	View.BlockHashes = MakeArrayView(reinterpret_cast<const FIoBlockHash*>(ContainerData.GetData()), Container.BlockCount);
+	ContainerData += View.BlockHashes.NumBytes();
+
+	View.TagSets = MakeArrayView(reinterpret_cast<const FOnDemandTagSetEntry*>(ContainerData.GetData()), Container.TagSetCount);
+	ContainerData += View.TagSets.NumBytes();
+
+	View.TagSetIndices = MakeArrayView(reinterpret_cast<const uint32*>(ContainerData.GetData()), Container.TagSetIndicesCount);
+	ContainerData += View.TagSetIndices.NumBytes();
+
+	if (ContainerData.IsEmpty() == false)
+	{
+		return FIoStatusBuilder(EIoErrorCode::CorruptToc)
+			<< TEXT("Unexpected container data size in '") << Filename << TEXT("'");
+	}
+
+	View.ContainerHeaderChunk = MoveTemp(ContainerHeaderChunk);
+
+	return View;
+}
+
+TIoStatusOr<FOnDemandTocReader> FOnDemandTocReader::Read(const FString& Filename)
+{
+	IPlatformFile& Ipf = FPlatformFileManager::Get().GetPlatformFile();
+
+	const int64 FileSize = Ipf.FileSize(*Filename);
+	if (FileSize < int64(sizeof(FOnDemandTocHeader)))
+	{
+		return FIoStatus(EIoErrorCode::FileOpenFailed);
+	}
+
+	FFileOpenResult OpenResult = Ipf.OpenRead(*Filename, IPlatformFile::EOpenReadFlags::None);
+	if (OpenResult.HasError())
+	{
+		FFileSystemError& Error = OpenResult.GetError();
+		return FIoStatusBuilder(EIoErrorCode::FileOpenFailed)
+			<< TEXT("Failed to open '") << Filename << TEXT("' reason: ")
+			<< Error.GetMessage();
+	}
+
+	TUniquePtr<IFileHandle> FileHandle = OpenResult.StealValue();
+
+	FOnDemandTocHeader Header;
+	if (FileHandle->Read(reinterpret_cast<uint8*>(&Header), sizeof(FOnDemandTocHeader)) == false)
+	{
+		return FIoStatusBuilder(EIoErrorCode::ReadError)
+			<< TEXT("Failed to read container TOC header from '") << Filename << TEXT("'");
+	}
+
+	if (Header.Signature.IsValid() == false)
+	{
+		return FIoStatusBuilder(EIoErrorCode::CorruptToc)
+			<< TEXT("Invalid TOC signature found in '") << Filename << TEXT("'");
+	}
+
+	if (Header.Version.IsValid() == false)
+	{
+		return FIoStatusBuilder(EIoErrorCode::CorruptToc)
+			<< TEXT("Invalid TOC version found in '") << Filename << TEXT("'");
+	}
+
+	const int64 TotalHeaderSize =
+		sizeof(FOnDemandTocHeader) +
+		(sizeof(UTF8CHAR) * Header.StringTableLen) +
+		(sizeof(FOnDemandContainerEntry) * Header.ContainerCount);
+
+	if (TotalHeaderSize >= FileSize)
+	{
+		return FIoStatusBuilder(EIoErrorCode::CorruptToc)
+			<< TEXT("TOC header size greater than file size in '") << Filename << TEXT("'");
+	}
+
+	FIoBuffer Buffer(TotalHeaderSize);
+	FMutableMemoryView BufferView = Buffer.GetMutableView();
+	FMemory::Memcpy(BufferView.GetData(), &Header, sizeof(FOnDemandTocHeader));
+	BufferView += sizeof(FOnDemandTocHeader);
+
+	if (FileHandle->Read(reinterpret_cast<uint8*>(BufferView.GetData()), BufferView.GetSize()) == false)
+	{
+		return FIoStatusBuilder(EIoErrorCode::ReadError)
+			<< TEXT("Failed to read container TOC header from '") << Filename << TEXT("'");
+	}
+
+	if (FileHandle->SeekFromEnd(-int64(sizeof(FOnDemandTocFooter))) == false)
+	{
+		return FIoStatusBuilder(EIoErrorCode::FileSeekFailed)
+			<< TEXT("Failed to seek to container TOC footer in '") << Filename << TEXT("'");
+	}
+
+	FOnDemandTocFooter Footer;
+	if (FileHandle->Read(reinterpret_cast<uint8*>(&Footer), sizeof(FOnDemandTocFooter)) == false)
+	{
+		return FIoStatusBuilder(EIoErrorCode::ReadError)
+			<< TEXT("Failed to read container TOC footer from '") << Filename << TEXT("'");
+	}
+
+	if (Footer.Signature.IsValid() == false)
+	{
+		return FIoStatusBuilder(EIoErrorCode::CorruptToc)
+			<< TEXT("Invalid TOC footer signature found in '") << Filename << TEXT("'");
+	}
+
+	return FOnDemandTocReader(Filename, Buffer);
+}
+
+} // V1 
+
+} // UE::IoStore::Serialization
