Removed UTF8 BOM when serializing on-demand container TOC as Json.

#rnx
#Synced-CL 49823132

diff --git a/Engine/Source/Runtime/Experimental/IoStore/OnDemand/Private/Tool/ConvertToc.cpp b/Engine/Source/Runtime/Experimental/IoStore/OnDemand/Private/Tool/ConvertToc.cpp
--- a/Engine/Source/Runtime/Experimental/IoStore/OnDemand/Private/Tool/ConvertToc.cpp
+++ b/Engine/Source/Runtime/Experimental/IoStore/OnDemand/Private/Tool/ConvertToc.cpp
@@ -224,7 +224,7 @@ static int32 ConvertTocCommandEntry(const FContext& Context)
 		TStringBuilder<4096> Json;
 		CompactBinaryToJson(Writer.Save().AsObject(), Json);
 
-		if (!FFileHelper::SaveStringToFile(Json.ToView(), *OutFilePath, FFileHelper::EEncodingOptions::ForceUTF8))
+		if (!FFileHelper::SaveStringToFile(Json.ToView(), *OutFilePath, FFileHelper::EEncodingOptions::ForceUTF8WithoutBOM))
 		{
 			UE_LOG(LogIoStoreOnDemand, Error, TEXT("Failed to save JSON '%s'"), *OutFilePath);
 			return -1;
