Removed logging experiments.

#rnx

- Removed the experimental code and went back to one UE_LOG per line.

#Synced-CL 49815329

diff --git a/Engine/Source/Runtime/Experimental/IoStore/OnDemand/Private/Tool/ChunkPlugin.cpp b/Engine/Source/Runtime/Experimental/IoStore/OnDemand/Private/Tool/ChunkPlugin.cpp
--- a/Engine/Source/Runtime/Experimental/IoStore/OnDemand/Private/Tool/ChunkPlugin.cpp
+++ b/Engine/Source/Runtime/Experimental/IoStore/OnDemand/Private/Tool/ChunkPlugin.cpp
@@ -690,31 +690,14 @@ static int32 ChunkPluginCommandEntry(const FContext& Context)
 					UE_LOG(LogIoStoreOnDemand, Error, TEXT("Bad settings file '%s'"), *SettingsFile);
 					return -1;
 				}
-				
-#if 1
-				UE_LOGF(LogIoStoreOnDemand, Display, "TEST: Before Large Message");
 
-				UE_LOG(LogIoStoreOnDemand, Log, TEXT("Found Package Set '%s'"), *Pair.Key);
+				UE_LOG(LogIoStoreOnDemand, Display, TEXT("Found Package Set '%s' with %d entries"), *Pair.Key, Packages.Num());
 				for (const FString& Package : Packages)
 				{
-					UE_LOG(LogIoStoreOnDemand, Log, TEXT("->\t'%s'"), *Package);
+					UE_LOG(LogIoStoreOnDemand, Display, TEXT("->\t'%s'"), *Package);
 				}
 
-				UE_LOGF(LogIoStoreOnDemand, Display, "TEST: After Large Message");
-#else
-				TUtf8StringBuilder<0> Output;
-				Output << "Found Package Set '" << Pair.Key << "'";
-				for (const FString& Package : Packages)
-				{
-					Output << LINE_TERMINATOR << "->\t'" << Package << "'";
-				}
-
-				UE_LOGF(LogIoStoreOnDemand, Display, "TEST: Before Large Message");
-				UE_LOGF(LogIoStoreOnDemand, Display, "%hs", *Output);
-				UE_LOGF(LogIoStoreOnDemand, Display, "TEST: After Large Message");
-#endif
 				Input.Settings.PackageSets.Add(FString(Pair.Key), MoveTemp(Packages));
-				UE_LOGF(LogIoStoreOnDemand, Display, "Processed %d package set entries", Packages.Num());
 			}
 		}
 	}
