Reduce log spam when running the ChunkPlugin command.

#rnx

- These three log lines produce pretty much all the loggign for the command and get especially out of hand with larger plugins.
- Demote them to 'Log' verbosity so that the info is still stored in the log file is we need it but the display summary is much easier to read.

#Synced-CL 49815329

diff --git a/Engine/Source/Runtime/Experimental/IoStore/OnDemand/Private/Tool/ChunkPlugin.cpp b/Engine/Source/Runtime/Experimental/IoStore/OnDemand/Private/Tool/ChunkPlugin.cpp
--- a/Engine/Source/Runtime/Experimental/IoStore/OnDemand/Private/Tool/ChunkPlugin.cpp
+++ b/Engine/Source/Runtime/Experimental/IoStore/OnDemand/Private/Tool/ChunkPlugin.cpp
@@ -193,7 +193,7 @@ public:
 			return FIoStatusBuilder(EIoErrorCode::WriteError) << TEXT("Failed to write to file '") << FString(Sb.ToString()) << TEXT("'");
 		}
 
-		UE_LOG(LogIoStoreOnDemand, Display, TEXT("Writing file '%s' (%.2lf KiB)"), Sb.ToString(), double(Chunk.GetView().GetSize()) / 1024.0);
+		UE_LOG(LogIoStoreOnDemand, Log, TEXT("Writing file '%s' (%.2lf KiB)"), Sb.ToString(), double(Chunk.GetView().GetSize()) / 1024.0);
 		return EIoErrorCode::Ok;
 	}
 
@@ -355,7 +355,7 @@ static int32 ChunkContainer(const FString& ContainerUTocPath, EOnDemandTocFlags
 					int32 PackageIndex = INDEX_NONE;
 					if (ContainerHeader.PackageIds.Find(PackageId, PackageIndex))
 					{
-						UE_LOG(LogIoStoreOnDemand, Display, TEXT("Found package %s for tag %s in container %s"), *Package, *Tag, *ContainerUTocPath);
+						UE_LOG(LogIoStoreOnDemand, Log, TEXT("Found package %s for tag %s in container %s"), *Package, *Tag, *ContainerUTocPath);
 						TagSetPackageList.PackageIndicies.Add(PackageIndex);
 					}
 				}
@@ -694,7 +694,7 @@ static int32 ChunkPluginCommandEntry(const FContext& Context)
 				UE_LOG(LogIoStoreOnDemand, Display, TEXT("Found Package Set '%s' with %d entries"), *Pair.Key, Packages.Num());
 				for (const FString& Package : Packages)
 				{
-					UE_LOG(LogIoStoreOnDemand, Display, TEXT("->\t'%s'"), *Package);
+					UE_LOG(LogIoStoreOnDemand, Log, TEXT("->\t'%s'"), *Package);
 				}
 
 				Input.Settings.PackageSets.Add(FString(Pair.Key), MoveTemp(Packages));
