Further simplification of the ChunkPlugin command.

#rnx

- Function ChunkContainer no longer takes a FOnDemandToc, the input flag decides which type of toc to output, which will be stored in the FChunkOutput structure.
- Removed cmdline args 'StreamOnDemand' and 'InstallOnDemand' as they were not being used.

#Synced-CL 49815329

diff --git a/Engine/Source/Runtime/Experimental/IoStore/OnDemand/Private/Tool/ChunkPlugin.cpp b/Engine/Source/Runtime/Experimental/IoStore/OnDemand/Private/Tool/ChunkPlugin.cpp
--- a/Engine/Source/Runtime/Experimental/IoStore/OnDemand/Private/Tool/ChunkPlugin.cpp
+++ b/Engine/Source/Runtime/Experimental/IoStore/OnDemand/Private/Tool/ChunkPlugin.cpp
@@ -261,13 +261,40 @@ struct FChunkParams
 
 struct FChunkOutput
 {
-//	FOnDemandToc OnDemandToc;
+	FChunkOutput(const FString& HostGroupName)
+	{
+		InstallOnDemandToc.Header.Flags = static_cast<uint32>(EOnDemandTocFlags::InstallOnDemand);
+		InstallOnDemandToc.Header.HostGroupName = HostGroupName;
+		InstallOnDemandToc.Header.ChunksDirectory = TEXT(""); // TODO?
+		//InstallOnDemandToc.Containers.Reserve(ContainerFilenames.Num());
+
+		StreamOnDemandToc.Header.Flags = static_cast<uint32>(EOnDemandTocFlags::StreamOnDemand);
+		StreamOnDemandToc.Header.HostGroupName = HostGroupName;
+		StreamOnDemandToc.Header.ChunksDirectory = TEXT(""); // TODO?
+		//StreamOnDemandToc.Containers.Reserve(ContainerFilenames.Num());
+	}
+
+	FOnDemandToc& GetToc(EOnDemandTocFlags Flags)
+	{
+		if (EnumHasAnyFlags(Flags, EOnDemandTocFlags::InstallOnDemand))
+		{
+			return InstallOnDemandToc;
+		}
+		else
+		{
+			return StreamOnDemandToc;
+		}
+	}
+
+	FOnDemandToc InstallOnDemandToc;
+	FOnDemandToc StreamOnDemandToc;
+
 	TMap<FString, TArray<FOnDemandTocTagSetPackageList>> FoundTagSets;
 	TArray<FString> ContainerFiles;
 	FChunkPluginStats Stats;
 };
 
-static int32 ChunkContainer(const FString& ContainerUTocPath, EOnDemandTocFlags TocFlags, const FChunkParams& Input, FOnDemandToc& OutOnDemandToc, FChunkOutput& Output)
+static int32 ChunkContainer(const FString& ContainerUTocPath, EOnDemandTocFlags TocFlags, const FChunkParams& Input, FChunkOutput& Output)
 {
 	TRACE_CPUPROFILER_EVENT_SCOPE(ChunkPluginCommand::ChunkContainer);
 
@@ -285,6 +312,9 @@ static int32 ChunkContainer(const FString& ContainerUTocPath, EOnDemandTocFlags
 
 	UE_LOG(LogIoStoreOnDemand, Display, TEXT("Processing container '%s'"), *ContainerUTocPath);
 
+
+	FOnDemandToc& OnDemandToc = Output.GetToc(TocFlags);
+
 	/*TArray<FString> ContainerFileList;
 	ContainerFileReader.GetFilenames(ContainerFileList);
 	for (const FString& File : ContainerFileList)
@@ -293,15 +323,15 @@ static int32 ChunkContainer(const FString& ContainerUTocPath, EOnDemandTocFlags
 	}*/
 
 	const uint32 BlockSize = ContainerFileReader.GetCompressionBlockSize();
-	if (OutOnDemandToc.Header.BlockSize == 0)
+	if (OnDemandToc.Header.BlockSize == 0)
 	{
-		OutOnDemandToc.Header.BlockSize = ContainerFileReader.GetCompressionBlockSize();
+		OnDemandToc.Header.BlockSize = ContainerFileReader.GetCompressionBlockSize();
 	}
 
-	check(OutOnDemandToc.Header.BlockSize == ContainerFileReader.GetCompressionBlockSize());
+	check(OnDemandToc.Header.BlockSize == ContainerFileReader.GetCompressionBlockSize());
 
-	const int32 ContainerIndex = OutOnDemandToc.Containers.AddDefaulted();
-	FOnDemandTocContainerEntry& ContainerEntry = OutOnDemandToc.Containers[ContainerIndex];
+	const int32 ContainerIndex = OnDemandToc.Containers.AddDefaulted();
+	FOnDemandTocContainerEntry& ContainerEntry = OnDemandToc.Containers[ContainerIndex];
 	ContainerEntry.ContainerId = ContainerFileReader.GetContainerId();
 	ContainerEntry.ContainerName = FPaths::GetBaseFilename(ContainerUTocPath);
 
@@ -426,9 +456,9 @@ static int32 ChunkContainer(const FString& ContainerUTocPath, EOnDemandTocFlags
 				if (BlockInfo.CompressionMethod != NAME_None)
 				{
 					CompressionFormat = BlockInfo.CompressionMethod;
-					if (OutOnDemandToc.Header.CompressionFormat.IsEmpty())
+					if (OnDemandToc.Header.CompressionFormat.IsEmpty())
 					{
-						OutOnDemandToc.Header.CompressionFormat = BlockInfo.CompressionMethod.ToString();
+						OnDemandToc.Header.CompressionFormat = BlockInfo.CompressionMethod.ToString();
 					}
 				}
 			}
@@ -575,26 +605,8 @@ static int32 ChunkPluginCommandEntry(const FContext& Context)
 	const FString HostGroupName			= FString(Context.Get<FStringView>(TEXT("-HostGroupName"), FString()));
 
 	const bool bIgnoreContainerHeader	= Context.Get<bool>(TEXT("-IgnoreContainerHeader"), false);
-	const bool bStreamOnDemand			= Context.Get<bool>(TEXT("-StreamOnDemand"), false);
-	const bool bInstallOnDemand			= Context.Get<bool>(TEXT("-InstallOnDemand"), false);
-
-	bool bSupportInstalledContainers = !Context.Get<bool>(TEXT("-SkipInstallOnDemand"), false);
-	bool bSupportStreamingContainers = !Context.Get<bool>(TEXT("-SkipStreamOnDemand"), false);
-
-	EOnDemandTocFlags TocFlags = EOnDemandTocFlags::None;
-	if (bStreamOnDemand)
-	{
-		TocFlags |= EOnDemandTocFlags::StreamOnDemand;
-	}
-	else if (bInstallOnDemand)
-	{
-		TocFlags |= EOnDemandTocFlags::InstallOnDemand;
-	}
-	else
-	{
-		TocFlags = EOnDemandTocFlags::InstallOnDemand;
-	}
-
+	const bool bSupportInstalledContainers = !Context.Get<bool>(TEXT("-SkipInstallOnDemand"), false);
+	const bool bSupportStreamingContainers = !Context.Get<bool>(TEXT("-SkipStreamOnDemand"), false);
 
 	FS3Params S3Params(Context);
 
@@ -628,7 +640,6 @@ static int32 ChunkPluginCommandEntry(const FContext& Context)
 	UE_LOG(LogIoStoreOnDemand, Display, TEXT("\tDeleteContainerFiles: %s"), bDeleteContainerFiles ? TEXT("true") : TEXT("false"));
 	UE_LOG(LogIoStoreOnDemand, Display, TEXT("\tIgnoreContainerHeader: %s"), bIgnoreContainerHeader ? TEXT("true") : TEXT("false"));
 	UE_LOG(LogIoStoreOnDemand, Display, TEXT("\tValidateChunks: %s"), bValidateChunks ? TEXT("true") : TEXT("false"));
-	UE_LOG(LogIoStoreOnDemand, Display, TEXT("\tTOC flags: '%s'"), *LexToString(TocFlags));
 
 	FChunkParams Input;
 	Input.ChunksRelativeFolder = ChunksRelativeFolder;
@@ -726,19 +737,7 @@ static int32 ChunkPluginCommandEntry(const FContext& Context)
 		return -1;
 	}
 
-	FOnDemandToc OnDemandToc;
-	OnDemandToc.Header.Flags = static_cast<uint32>(TocFlags);
-	OnDemandToc.Header.HostGroupName = HostGroupName;
-	OnDemandToc.Header.ChunksDirectory = TEXT(""); // TODO?
-	OnDemandToc.Containers.Reserve(ContainerFilenames.Num());
-
-	FOnDemandToc StreamingOnDemandToc;
-	StreamingOnDemandToc.Header.Flags = static_cast<uint32>(EOnDemandTocFlags::StreamOnDemand);
-	StreamingOnDemandToc.Header.HostGroupName = HostGroupName;
-	StreamingOnDemandToc.Header.ChunksDirectory = TEXT(""); // TODO?
-	StreamingOnDemandToc.Containers.Reserve(ContainerFilenames.Num());
-
-	FChunkOutput Output;
+	FChunkOutput Output(HostGroupName);
 	Output.FoundTagSets.Reserve(Input.Settings.PackageSets.Num());
 
 	for (const FString& Filename : ContainerFilenames)
@@ -751,7 +750,7 @@ static int32 ChunkPluginCommandEntry(const FContext& Context)
 		if (bSupportStreamingContainers && bIsStreamOnDemandContainer)
 		{
 			// Attempt to add to IAS output container
-			int32 Result = ChunkContainer(FullPath, EOnDemandTocFlags::StreamOnDemand, Input, StreamingOnDemandToc, Output);
+			int32 Result = ChunkContainer(FullPath, EOnDemandTocFlags::StreamOnDemand, Input, Output);
 			if (Result < 0)
 			{
 				return Result;
@@ -760,7 +759,7 @@ static int32 ChunkPluginCommandEntry(const FContext& Context)
 		else if(bSupportInstalledContainers && !bIsStreamOnDemandContainer)
 		{
 			// Attempt to add to IAD output container
-			int32 Result = ChunkContainer(FullPath, EOnDemandTocFlags::InstallOnDemand, Input, OnDemandToc, Output);
+			int32 Result = ChunkContainer(FullPath, EOnDemandTocFlags::InstallOnDemand, Input, Output);
 			if (Result < 0)
 			{
 				return Result;
@@ -772,10 +771,10 @@ static int32 ChunkPluginCommandEntry(const FContext& Context)
 
 	if (bSupportInstalledContainers)
 	{
-		OnDemandToc.TagSets.Reserve(Output.FoundTagSets.Num());
+		Output.InstallOnDemandToc.TagSets.Reserve(Output.FoundTagSets.Num());
 		for (TPair<FString, TArray<FOnDemandTocTagSetPackageList>>& Pair : Output.FoundTagSets)
 		{
-			FOnDemandTocTagSet& TocTagSet = OnDemandToc.TagSets.Emplace_GetRef();
+			FOnDemandTocTagSet& TocTagSet = Output.InstallOnDemandToc.TagSets.Emplace_GetRef();
 			TocTagSet.Tag = Pair.Key;
 			TocTagSet.Packages = MoveTemp(Pair.Value);
 		}
@@ -850,7 +849,7 @@ static int32 ChunkPluginCommandEntry(const FContext& Context)
 			}
 
 			UE_LOG(LogIoStoreOnDemand, Display, TEXT("Adding additional file '%s'"), *Filename);
-			OnDemandToc.AdditionalFiles.Add(FOnDemandTocAdditionalFile
+			Output.InstallOnDemandToc.AdditionalFiles.Add(FOnDemandTocAdditionalFile
 			{
 				.Hash = ChunkHash,
 				.Filename = Filename,
@@ -864,27 +863,27 @@ static int32 ChunkPluginCommandEntry(const FContext& Context)
 	// const FString TocPath = OutputFolder / Filename; // like this yo
 	const FString OnDemandTocOutputFolder = OutputFolder / IoStoreRelativeFolder;
 
-	if (bSupportInstalledContainers && !OnDemandToc.Containers.IsEmpty())
+	if (bSupportInstalledContainers && !Output.InstallOnDemandToc.Containers.IsEmpty())
 	{
-		OnDemandToc.Meta.EpochTimestamp = Timestamp;
-		OnDemandToc.Meta.TargetPlatform = Platform;
-		OnDemandToc.Meta.BuildVersion = BuildVersion;
+		Output.InstallOnDemandToc.Meta.EpochTimestamp = Timestamp;
+		Output.InstallOnDemandToc.Meta.TargetPlatform = Platform;
+		Output.InstallOnDemandToc.Meta.BuildVersion = BuildVersion;
 
-		int32 Result = WriteOnDemandToc(OnDemandToc, OnDemandTocName, OnDemandTocOutputFolder);
+		int32 Result = WriteOnDemandToc(Output.InstallOnDemandToc, OnDemandTocName, OnDemandTocOutputFolder);
 		if (Result < 0)
 		{
 			return Result;
 		}
 	}
 
-	if (bSupportStreamingContainers && !StreamingOnDemandToc.Containers.IsEmpty())
+	if (bSupportStreamingContainers && !Output.StreamOnDemandToc.Containers.IsEmpty())
 	{
-		StreamingOnDemandToc.Meta.EpochTimestamp = Timestamp;
-		StreamingOnDemandToc.Meta.TargetPlatform = Platform;
-		StreamingOnDemandToc.Meta.BuildVersion = BuildVersion;
+		Output.StreamOnDemandToc.Meta.EpochTimestamp = Timestamp;
+		Output.StreamOnDemandToc.Meta.TargetPlatform = Platform;
+		Output.StreamOnDemandToc.Meta.BuildVersion = BuildVersion;
 
 		// TODO: OnDemandTocName
-		int32 Result = WriteOnDemandToc(StreamingOnDemandToc, OnDemandTocName + TEXT("ondemand"), OnDemandTocOutputFolder);
+		int32 Result = WriteOnDemandToc(Output.StreamOnDemandToc, OnDemandTocName + TEXT("ondemand"), OnDemandTocOutputFolder);
 		if (Result < 0)
 		{
 			return Result;
@@ -968,8 +967,6 @@ static FCommand ChunkPluginCommand(
 		TArgument<FStringView>(TEXT("-CryptoKeys"),			TEXT("")),
 		TArgument<bool>(TEXT("-IncludeSigPak"),				TEXT("Include .sig and .pak file in the uondemandtoc")),
 		TArgument<bool>(TEXT("-KeepContainerFiles"),		TEXT("Should we keep the container files after processing them.")),
-		TArgument<bool>(TEXT("-StreamOnDemand"),			TEXT("Set the content to be streamed on-demand")),
-		TArgument<bool>(TEXT("-InstallOnDemand"),			TEXT("Set the content to be installed on-demand")),
 		TArgument<bool>(TEXT("-SkipStreamOnDemand"),		TEXT("Do not chunk any stream on-demand containers")),
 		TArgument<bool>(TEXT("-SkipInstallOnDemand"),		TEXT("Do not chunk any install on-demand containers")),
 		TArgument<bool>(TEXT("-IgnoreContainerHeader"),		TEXT("Ignore the container header?")),
