Remove disabled code from the ChunkPlugin command.

#rnx

#Synced-CL 49815329

diff --git a/Engine/Source/Runtime/Experimental/IoStore/OnDemand/Private/Tool/ChunkPlugin.cpp b/Engine/Source/Runtime/Experimental/IoStore/OnDemand/Private/Tool/ChunkPlugin.cpp
--- a/Engine/Source/Runtime/Experimental/IoStore/OnDemand/Private/Tool/ChunkPlugin.cpp
+++ b/Engine/Source/Runtime/Experimental/IoStore/OnDemand/Private/Tool/ChunkPlugin.cpp
@@ -312,16 +312,8 @@ static int32 ChunkContainer(const FString& ContainerUTocPath, EOnDemandTocFlags
 
 	UE_LOG(LogIoStoreOnDemand, Display, TEXT("Processing container '%s'"), *ContainerUTocPath);
 
-
 	FOnDemandToc& OnDemandToc = Output.GetToc(TocFlags);
 
-	/*TArray<FString> ContainerFileList;
-	ContainerFileReader.GetFilenames(ContainerFileList);
-	for (const FString& File : ContainerFileList)
-	{
-		UE_LOG(LogIoStoreOnDemand, Display, TEXT("Found file in container '%s'"), *File);
-	}*/
-
 	const uint32 BlockSize = ContainerFileReader.GetCompressionBlockSize();
 	if (OnDemandToc.Header.BlockSize == 0)
 	{
