Remove duplicate code from ChunkPlugin command, just determine the flags instead.

#rnx

#Synced-CL 49815329

diff --git a/Engine/Source/Runtime/Experimental/IoStore/OnDemand/Private/Tool/ChunkPlugin.cpp b/Engine/Source/Runtime/Experimental/IoStore/OnDemand/Private/Tool/ChunkPlugin.cpp
--- a/Engine/Source/Runtime/Experimental/IoStore/OnDemand/Private/Tool/ChunkPlugin.cpp
+++ b/Engine/Source/Runtime/Experimental/IoStore/OnDemand/Private/Tool/ChunkPlugin.cpp
@@ -739,23 +739,26 @@ static int32 ChunkPluginCommandEntry(const FContext& Context)
 		// TODO: At the moment the only way we can identify an OnDemand container is by the name. Might be better to pass this info into the tool
 		const bool bIsStreamOnDemandContainer = FullPath.Contains(TEXT("OnDemand"));
 
+		EOnDemandTocFlags ContainerFlags = EOnDemandTocFlags::None;
+
 		if (bSupportStreamingContainers && bIsStreamOnDemandContainer)
 		{
-			// Attempt to add to IAS output container
-			int32 Result = ChunkContainer(FullPath, EOnDemandTocFlags::StreamOnDemand, Input, Output);
-			if (Result < 0)
-			{
-				return Result;
-			}
+			ContainerFlags = EOnDemandTocFlags::StreamOnDemand;
 		}
-		else if(bSupportInstalledContainers && !bIsStreamOnDemandContainer)
+		else if (bSupportInstalledContainers && !bIsStreamOnDemandContainer)
 		{
-			// Attempt to add to IAD output container
-			int32 Result = ChunkContainer(FullPath, EOnDemandTocFlags::InstallOnDemand, Input, Output);
-			if (Result < 0)
-			{
-				return Result;
-			}
+			ContainerFlags = EOnDemandTocFlags::InstallOnDemand;
+		}
+		else
+		{
+			UE_LOGF(LogIoStoreOnDemand, Display, "Skipping container '%ls'", *FullPath);
+			continue;
+		}
+
+		int32 Result = ChunkContainer(FullPath, ContainerFlags, Input, Output);
+		if (Result < 0)
+		{
+			return Result;
 		}
 	}
 
