Add Asset path nodes ( Get asset path , Path Replace  )

#Synced-CL 49892100

diff --git a/Engine/Source/Runtime/Experimental/Dataflow/Engine/Private/Dataflow/DataflowAssetPathNodes.cpp b/Engine/Source/Runtime/Experimental/Dataflow/Engine/Private/Dataflow/DataflowAssetPathNodes.cpp
new file mode 100644
--- /dev/null
+++ b/Engine/Source/Runtime/Experimental/Dataflow/Engine/Private/Dataflow/DataflowAssetPathNodes.cpp
@@ -0,0 +1,151 @@
+// Copyright Epic Games, Inc. All Rights Reserved.
+
+#include "Dataflow/DataflowAssetPathNodes.h"
+#include "Dataflow/DataflowNodeFactory.h"
+
+#include UE_INLINE_GENERATED_CPP_BY_NAME(DataflowAssetPathNodes)
+
+#define LOCTEXT_NAMESPACE "DataflowAssetPathNodes"
+
+namespace UE::Dataflow
+{
+	void RegisterAssetPathNodes()
+	{
+		DATAFLOW_NODE_REGISTER_CREATION_FACTORY(FDataflowGetAssetPathNode);
+		DATAFLOW_NODE_REGISTER_CREATION_FACTORY(FDataflowPathReplaceNode);
+	}
+}
+
+////////////////////////////////////////////////////////////////////////////////////////////
+
+FDataflowGetAssetPathNode::FDataflowGetAssetPathNode(const UE::Dataflow::FNodeParameters& InParam, FGuid InGuid)
+	: FDataflowNode(InParam, InGuid)
+{
+	RegisterInputConnection(&Asset);
+	RegisterOutputConnection(&Path);
+}
+
+void FDataflowGetAssetPathNode::Evaluate(UE::Dataflow::FContext& Context, const FDataflowOutput* Out) const
+{
+	if (Out->IsA(&Path))
+	{
+		FString OutPath;
+
+		TObjectPtr<UObject> InAsset = GetValue(Context, &Asset);
+		if (InAsset == nullptr)
+		{
+			Context.Warning(LOCTEXT("NullAsset", "Asset is null, empty path will be returned"), this, Out);
+			
+		}
+		else
+		{
+			OutPath = InAsset->GetPathName();
+		}
+		SetValue(Context, OutPath, &Path);
+	}
+}
+
+bool FDataflowGetAssetPathNode::SupportsAssetProperty(UObject* InAsset) const
+{
+	return (InAsset != nullptr);
+}
+
+void FDataflowGetAssetPathNode::SetAssetProperty(UObject* InAsset)
+{
+	if (InAsset)
+	{
+		Asset.Value = InAsset;
+	}
+}
+
+////////////////////////////////////////////////////////////////////////////////////////////
+
+FDataflowPathReplaceNode::FDataflowPathReplaceNode(const UE::Dataflow::FNodeParameters& InParam, FGuid InGuid)
+	: FDataflowNode(InParam, InGuid)
+{
+	RegisterInputConnection(&Path);
+	RegisterOutputConnection(&Path, &Path);
+}
+
+void FDataflowPathReplaceNode::Evaluate(UE::Dataflow::FContext& Context, const FDataflowOutput* Out) const
+{
+	if (Out->IsA(&Path))
+	{
+		const FString& InPath = GetValue(Context, &Path);
+
+		FString PathSection;
+		FString NameSection;
+		InPath.Split("/", &PathSection, &NameSection, ESearchCase::CaseSensitive, ESearchDir::FromEnd);
+		FString LeftName;
+		FString RightName;
+		NameSection.Split(".", &LeftName, &RightName, ESearchCase::CaseSensitive, ESearchDir::FromEnd);
+
+		if (!PathSection.IsEmpty())
+		{
+			for (const FDataflowPathReplaceRule& PathRule : PathRules)
+			{
+				PathSection = ApplyRule(PathSection, PathRule);
+			}
+			PathSection += "/";
+		}
+
+		if (!LeftName.IsEmpty())
+		{
+			for (const FDataflowPathReplaceRule& NameRule : NameRules)
+			{
+				LeftName = ApplyRule(LeftName, NameRule);
+			}
+
+			LeftName += ".";
+		}
+
+		if (!RightName.IsEmpty())
+		{
+			for (const FDataflowPathReplaceRule& NameRule : NameRules)
+			{
+				RightName = ApplyRule(RightName, NameRule);
+			}
+		}
+		
+		const FString OutPath = PathSection + LeftName + RightName;
+		SetValue(Context, OutPath, &Path);
+	}
+}
+
+FString FDataflowPathReplaceNode::ApplyRule(const FString& InString, const FDataflowPathReplaceRule& Rule) const
+{
+	const ESearchCase::Type SearchCase = Rule.bIgnoreCase ? ESearchCase::IgnoreCase : ESearchCase::CaseSensitive;
+
+	FString OutString = InString;
+	switch (Rule.Method)
+	{
+	case EDataflowPathReplaceMethod::Anywhere:
+		OutString =  InString.Replace(*Rule.From, *Rule.To, SearchCase);
+		break;
+		
+	case EDataflowPathReplaceMethod::Beginning:
+		if (InString.StartsWith(Rule.From, SearchCase))
+		{
+			OutString = Rule.To + InString.RightChop(Rule.From.Len());
+		}
+		else if (Rule.bInsertIfNoFound)
+		{
+			OutString = Rule.To + InString;
+		}
+		break;
+
+	case EDataflowPathReplaceMethod::End:
+		if (InString.EndsWith(Rule.From, SearchCase))
+		{
+			OutString = InString.LeftChop(Rule.From.Len()) + Rule.To;
+		}
+		else if (Rule.bInsertIfNoFound)
+		{
+			OutString = InString + Rule.To;
+		}
+		break;
+	}
+	return OutString;
+}
+
+#undef LOCTEXT_NAMESPACE
diff --git a/Engine/Source/Runtime/Experimental/Dataflow/Engine/Private/Dataflow/DataflowAssetPathNodes.h b/Engine/Source/Runtime/Experimental/Dataflow/Engine/Private/Dataflow/DataflowAssetPathNodes.h
new file mode 100644
--- /dev/null
+++ b/Engine/Source/Runtime/Experimental/Dataflow/Engine/Private/Dataflow/DataflowAssetPathNodes.h
@@ -0,0 +1,117 @@
+// Copyright Epic Games, Inc. All Rights Reserved.
+
+#pragma once 
+
+#include "Dataflow/DataflowAnyType.h"
+#include "Dataflow/DataflowNode.h"
+
+#include "DataflowAssetPathNodes.generated.h"
+
+//////////////////////////////////////////////////////////////////////////////////////////////////////////////////////////
+
+/*
+* Get the path from an asset
+*/
+USTRUCT()
+struct FDataflowGetAssetPathNode : public FDataflowNode
+{
+	GENERATED_USTRUCT_BODY()
+	DATAFLOW_NODE_DEFINE_INTERNAL(FDataflowGetAssetPathNode, "GetAssetPath", "Path", "UObject Folder Name String")
+
+private:
+	/** Asset to get the path from */
+	UPROPERTY(EditAnywhere, Category = Asset, meta = (DataflowInput))
+	FDataflowUObjectConvertibleTypes Asset;
+
+	UPROPERTY(meta = (DataflowOutput))
+	FString Path;
+
+	virtual void Evaluate(UE::Dataflow::FContext& Context, const FDataflowOutput* Out) const override;
+
+public:
+	FDataflowGetAssetPathNode(const UE::Dataflow::FNodeParameters& InParam, FGuid InGuid = FGuid::NewGuid());
+
+	virtual bool SupportsAssetProperty(UObject* InAsset) const override;
+	virtual void SetAssetProperty(UObject* InAsset) override;
+};
+
+//////////////////////////////////////////////////////////////////////////////////////////////////////////////////////////
+
+UENUM()
+enum class EDataflowPathReplaceMethod: uint8
+{
+	Anywhere,
+	Beginning,
+	End,
+};
+
+USTRUCT()
+struct FDataflowPathReplaceRule
+{
+	GENERATED_USTRUCT_BODY()
+
+public:
+	/** Method to use to find the From string */
+	UPROPERTY(EditAnywhere, Category = Rule)
+	EDataflowPathReplaceMethod Method = EDataflowPathReplaceMethod::Anywhere;
+
+	/** String to find and replace */
+	UPROPERTY(EditAnywhere, Category = Rule)
+	FString From;
+
+	/** String to use as a replacement */
+	UPROPERTY(EditAnywhere, Category = Rule)
+	FString To;
+
+	UPROPERTY(EditAnywhere, Category = Rule)
+	bool bIgnoreCase = true;
+
+	/** if true insert the string even if not found ( only for Beginning and End methods )*/
+	UPROPERTY(EditAnywhere, Category = Rule, meta = (EditCondition = "Method == EDataflowPathReplaceMethod::Beginning || EDataflowPathReplaceMethod::End"))
+	bool bInsertIfNoFound = false;
+};
+
+/*
+* Replace part of a path
+* The path is assumed to have the form  Path/Names ( Path can have multiple / and folders and Names can have multiple names seperated by a .)
+*/
+USTRUCT()
+struct FDataflowPathReplaceNode : public FDataflowNode
+{
+	GENERATED_USTRUCT_BODY()
+	DATAFLOW_NODE_DEFINE_INTERNAL(FDataflowPathReplaceNode, "PathReplace", "Path", "Folder File Name String")
+
+private:
+	/** Path to modify and return */
+	UPROPERTY(meta = (DataflowInput, DataflowOutput, DataflowPassthrough = Path))
+	FString Path;
+
+	/** 
+	* Rules to apply to the folder part of the path 
+	* Note that Rules are applied in order 
+	*/
+	UPROPERTY(EditAnywhere, Category = Rules, meta = (TitleProperty = "'{From}' -> '{To}' [{Method}]"))
+	TArray<FDataflowPathReplaceRule> PathRules;
+
+	/** 
+	* Rules to apply to the name part of the path
+	* Note that Rules are applied in order 
+	*/
+	UPROPERTY(EditAnywhere, Category = Rules, meta=(TitleProperty = "'{From}' -> '{To}' [{Method}]"))
+	TArray<FDataflowPathReplaceRule> NameRules;
+
+	virtual void Evaluate(UE::Dataflow::FContext& Context, const FDataflowOutput* Out) const override;
+
+	FString ApplyRule(const FString& InString, const FDataflowPathReplaceRule& Rule) const;
+
+public:
+	FDataflowPathReplaceNode(const UE::Dataflow::FNodeParameters& InParam, FGuid InGuid = FGuid::NewGuid());
+};
+
+//////////////////////////////////////////////////////////////////////////////////////////////////////////////////////////
+
+namespace UE::Dataflow
+{
+	void RegisterAssetPathNodes();
+}
+
diff --git a/Engine/Source/Runtime/Experimental/Dataflow/Engine/Private/Dataflow/DataflowEngineModule.cpp b/Engine/Source/Runtime/Experimental/Dataflow/Engine/Private/Dataflow/DataflowEngineModule.cpp
--- a/Engine/Source/Runtime/Experimental/Dataflow/Engine/Private/Dataflow/DataflowEngineModule.cpp
+++ b/Engine/Source/Runtime/Experimental/Dataflow/Engine/Private/Dataflow/DataflowEngineModule.cpp
@@ -6,6 +6,7 @@
 #include "Modules/ModuleManager.h"
 
 #include "Dataflow/DataflowTextureAssetNodes.h"
+#include "Dataflow/DataflowAssetPathNodes.h"
 #include "Dataflow/DataflowVariableNodes.h"
 #include "Dataflow/DataflowSubGraphNodes.h"
 
@@ -21,6 +22,7 @@ public:
 	virtual void StartupModule() override
 	{
 		UE::Dataflow::RegisterTextureAssetNodes();
+		UE::Dataflow::RegisterAssetPathNodes();
 		UE::Dataflow::RegisterVariableNodes();
 		UE::Dataflow::RegisterSubGraphNodes();
 	}
