Updated on-demand container TOC with missing paramaters.

#rnx

#Synced-CL 49823132

diff --git a/Engine/Source/Runtime/Experimental/IoStore/OnDemandCore/Internal/IO/Serialization/OnDemandContainerToc.h b/Engine/Source/Runtime/Experimental/IoStore/OnDemandCore/Internal/IO/Serialization/OnDemandContainerToc.h
--- a/Engine/Source/Runtime/Experimental/IoStore/OnDemandCore/Internal/IO/Serialization/OnDemandContainerToc.h
+++ b/Engine/Source/Runtime/Experimental/IoStore/OnDemandCore/Internal/IO/Serialization/OnDemandContainerToc.h
@@ -170,6 +170,14 @@ struct FOnDemandTocHeader
 };
 static_assert(sizeof(FOnDemandTocHeader) == 128);
 
+enum class EOnDemandContainerEntryFlags : uint32
+{
+	None			= 0,
+	InstallOnDemand	= (1 << 0),
+	StreamOnDemand	= (1 << 1),
+};
+ENUM_CLASS_FLAGS(EOnDemandContainerEntryFlags);
+
 struct FOnDemandContainerEntry
 {
 	FGuid					EncryptionKeyGuid;
@@ -184,8 +192,9 @@ struct FOnDemandContainerEntry
 	uint32					TagSetCount = 0;
 	uint32					TagSetIndicesCount = 0;
 	FIoHash					UTocHash = FIoHash::Zero;
-	uint32					ContainerFlags = 0; // EIoContainerFlags
-	int8					Pad[40] = {0};
+	uint32					ContainerFlags = 0; // EOnDemandContainerEntryFlags
+	uint32					FileContainerFlags = 0; // EIoContainerFlags
+	int8					Pad[36] = {0};
 };
 static_assert(sizeof(FOnDemandContainerEntry) == 128);
 
@@ -196,10 +205,11 @@ struct FOnDemandChunkBlockInfo
 		FMemory::Memzero(this, sizeof(FOnDemandChunkBlockInfo));
 	}
 
-	uint32 Offset()	const { return bHasOffset ? OffsetOrSize : ~uint32(0); }
-	uint32 Count()	const { return bHasOffset ? CountOrHash : OffsetOrSize > 0 ? 1 : 0; }
-	uint32 Size()	const { return bHasOffset ? 0 : OffsetOrSize; }
-	uint32 Hash()	const { return bHasOffset ? 0 : CountOrHash; }
+	uint32			Offset()	const { return bHasOffset ? OffsetOrSize : ~uint32(0); }
+	uint32			Count()		const { return bHasOffset ? CountOrHash : OffsetOrSize > 0 ? 1 : 0; }
+	uint32			Size()		const { return bHasOffset ? 0 : OffsetOrSize; }
+	uint32			Hash()		const { return bHasOffset ? 0 : CountOrHash; }
+	const uint32*	Data()		const { return reinterpret_cast<const uint32*>(this); } 
 
 	static FOnDemandChunkBlockInfo FromOffsetAndCount(uint32 BlockOffset, uint32 BlockCount)
 	{
@@ -228,7 +238,7 @@ static_assert(sizeof(FOnDemandChunkBlockInfo) == 8);
 
 struct FOnDemandChunkEntry
 {
-	uint32	DiskSize() const { return Align(EncodedSize, FAES::AESBlockSize); }
+	uint32						GetDiskSize() const { return Align(EncodedSize, FAES::AESBlockSize); }
 
 	FIoHash						Hash = FIoHash::Zero;
 	uint32						RawSize = 0;
@@ -248,6 +258,8 @@ static_assert(sizeof(FOnDemandTagSetEntry) == 16);
 ////////////////////////////////////////////////////////////////////////////////
 struct FOnDemandContainerTocHeaderView
 {
+	using EContainerFlags = EOnDemandContainerEntryFlags;
+
 	FOnDemandContainerTocHeaderView() = default;
 	FOnDemandContainerTocHeaderView(
 		const FOnDemandTocHeader* InHeader,
@@ -260,7 +272,8 @@ struct FOnDemandContainerTocHeaderView
 
 	FUtf8StringView			ContainerName() const { return GetString(ContainerEntry->ContainerName); }
 	FIoContainerId			ContainerId() const { return ContainerEntry->ContainerId; }
-	EIoContainerFlags		ContainerFlags() const;
+	EContainerFlags			ContainerFlags() const { return static_cast<EContainerFlags>(ContainerEntry->ContainerFlags); }
+	EIoContainerFlags		FileContainerFlags() const;
 	const FGuid&			EncryptionKeyGuid() const { return ContainerEntry->EncryptionKeyGuid; }
 	uint32					BlockSize() const { return ContainerEntry->BlockSize; }
 	UE_API FUtf8StringView	GetString(FOnDemandStringEntry StringEntry) const;
@@ -307,6 +320,7 @@ public:
 	FUtf8StringView			BuildVersion() const { return GetString(Header->BuildVersion); }
 	FUtf8StringView			TargetPlatform() const { return GetString(Header->TargetPlatform); }
 	FUtf8StringView			ChunksDirectory() const { return GetString(Header->ChunksDirectory); }
+	FUtf8StringView			CompressionFormat() const { return GetString(Header->CompressionFormat); }
 	UE_API FUtf8StringView	GetString(FOnDemandStringEntry StringEntry) const;
 
 	TConstArrayView<FOnDemandContainerEntry> Containers() const { return ContainerEntries; }
diff --git a/Engine/Source/Runtime/Experimental/IoStore/OnDemandCore/Private/Serialization/OnDemandContainerToc.cpp b/Engine/Source/Runtime/Experimental/IoStore/OnDemandCore/Private/Serialization/OnDemandContainerToc.cpp
--- a/Engine/Source/Runtime/Experimental/IoStore/OnDemandCore/Private/Serialization/OnDemandContainerToc.cpp
+++ b/Engine/Source/Runtime/Experimental/IoStore/OnDemandCore/Private/Serialization/OnDemandContainerToc.cpp
@@ -124,9 +124,9 @@ public:
 				TConstArrayView<FIoBlockHash>	LegacyBlockHashes = LegacyContainerBlockHashes.Mid(LegacyEntry.BlockOffset, LegacyEntry.BlockCount);
 
 				ChunkIds.Add(LegacyEntry.ChunkId);
-				
+
 				FOnDemandChunkBlockInfo ChunkBlockInfo =
-					FOnDemandChunkBlockInfo::FromOffsetAndCount(LegacyEntry.BlockOffset, LegacyEntry.BlockCount);
+					FOnDemandChunkBlockInfo::FromOffsetAndCount(uint32(BlockSizes.Num()), LegacyEntry.BlockCount);
 
 				// If the chunk only has a single compression block, the size and block hash are encoded into the chunk entry
 				if (LegacyBlockSizes.Num() > 1)
@@ -175,7 +175,8 @@ public:
 			ContainerEntry.TagSetCount			= TagSets.Num();
 			ContainerEntry.TagSetIndicesCount	= TagSetIndices.Num();
 			ContainerEntry.UTocHash				= LegacyContainer.UTocHash;
-			ContainerEntry.ContainerFlags 		= LegacyContainer.ContainerFlags;
+			ContainerEntry.ContainerFlags		= Toc.Header.Flags;
+			ContainerEntry.FileContainerFlags 	= LegacyContainer.ContainerFlags;
 
 			if (LegacyContainer.EncryptionKeyGuid.IsEmpty() == false)
 			{
@@ -309,9 +310,9 @@ namespace V1
 {
 
 ////////////////////////////////////////////////////////////////////////////////
-EIoContainerFlags FOnDemandContainerTocHeaderView::ContainerFlags() const
+EIoContainerFlags FOnDemandContainerTocHeaderView::FileContainerFlags() const
 {
-	return static_cast<EIoContainerFlags>(ContainerEntry->ContainerFlags);
+	return static_cast<EIoContainerFlags>(ContainerEntry->FileContainerFlags);
 }
 
 FUtf8StringView	FOnDemandContainerTocHeaderView::GetString(FOnDemandStringEntry StringEntry) const
