Freeze/Unfreeze optimizations during buble resim

#Synced-CL 49917485

diff --git a/Engine/Source/Runtime/Experimental/Chaos/Private/Chaos/ChaosMarshallingManager.cpp b/Engine/Source/Runtime/Experimental/Chaos/Private/Chaos/ChaosMarshallingManager.cpp
--- a/Engine/Source/Runtime/Experimental/Chaos/Private/Chaos/ChaosMarshallingManager.cpp
+++ b/Engine/Source/Runtime/Experimental/Chaos/Private/Chaos/ChaosMarshallingManager.cpp
@@ -297,6 +297,15 @@ void FChaosMarshallingManager::SetHistoryLength_Internal(int32 InHistoryLength)
 	}
 }
 
+void FChaosMarshallingManager::DebugHistory_Internal(int32 NumFrames)
+{
+	const int32 UseNumFrames = FMath::Min(NumFrames, HistoryQueue_Internal.Num());
+	for (int32 Idx = 0; Idx < UseNumFrames; ++Idx)
+	{
+		UE_LOG(LogChaos, Log, TEXT("					History internal[%d] at frame %d"), Idx, HistoryQueue_Internal[Idx]->InternalStep);
+	}
+}
+
 TArray<FPushPhysicsData*> FChaosMarshallingManager::StealHistory_Internal(int32 NumFrames)
 {
 	ensure(NumFrames <= HistoryQueue_Internal.Num());
diff --git a/Engine/Source/Runtime/Experimental/Chaos/Private/Chaos/PBDRigidsEvolution.cpp b/Engine/Source/Runtime/Experimental/Chaos/Private/Chaos/PBDRigidsEvolution.cpp
--- a/Engine/Source/Runtime/Experimental/Chaos/Private/Chaos/PBDRigidsEvolution.cpp
+++ b/Engine/Source/Runtime/Experimental/Chaos/Private/Chaos/PBDRigidsEvolution.cpp
@@ -1128,70 +1128,132 @@ namespace Chaos
 		BroadPhaseConfig.BroadphaseType = DefaultBroadphaseType;
 	}
 
-	void FPBDRigidsEvolutionBase::SetParticleObjectState(FPBDRigidParticleHandle* Particle, EObjectStateType ObjectState)
+	void FPBDRigidsEvolutionBase::SetParticlesObjectState(TArrayView<FPBDRigidParticleHandle*> RigidParticles, 
+		const EObjectStateType ObjectState, const bool bUpdateViews)
 	{
-		const EObjectStateType InitialState = Particle->ObjectState();
-		const bool bWasDynamic = (InitialState == EObjectStateType::Dynamic) || (InitialState == EObjectStateType::Sleeping);
-		const bool bIsDynamic = (ObjectState == EObjectStateType::Dynamic) || (ObjectState == EObjectStateType::Sleeping);
-		const bool bIsSleeping = (ObjectState == EObjectStateType::Sleeping);
-
-		Particle->SetObjectStateLowLevel(ObjectState);
-
-		// Put the particle into the correct array. This will change the location of the particle data
-		if (auto ClusteredParticle = Particle->CastToClustered())
+		if (RigidParticles.IsEmpty())
 		{
-			Particles.SetClusteredParticleSOA(ClusteredParticle);
+			return;
 		}
-		else
+		const bool bIsDynamic = (ObjectState == EObjectStateType::Dynamic) || (ObjectState == EObjectStateType::Sleeping);
+		const bool bIsSleeping = (ObjectState == EObjectStateType::Sleeping);
+		
+		// Sorting algorithm are modifying the array view in place and is O(N-ParticleOffset), N being the view size
+		int32 ParticleOffset = 0;
+		
+		// The particles in [0-ParticleOffset) will be the valid ones the [ParticleOffset,NumParticles) will be the others
+		Private::ReorderParticleView(RigidParticles, [](FPBDRigidParticleHandle* RigidParticle) -> bool 
+			{ return (RigidParticle != nullptr);},ParticleOffset);
+		
+		const int32 NumValids = ParticleOffset;
+		if (NumValids > 0)
 		{
-			Particles.SetDynamicParticleSOA(Particle);
-		}
+			TArrayView<FPBDRigidParticleHandle*> ValidParticles(&RigidParticles[0], NumValids);
+			ParticleOffset = 0;
 
-		if (InitialState != ObjectState && !Particle->Disabled())
-		{
-			// If we were just put to sleep or made kinematic, we should still report info back to GT
-			// @todo(chaos): why? add reason to the comments
-			if (ObjectState != EObjectStateType::Dynamic)
+			// The particles in [0-ParticleOffset) will be the clustered one whereas the [ParticleOffset,NumParticles) will be the others
+			Private::ReorderParticleView(ValidParticles, [](FPBDRigidParticleHandle* RigidParticle) -> bool
+				{ return (RigidParticle->CastToClustered() != nullptr);}, ParticleOffset);
+		
+			// Move particles to clustered SOA
+			const int32 NumClustered = ParticleOffset;
+			if(NumClustered > 0)
 			{
-				Particles.MarkTransientDirtyParticle(Particle);
+				TArrayView<FPBDRigidParticleHandle*> ClusteredParticles(&ValidParticles[0], NumClustered);
+				Particles.SetClusteredParticlesSOA(ClusteredParticles, ObjectState, false);
 			}
-
-			// Reset the sleep counter when the state changes
-			// NOTE: this sleep counter is only used for particles
-			// that are not in the graph (have no constraints on them)
-			if (bIsDynamic)
+		
+			// Move particles to dynamics SOA
+			const int32 NumDynamics = ValidParticles.Num()-NumClustered;
+			if (NumDynamics > 0)
 			{
-				Particle->SetSleepCounter(0);
+				TArrayView<FPBDRigidParticleHandle*> DynamicParticles(&ValidParticles[NumClustered], NumDynamics);
+				Particles.SetDynamicParticlesSOA(DynamicParticles, ObjectState, false);
 			}
-
-			// If all particles in an island are explicitly put to sleep, the island and all its
-			// constraints need to sleep at the start of the next tick. This is because we don't
-			// detect collisions on sleeping particles, and we destroy collisions that are 
-			// awake and were not updated this tick.
-			if (bIsSleeping)
+		}
+		
+		ParticleOffset = 0;
+			
+		// The particles in [0-ParticleOffset) will be the enabled ones switching to the new state
+		Private::ReorderParticleView(RigidParticles, [&ObjectState](FPBDRigidParticleHandle* RigidParticle) -> bool 
+			{ return RigidParticle && (RigidParticle->ObjectState() != ObjectState && !RigidParticle->Disabled());}, ParticleOffset);	
+		
+		const int32 NumTransients = ParticleOffset;
+		if (NumTransients > 0)
+		{
+			// Mark transient particles dirty
+			TArrayView<FPBDRigidParticleHandle*> TransientParticles(&RigidParticles[0], NumTransients);
+			if (ObjectState != EObjectStateType::Dynamic)
 			{
-				IslandManager.SleepParticle(Particle);
+				Particles.MarkTransientDirtyParticles(TransientParticles, false);
 			}
+			
+			for (FPBDRigidParticleHandle* Particle : TransientParticles)
+			{
+				if (Particle)
+				{
+					// Reset the sleep counter when the state changes
+					// NOTE: this sleep counter is only used for particles
+					// that are not in the graph (have no constraints on them)
+					if (bIsDynamic)
+					{
+						Particle->SetSleepCounter(0);
+					}
 
-			// If we are not kinematic, the MovingKinematic flag should be cleared
-			// If we are kinematic and are (or become) moving, the flag will be updated in ApplyKinematicTargets
-			// Either way, we can reset the flag here
-			Particle->ClearIsMovingKinematic();
-		}
+					// If all particles in an island are explicitly put to sleep, the island and all its
+					// constraints need to sleep at the start of the next tick. This is because we don't
+					// detect collisions on sleeping particles, and we destroy collisions that are 
+					// awake and were not updated this tick.
+					if (bIsSleeping)
+					{
+						IslandManager.SleepParticle(Particle);
+					}
 
-		// If we are now dynamic and enabled, we need to be in the graph if not already there. We may not be
-		// because Kinematic particles are only in the graph if referenced by a constraint.
-		if (bIsDynamic && !bWasDynamic && !Particle->Disabled())
-		{
-			IslandManager.AddParticle(Particle);
+					// If we are not kinematic, the MovingKinematic flag should be cleared
+					// If we are kinematic and are (or become) moving, the flag will be updated in ApplyKinematicTargets
+					// Either way, we can reset the flag here
+					Particle->ClearIsMovingKinematic();
+				}
+			}
 		}
-
-		// If we switched to/from dynamic the inertia conditioning may need to be refreshed
-		if (Particle->InertiaConditioningEnabled() && (bWasDynamic != bIsDynamic))
+		
+		if(bUpdateViews)
+		{ 
+			Particles.UpdateDirtyViews();
+		}
+		
+		// Modify particle object state
+		int32 ParticleIndex = 0;
+		for (FPBDRigidParticleHandle* Particle : RigidParticles)
 		{
-			Particle->SetInertiaConditioningDirty();
+			if (Particle)
+			{
+				const bool bWasDynamic = (Particle->ObjectState() == EObjectStateType::Dynamic) || (Particle->ObjectState() == EObjectStateType::Sleeping);
+			
+				// If we are now dynamic and enabled, we need to be in the graph if not already there. We may not be
+				// because Kinematic particles are only in the graph if referenced by a constraint.
+				if (bIsDynamic && !bWasDynamic && !Particle->Disabled())
+				{
+					IslandManager.AddParticle(Particle);
+				}
+
+				// If we switched to/from dynamic the inertia conditioning may need to be refreshed
+				if (Particle->InertiaConditioningEnabled() && (bWasDynamic != bIsDynamic))
+				{
+					Particle->SetInertiaConditioningDirty();
+				}
+				
+				// Change the low level object state 
+				Particle->SetObjectStateLowLevel(ObjectState);
+			}
 		}
 	}
+	
+	void FPBDRigidsEvolutionBase::SetParticleObjectState(FPBDRigidParticleHandle* Particle, EObjectStateType ObjectState)
+	{
+		TArrayView<FPBDRigidParticleHandle*> RigidParticles(&Particle, 1);
+		SetParticlesObjectState(RigidParticles, ObjectState, true);
+	}
 
 	void FPBDRigidsEvolutionBase::WakeParticle(FPBDRigidParticleHandle* Particle)
 	{
diff --git a/Engine/Source/Runtime/Experimental/Chaos/Private/RewindData.cpp b/Engine/Source/Runtime/Experimental/Chaos/Private/RewindData.cpp
--- a/Engine/Source/Runtime/Experimental/Chaos/Private/RewindData.cpp
+++ b/Engine/Source/Runtime/Experimental/Chaos/Private/RewindData.cpp
@@ -33,7 +33,7 @@ namespace Chaos
 		static FAutoConsoleVariableRef CVarBubbleResimEnableCheapFreeze(
 			TEXT("p.Chaos.BubbleResim.EnableCheapFreeze"),
 			BubbleResimEnableCheapFreeze,
-			TEXT("Use SetObjectStateLowLevel instead of going through the evolution to modify the object state, to freeze/unfreeze a particle"));
+			TEXT("Use batched SetObjectState version instead of going through every single particles to modify the object state, to freeze/unfreeze a particle"));
 
 		float BubbleResimExtraSweptDistance = 10.0f;
 		static FAutoConsoleVariableRef CVarBubbleResimExtraSweptDistance(
@@ -671,29 +671,37 @@ void FRewindData::FreezeParticlesOutsideResimBubbles()
 	}
 
 	FrozenParticles.Reserve(ParticlesToFreezeInResim.Num());
+
+	TArray<FPBDRigidParticleHandle*> ParticleHandles;
+	if (Chaos::ResimConsoleVars::BubbleResimEnableCheapFreeze)
+	{
+		ParticleHandles.Reserve(ParticlesToFreezeInResim.Num());
+	}
+	const EObjectStateType FrozenObjectState = (ResimConsoleVars::BubbleResimFreezeMode == 0) ? EObjectStateType::Static : EObjectStateType::Kinematic;
 	for (int32 ParticleToFreezeIndex = 0; ParticleToFreezeIndex < ParticlesToFreezeInResim.Num(); ++ParticleToFreezeIndex)
 	{
 		FPBDRigidParticleHandle* RigidParticleHandle = ParticlesToFreezeInResim[ParticleToFreezeIndex];
 		if (IPhysicsProxyBase* PhysicsProxy = RigidParticleHandle->PhysicsProxy())
 		{
 			FrozenParticles.Add({PhysicsProxy, RigidParticleHandle->ObjectState()});
-		}
-		EObjectStateType FrozenObjectState = (ResimConsoleVars::BubbleResimFreezeMode == 0) ? EObjectStateType::Static : EObjectStateType::Kinematic;
-		if (Chaos::ResimConsoleVars::BubbleResimEnableCheapFreeze)
-		{
-			RigidParticleHandle->SetObjectStateLowLevel(FrozenObjectState);
-		}
-		else
-		{
-			// This is too slow, e.g. it calls UpdateViews every time, but it might be more correct
-			// TODO: extract the minimum of what's needed from SetParticleObjectState
-			Evolution->SetParticleObjectState(RigidParticleHandle, FrozenObjectState); 
+
+			if (Chaos::ResimConsoleVars::BubbleResimEnableCheapFreeze)
+			{
+				ParticleHandles.Add(RigidParticleHandle);
+			}
+			else
+			{
+				// This is too slow, e.g. it calls UpdateViews every time, but it might be more correct
+				// TODO: extract the minimum of what's needed from SetParticleObjectState
+				Evolution->SetParticleObjectState(RigidParticleHandle, FrozenObjectState); 
+			}
 		}
 	}
 	if (Chaos::ResimConsoleVars::BubbleResimEnableCheapFreeze)
 	{
-		QUICK_SCOPE_CYCLE_COUNTER(ChaosRewindAndResim_UpdateDirtyViews);
-		Evolution->GetParticles().UpdateDirtyViews();
+		QUICK_SCOPE_CYCLE_COUNTER(ChaosRewindAndResim_UpdateObjectsState);
+		TArrayView<FPBDRigidParticleHandle*> DynamicParticles(ParticleHandles.GetData(), ParticleHandles.Num());
+		Evolution->SetParticlesObjectState(DynamicParticles, FrozenObjectState, true);
 	}
 }
 
@@ -705,9 +713,18 @@ void FRewindData::UnfreezeParticlesOutsideResimBubbles()
 	}
 	
 	FPBDRigidsEvolution* Evolution = Solver->GetEvolution();
+	
 	// Revert to pre-resim object state
 	{
 		QUICK_SCOPE_CYCLE_COUNTER(ChaosRewindAndResim_RevertObjectState);
+		
+		if(Chaos::ResimConsoleVars::BubbleResimEnableCheapFreeze)
+		{
+			FrozenDynamicHandles.Reset(FrozenParticles.Num());
+			FrozenSleepingHandles.Reset(FrozenParticles.Num());
+			FrozenStaticHandles.Reset(FrozenParticles.Num());
+			FrozenKinematicHandles.Reset(FrozenParticles.Num());
+		}
 		for (TPair<Chaos::FPhysicsProxyHandle, EObjectStateType>& FrozenParticle : FrozenParticles)
 		{
 			if (IPhysicsProxyBase* Proxy = FrozenParticle.Key.Get())
@@ -720,9 +737,22 @@ void FRewindData::UnfreezeParticlesOutsideResimBubbles()
 					{
 						if (Chaos::ResimConsoleVars::BubbleResimEnableCheapFreeze)
 						{
-							RigidParticleHandle->SetObjectStateLowLevel(PreResimObjectState);
-							QUICK_SCOPE_CYCLE_COUNTER(ChaosRewindAndResim_AddToIslandManager);
-							Evolution->GetIslandManager().AddParticle(RigidParticleHandle);
+							if(PreResimObjectState == EObjectStateType::Dynamic)
+							{
+								FrozenDynamicHandles.Add(RigidParticleHandle);
+							}
+							else if(PreResimObjectState == EObjectStateType::Sleeping)
+							{
+								FrozenSleepingHandles.Add(RigidParticleHandle);
+							}
+							if (PreResimObjectState == EObjectStateType::Static)
+							{
+								FrozenStaticHandles.Add(RigidParticleHandle);
+							}
+							else if (PreResimObjectState == EObjectStateType::Kinematic)
+							{
+								FrozenKinematicHandles.Add(RigidParticleHandle);
+							}
 						}
 						else
 						{
@@ -741,6 +771,19 @@ void FRewindData::UnfreezeParticlesOutsideResimBubbles()
 			}
 		}
 	}
+	if (Chaos::ResimConsoleVars::BubbleResimEnableCheapFreeze)
+	{
+		QUICK_SCOPE_CYCLE_COUNTER(ChaosRewindAndResim_ReportObjectsState);
+		
+		// Update the frozen particles handles with the value pre-resim
+		Evolution->SetParticlesObjectState(FrozenDynamicHandles, EObjectStateType::Dynamic, false);
+		Evolution->SetParticlesObjectState(FrozenSleepingHandles, EObjectStateType::Sleeping, false);
+		Evolution->SetParticlesObjectState(FrozenStaticHandles, EObjectStateType::Static, false);
+		Evolution->SetParticlesObjectState(FrozenKinematicHandles, EObjectStateType::Kinematic, false);
+	
+		// Update the dirty views
+		Evolution->GetParticles().UpdateDirtyViews();
+	}
 	{
 		QUICK_SCOPE_CYCLE_COUNTER(ChaosRewindAndResim_MarkAllFrozenInResimByDefault);
 		// Mark all particles as frozen in resim (default value)
@@ -750,11 +793,6 @@ void FRewindData::UnfreezeParticlesOutsideResimBubbles()
 			Particle.SetResimType(EResimType::FrozenDuringResim);
 		}
 	}
-	if (Chaos::ResimConsoleVars::BubbleResimEnableCheapFreeze)
-	{
-		QUICK_SCOPE_CYCLE_COUNTER(ChaosRewindAndResim_UpdateDirtyViews);
-		Evolution->GetParticles().UpdateDirtyViews();
-	}
 	// Empty frozen particles so we can start from scratch on the next rewind/resims
 	FrozenParticles.Empty();
 }
diff --git a/Engine/Source/Runtime/Experimental/Chaos/Public/Chaos/ChaosMarshallingManager.h b/Engine/Source/Runtime/Experimental/Chaos/Public/Chaos/ChaosMarshallingManager.h
--- a/Engine/Source/Runtime/Experimental/Chaos/Public/Chaos/ChaosMarshallingManager.h
+++ b/Engine/Source/Runtime/Experimental/Chaos/Public/Chaos/ChaosMarshallingManager.h
@@ -483,6 +483,9 @@ public:
 	/** Returns the history buffer of the latest NumFrames. The history that comes before these frames is discarded*/
 	CHAOS_API TArray<FPushPhysicsData*> StealHistory_Internal(int32 NumFrames);
 
+	/** Debug and print the history buffer of the latest NumFrames. */
+	CHAOS_API void DebugHistory_Internal(int32 NumFrames);
+
 	/** Return the size of the history queue */
 	int32 GetNumHistory_Internal() const {return HistoryQueue_Internal.Num();}
 
diff --git a/Engine/Source/Runtime/Experimental/Chaos/Public/Chaos/PBDRigidsEvolution.h b/Engine/Source/Runtime/Experimental/Chaos/Public/Chaos/PBDRigidsEvolution.h
--- a/Engine/Source/Runtime/Experimental/Chaos/Public/Chaos/PBDRigidsEvolution.h
+++ b/Engine/Source/Runtime/Experimental/Chaos/Public/Chaos/PBDRigidsEvolution.h
@@ -636,6 +636,9 @@ public:
 	}
 
 	CHAOS_API void SetParticleObjectState(FPBDRigidParticleHandle* Particle, EObjectStateType ObjectState);
+	
+	/** Change the object state for a list of particles */
+	CHAOS_API void SetParticlesObjectState(TArrayView<FPBDRigidParticleHandle*> Particles, const EObjectStateType ObjectState, const bool bUpdateViews);
 
 	// Wake a dynamic particle and reset sleep counters for its island
 	CHAOS_API void WakeParticle(FPBDRigidParticleHandle* Particle);
diff --git a/Engine/Source/Runtime/Experimental/Chaos/Public/Chaos/PBDRigidsSOAs.h b/Engine/Source/Runtime/Experimental/Chaos/Public/Chaos/PBDRigidsSOAs.h
--- a/Engine/Source/Runtime/Experimental/Chaos/Public/Chaos/PBDRigidsSOAs.h
+++ b/Engine/Source/Runtime/Experimental/Chaos/Public/Chaos/PBDRigidsSOAs.h
@@ -12,6 +12,29 @@ namespace CVars
 	extern CHAOS_API bool bRemoveParticleFromMovingKinematicsOnDisable;
 	extern CHAOS_API bool bChaosSolverCheckParticleViews;
 }
+	
+namespace Private
+{
+	// The goal of that function is to reorder a particle view per category in place without reallocation in O(N). 
+	// The different categories should not overlap (example : kinematic vs dynamic vs static...)
+	template<typename CategoryFunc>
+	FORCEINLINE void ReorderParticleView(TArrayView<FPBDRigidParticleHandle*>& ParticlesView, const CategoryFunc& CategoryTrigger, int32& CategoryOffset)
+	{		
+		const int32 NumParticles = ParticlesView.Num();
+		for (int32 ParticleIndex = CategoryOffset; ParticleIndex < NumParticles; ++ParticleIndex)
+		{
+			if (CategoryTrigger(ParticlesView[ParticleIndex]))
+			{
+				if (ParticleIndex != CategoryOffset)
+				{
+					Swap(ParticlesView[ParticleIndex], ParticlesView[CategoryOffset]);
+				}
+				++CategoryOffset;
+			}
+		}
+	}
+}
+
 
 // A helper function to get the debug name for a particle or "UNKNOWN" when debug names are unavailable
 template<typename TParticleType>
@@ -143,10 +166,20 @@ public:
 	{
 		return ParticleToIndex.Contains(Particle);
 	}
-
-	template <typename TParticle1>
-	void Insert(const TArray<TParticle1*>& ParticlesToInsert)
+	
+	/** Check if the particle handle has a type compatible with the map array */
+	bool HasValidType(FPBDRigidParticleHandle* ParticleHandle)
 	{
+		return ParticleHandle && (ParticleHandle->GetParticleType() >= TParticleType::StaticType());
+	}
+	
+	/** Add a list of particles to the map array in O(k) */
+	void Insert(const TArrayView<FPBDRigidParticleHandle*>& ParticlesToInsert)
+	{
+		if (ParticlesToInsert.IsEmpty())
+		{
+			return;
+		}
 		TArray<bool> Contains;
 		Contains.AddZeroed(ParticlesToInsert.Num());
 
@@ -154,19 +187,22 @@ public:
 		int32 NextIdx = ParticleArray.Num();
 		for (int32 Idx = 0; Idx < ParticlesToInsert.Num(); ++Idx)
 		{
-			auto Particle = ParticlesToInsert[Idx];
-			Contains[Idx] = ParticleToIndex.Contains(Particle);
-			if (!Contains[Idx])
+			if (HasValidType(ParticlesToInsert[Idx]))
 			{
-				ParticleToIndex.Add(Particle, NextIdx++);
+				TParticleType* Particle = static_cast<TParticleType*>(ParticlesToInsert[Idx]);
+				Contains[Idx] = ParticleToIndex.Contains(Particle);
+				if (!Contains[Idx])
+				{
+					ParticleToIndex.Add(Particle, NextIdx++);
+				}
 			}
 		}
 		ParticleArray.Reserve(ParticleArray.Num() + NextIdx - ParticleArray.Num());
 		for (int32 Idx = 0; Idx < ParticlesToInsert.Num(); ++Idx)
 		{
-			if (!Contains[Idx])
+			if (!Contains[Idx] && HasValidType(ParticlesToInsert[Idx]))
 			{
-				auto Particle = ParticlesToInsert[Idx];
+				TParticleType* Particle = static_cast<TParticleType*>(ParticlesToInsert[Idx]);
 				ParticleArray.Add(Particle);
 			}
 		}
@@ -181,6 +217,34 @@ public:
 		}
 	}
 
+	/** Remove a list of particles from the map array in O(k) (Find, RemoveAtSwap and Remove are O(1)) */
+	void Remove(const TArrayView<FPBDRigidParticleHandle*>& ParticlesToRemove)
+	{
+		if (ParticlesToRemove.IsEmpty())
+		{
+			return;
+		}
+		for (int32 Idx = 0; Idx < ParticlesToRemove.Num(); ++Idx)
+		{
+			if (HasValidType(ParticlesToRemove[Idx]))
+			{
+				const TParticleType* Particle = static_cast<TParticleType*>(ParticlesToRemove[Idx]);
+				if (int32* IdxPtr = ParticleToIndex.Find(Particle))
+				{
+					int32 ParticleIdx = *IdxPtr;
+					ParticleArray.RemoveAtSwap(ParticleIdx, EAllowShrinking::No);
+					if (ParticleIdx < ParticleArray.Num())
+					{
+						//update swapped element with new index
+						ParticleToIndex[ParticleArray[ParticleIdx]] = ParticleIdx;
+					}
+					ParticleToIndex.Remove(Particle);
+				}
+			}
+		}
+		ParticleArray.Shrink();
+	}
+
 	void Remove(TParticleType* Particle)
 	{
 		if (int32* IdxPtr = ParticleToIndex.Find(Particle))
@@ -415,20 +479,28 @@ public:
 	{
 		TransientDirtyMapArray.Reset();
 	}
-
-	// @todo(chaos): keep track of which views may be dirty and lazily update them
-	// rather than calling UpdateViews all the time
-	void MarkTransientDirtyParticle(FGeometryParticleHandle* Particle, const bool bUpdateViews = true)
+	void MarkTransientDirtyParticles(TArrayView<FPBDRigidParticleHandle*>& RigidParticles, const bool bUpdateViews = true)
 	{
-		if(FPBDRigidParticleHandle* Rigid =  Particle->CastToRigidParticle())
+		if (!RigidParticles.IsEmpty())
 		{
-			//if active it's already in the dirty view
-			if (!ActiveParticlesMapArray.Contains(Rigid) && !MovingKinematicsMapArray.Contains(Rigid))
-			{
-				AddToList(Rigid, TransientDirtyMapArray);
+			int32 ParticleOffset = 0;
+			
+			// The particles in [0-ParticleOffset) will be the enabled ones switching to the new state
+			Private::ReorderParticleView(RigidParticles, [this](FPBDRigidParticleHandle* RigidParticle) -> bool 
+				{ return (!ActiveParticlesMapArray.Contains(RigidParticle) && !MovingKinematicsMapArray.Contains(RigidParticle));},ParticleOffset);	
 		
-				CheckParticleViewMask(Particle);
-
+			const int32 NumTransients = ParticleOffset;
+			if (NumTransients > 0)
+			{
+				// Mark transient particles dirty
+				TArrayView<FPBDRigidParticleHandle*> TransientParticles(&RigidParticles[0], NumTransients);
+				AddToList(TransientParticles, TransientDirtyMapArray);
+			
+				for (FPBDRigidParticleHandle* Particle : TransientParticles)
+				{
+					CheckParticleViewMask(Particle);
+				}
+			
 				if (bUpdateViews)
 				{
 					UpdateViews();
@@ -436,6 +508,18 @@ public:
 			}
 		}
 	}
+	
+
+	// @todo(chaos): keep track of which views may be dirty and lazily update them
+	// rather than calling UpdateViews all the time
+	void MarkTransientDirtyParticle(FGeometryParticleHandle* GeometryParticle, const bool bUpdateViews = true)
+	{
+		if(FPBDRigidParticleHandle* RigidParticle =  GeometryParticle->CastToRigidParticle())
+		{
+			TArrayView<FPBDRigidParticleHandle*> RigidParticles(&RigidParticle, 1);
+			MarkTransientDirtyParticles(RigidParticles, bUpdateViews);
+		}
+	}
 
 	void DestroyParticle(FGeometryParticleHandle* Particle)
 	{
@@ -691,127 +775,248 @@ public:
 		UpdateViews();
 	}
 
-	void SetDynamicParticleSOA(FPBDRigidParticleHandle* Particle)
+	template<typename T, int d>
+	void MoveParticlesToSOA(TArrayView<FPBDRigidParticleHandle*>& RigidParticles, TGeometryParticles<T, d>& ToSOA)
 	{
-		const EObjectStateType State = Particle->ObjectState();
-
-		if (Particle->Disabled())
+		if (!RigidParticles.IsEmpty() && RigidParticles[0])
 		{
-			Particle->MoveToSOA(*DynamicDisabledParticles);
-			RemoveFromActiveArray(Particle->CastToRigidParticle(), /*bStillDirty=*/ false);
+			int32 ParticleOffset = 0;
+			TGeometryParticles<T, d>* FromSOA = RigidParticles[0]->GeometryParticles;
+			
+			// Only allocation in the SetObjectState pipeline but we need to transfer particles->indices anyway
+			TArray<int32> ParticleIndices;
+			ParticleIndices.Init(0, RigidParticles.Num());
+			
+			while (ParticleOffset < RigidParticles.Num())
+			{
+				const uint32 PreviousOffset = ParticleOffset;
+			
+				// Reorder all the particles that have the same SOA as the FromSOA among the remaining ones 
+				// [ParticleOffset+1...RigidParticles.Num()] and push them at the front of the array
+				Private::ReorderParticleView(RigidParticles, [&FromSOA](FPBDRigidParticleHandle* RigidParticle) -> bool
+					{
+						return RigidParticle && (RigidParticle->GeometryParticles == FromSOA);
+					}, ParticleOffset);
+				
+				if (FromSOA != &ToSOA)
+				{
+					// Fill partially the particle indices that we want to move to the target SOA
+					for (int32 ViewIndex = PreviousOffset; ViewIndex < ParticleOffset; ++ViewIndex)
+					{
+						ParticleIndices[ViewIndex] = RigidParticles[ViewIndex]->ParticleIdx;
+					}
+					// Build the particle indices view used to update the SOAs
+					const int32 NumParticles = ParticleOffset - PreviousOffset;
+					TArrayView<int32> IndicesView(&ParticleIndices[PreviousOffset], NumParticles);
+					
+					// For each SOA move the particles to the target SOA
+					uint32 OffsetIndex = ToSOA.Size();
+					FromSOA->MoveToOtherParticles(IndicesView, ToSOA);
+					check(ToSOA.Size() == (OffsetIndex + NumParticles))
+	
+					// Update the particle indices for the removed/added particles
+					for(int32 GeometryIndex : IndicesView)
+					{
+						// Particles that have been moved with the swap should be updated with the right index
+						if (static_cast<uint32>(GeometryIndex) < FromSOA->Size())
+						{
+							FromSOA->Handle(GeometryIndex)->ParticleIdx = GeometryIndex;
+						}
+						// Particles SOA index and mask are updated
+						if(OffsetIndex < ToSOA.Size())
+						{ 
+							ToSOA.Handle(OffsetIndex)->ParticleIdx = OffsetIndex;
+							ToSOA.Handle(OffsetIndex)->GeometryParticles = &ToSOA;
+							
+							ToSOA.ListMask(OffsetIndex) &= ~FromSOA->GetContainerListMask();
+							ToSOA.ListMask(OffsetIndex) |= ToSOA.GetContainerListMask();
+						}
+						++OffsetIndex;
+					}
+				}
+				if (ParticleOffset < RigidParticles.Num())
+				{
+					FromSOA = RigidParticles[ParticleOffset]->GeometryParticles;
+				}
+			}
+		}
+	}
+	
+	/** Update active particles SOAs, list and views */
+	void UpdateActiveParticlesSOA(TArrayView<FPBDRigidParticleHandle*>& RigidParticles, const EObjectStateType ObjectState, int32& ParticleOffset)
+	{	
+		ParticleOffset = 0;
+		
+		// The particles in [0-ParticleOffset) will be the disabled one whereas the [ParticleOffset,NumParticles) will be the enabled
+		Private::ReorderParticleView(RigidParticles, [](FPBDRigidParticleHandle* RigidParticle) -> bool 
+			{ return RigidParticle && RigidParticle->Disabled(); },ParticleOffset);
+		
+		const int32 NumDisabled = ParticleOffset;
+		if (NumDisabled > 0)
+		{
+			TArrayView<FPBDRigidParticleHandle*> DisabledParticles(&RigidParticles[0], NumDisabled);
+			
+			RemoveFromActiveArray(DisabledParticles, /*bStillDirty=*/ false);
 			if (CVars::bRemoveParticleFromMovingKinematicsOnDisable)
 			{
-				RemoveFromList(Particle, MovingKinematicsMapArray);
+				RemoveFromList(DisabledParticles, MovingKinematicsMapArray);
 			}
 		}
-		else
+		
+		const int32 NumEnabled= RigidParticles.Num()-ParticleOffset;
+		if (NumEnabled > 0)
 		{
-			if (Particle->ObjectState() != EObjectStateType::Dynamic)
+			TArrayView<FPBDRigidParticleHandle*> EnabledParticles(&RigidParticles[ParticleOffset], NumEnabled);
+			if (ObjectState != EObjectStateType::Dynamic)
 			{
-				RemoveFromActiveArray(Particle->CastToRigidParticle(), /*bStillDirty=*/true);
+				RemoveFromActiveArray(EnabledParticles, /*bStillDirty=*/true);
 			}
 			else
 			{
-				AddToActiveArray(Particle->CastToRigidParticle());
+				AddToActiveArray(EnabledParticles);
 			}
 
-			if (Particle->ObjectState() != EObjectStateType::Kinematic)
+			if (ObjectState != EObjectStateType::Kinematic)
 			{
-				RemoveFromList(Particle, MovingKinematicsMapArray);
+				RemoveFromList(EnabledParticles, MovingKinematicsMapArray);
 			}
+		}
+	}
+	/** Function to be executed after the particles are moved to the dynamic/clustered SOAs */
+	void PostMoveParticlesSOA(TArrayView<FPBDRigidParticleHandle*>& RigidParticles, const bool bUpdateViews)
+	{
+		// Check if the view mask is correct only for no shipping build
+		for(FPBDRigidParticleHandle* RigidParticle : RigidParticles)
+		{ 
+			CheckParticleViewMask(RigidParticle);
+		}
 
-			if (Particle->Type == EParticleType::Clustered)
-			{
-				Particle->MoveToSOA(*ClusteredParticles);
-				if (bResimulating)
-				{
-					RemoveFromList(Particle, ResimDynamicKinematicParticles);
-					AddToList(Particle, ResimDynamicParticles);
-				}
-			}
-			else if (Particle->Type == EParticleType::GeometryCollection)
-			{
-				Particle->MoveToSOA(*GeometryCollectionParticles);
-				if (bResimulating)
-				{
-					RemoveFromList(Particle, ResimDynamicKinematicParticles);
-					AddToList(Particle, ResimDynamicParticles);
-				}
-			}
-			else
+		// Update the views if necessary
+		if(bUpdateViews)
+		{ 
+			UpdateViews();
+		}
+	}
+
+	/** Update dynamic particles SOAs, list and views */
+	void SetDynamicParticlesSOA(TArrayView<FPBDRigidParticleHandle*>& RigidParticles, const EObjectStateType ObjectState, const bool bUpdateViews)
+	{
+        int32 ParticleOffset = 0;
+		
+		// Update the active particles SOA
+		UpdateActiveParticlesSOA(RigidParticles, ObjectState, ParticleOffset);
+		
+		const int32 NumDisabled = ParticleOffset;
+		if (NumDisabled > 0)
+		{
+			// For disabled particles move them into the dynamic disabled SOA
+			TArrayView<FPBDRigidParticleHandle*> DisabledParticles(&RigidParticles[0], NumDisabled);
+			MoveParticlesToSOA(DisabledParticles, *DynamicDisabledParticles);
+		}
+		
+		const int32 NumEnabled = RigidParticles.Num()-ParticleOffset;
+		if (NumEnabled > 0)
+		{
+			TArrayView<FPBDRigidParticleHandle*> EnabledParticles(&RigidParticles[ParticleOffset], NumEnabled);
+			// For enabled particles move to dynamic/kinematic SOA. No sleeping/static ones so default to dynamic.
+			switch (ObjectState)
 			{
-				// Move to appropriate dynamic SOA
-				switch (State)
-				{
 				case EObjectStateType::Kinematic:
-					Particle->MoveToSOA(*DynamicKinematicParticles);
+					MoveParticlesToSOA(EnabledParticles, *DynamicKinematicParticles);
 					if (bResimulating)
 					{
-						RemoveFromList(Particle, ResimDynamicParticles);
-						AddToList(Particle, ResimDynamicKinematicParticles);
+						RemoveFromList(EnabledParticles, ResimDynamicParticles);
+						AddToList(EnabledParticles, ResimDynamicKinematicParticles);
 					}
 					break;
 
 				case EObjectStateType::Dynamic:
-					Particle->MoveToSOA(*DynamicParticles);
+				default : // TODO: Special SOAs for sleeping and static particles?
+					MoveParticlesToSOA(EnabledParticles, *DynamicParticles);
 					if (bResimulating)
 					{
-						RemoveFromList(Particle, ResimDynamicKinematicParticles);
-						AddToList(Particle, ResimDynamicParticles);
+						RemoveFromList(EnabledParticles, ResimDynamicKinematicParticles);
+						AddToList(EnabledParticles, ResimDynamicParticles);
 					}
 					break;
-
-				default:
-					// TODO: Special SOAs for sleeping and static particles?
-					Particle->MoveToSOA(*DynamicParticles);
-					if (bResimulating)
-					{
-						RemoveFromList(Particle, ResimDynamicKinematicParticles);
-						AddToList(Particle, ResimDynamicParticles);
-					}
-					break;
-				}
 			}
 		}
-
-		CheckParticleViewMask(Particle);
-
-		UpdateViews();
+		PostMoveParticlesSOA(RigidParticles, bUpdateViews);
 	}
 
-	void SetClusteredParticleSOA(FPBDRigidClusteredParticleHandle* ClusteredParticle)
+	/** Update dynamic particle SOAs, list and views */
+	void SetDynamicParticleSOA(FPBDRigidParticleHandle* RigidParticle)
 	{
-		if (ClusteredParticle->ObjectState() != EObjectStateType::Kinematic)
+		if (RigidParticle)
 		{
-			RemoveFromList(ClusteredParticle, MovingKinematicsMapArray);
+			TArrayView<FPBDRigidParticleHandle*> RigidParticles(&RigidParticle, 1);
+			SetDynamicParticlesSOA(RigidParticles, RigidParticle->ObjectState(), true);
 		}
+	}
+	
+	template<typename SOAType>
+	void UpdateClusteredParticlesSOA(TArrayView<FPBDRigidParticleHandle*>& RigidParticles, 
+		const EObjectStateType ObjectState, SOAType& ClusteredSOA, int32& ParticleOffset)
+	{
+		MoveParticlesToSOA(RigidParticles, ClusteredSOA);
 
-		if (TPBDGeometryCollectionParticleHandle<FReal, 3>* GCParticle = ClusteredParticle->CastToGeometryCollection())
+		if (bResimulating)
 		{
-			// Geometry collection particles have their own arrays which are also included in the active view
-			RemoveGeometryCollectionParticle(GCParticle);
-			InsertGeometryCollectionParticle(GCParticle);
-		}
-		else 
-		{ 
-			check(ClusteredParticle->GetParticleType() != Chaos::EParticleType::GeometryCollection);
-			RemoveClusteredParticle(ClusteredParticle);
-			InsertClusteredParticle(ClusteredParticle);
+			RemoveFromList(RigidParticles, ResimDynamicKinematicParticles);
+			AddToList(RigidParticles, ResimDynamicParticles);
 		}
+		
+		ParticleOffset = 0;
+		UpdateActiveParticlesSOA(RigidParticles, ObjectState, ParticleOffset);
+	}
 
-		// Cluster particle must go in the active array because the Cluster arrays are not in the active view
-		if (ClusteredParticle->ObjectState() != EObjectStateType::Dynamic || ClusteredParticle->Disabled())
+	/** Update clustered particle SOAs, list and views */
+	void SetClusteredParticlesSOA(TArrayView<FPBDRigidParticleHandle*>& RigidParticles, const EObjectStateType ObjectState, const bool bUpdateViews, const bool bIsSleeping = false)
+	{
+		int32 ParticleOffset = 0;
+		
+		// The particles in [0-ParticleOffset) will be the geometry collection ones one whereas the [ParticleOffset,NumParticles) will be the others
+		Private::ReorderParticleView(RigidParticles, [](FPBDRigidParticleHandle* RigidParticle) -> bool 
+			{ return RigidParticle && (RigidParticle->CastToGeometryCollection() != nullptr);}, ParticleOffset); 
+	
+		const int32 NumGeometryCollection = ParticleOffset;
+		if (NumGeometryCollection)
 		{
-			RemoveFromActiveArray(ClusteredParticle->CastToRigidParticle(), /*bStillDirty=*/true);
+			TArrayView<FPBDRigidParticleHandle*> GCParticles(&RigidParticles[0], NumGeometryCollection);
+			UpdateClusteredParticlesSOA(GCParticles, ObjectState, *GeometryCollectionParticles, ParticleOffset);
+			RemoveGeometryCollectionParticles(GCParticles);
+			
+			const int32 NumEnabled = GCParticles.Num() - ParticleOffset;
+			if (NumEnabled > 0)
+			{
+				TArrayView<FPBDRigidParticleHandle*> EnabledParticles(&GCParticles[ParticleOffset], NumEnabled);
+				InsertGeometryCollectionParticles(EnabledParticles, false, false, ObjectState);	
+			}		
 		}
-		else
+	
+		const int32 NumClustered = RigidParticles.Num()-NumGeometryCollection;
+		if (NumClustered)
 		{
-			AddToActiveArray(ClusteredParticle->CastToRigidParticle());
+			TArrayView<FPBDRigidParticleHandle*> CLParticles(&RigidParticles[NumGeometryCollection], NumClustered);
+			UpdateClusteredParticlesSOA(CLParticles, ObjectState, *ClusteredParticles, ParticleOffset);
+			RemoveClusteredParticles(CLParticles);
+			
+			const int32 NumEnabled = CLParticles.Num() - ParticleOffset;
+			if (NumEnabled > 0)
+			{
+				TArrayView<FPBDRigidParticleHandle*> EnabledParticles(&CLParticles[ParticleOffset], NumEnabled);
+				InsertClusteredParticles(EnabledParticles, false, ObjectState);	
+			}		
 		}
+		PostMoveParticlesSOA(RigidParticles, bUpdateViews);
+	}
 
-		CheckParticleViewMask(ClusteredParticle);
+	void SetClusteredParticleSOA(FPBDRigidClusteredParticleHandle* ClusteredParticle)
+	{
+		FPBDRigidParticleHandle* RigidParticle = static_cast<FPBDRigidParticleHandle*>(ClusteredParticle);
+		TArrayView<FPBDRigidParticleHandle*> RigidParticles(&RigidParticle, 1);
 		
-		UpdateViews();
+		SetClusteredParticlesSOA(RigidParticles, ClusteredParticle->ObjectState(), true, ClusteredParticle->Sleeping());
 	}
 
 	void MarkMovingKinematic(FKinematicGeometryParticleHandle* Particle)
@@ -972,38 +1177,53 @@ public:
 	const TArray<FPBDGeometryCollectionParticleHandle*>& GetSleepingGeometryCollectionArray() const { return SleepingGeometryCollectionArray.GetArray(); }
 	const TArray<FPBDGeometryCollectionParticleHandle*>& GetDynamicGeometryCollectionArray() const { return DynamicGeometryCollectionArray.GetArray(); }
 
-	void InsertGeometryCollectionParticle(TPBDGeometryCollectionParticleHandle<FReal, 3>* GCParticle)
+	void InsertGeometryCollectionParticles(const TArrayView<FPBDRigidParticleHandle*>& RigidParticles, const bool bIsDisabled, const bool bIsSleeping, const EObjectStateType ObjectState)
 	{
-		if (!GCParticle->Disabled())
+		if (!bIsDisabled)
 		{
-			const Chaos::EObjectStateType State = GCParticle->Sleeping() ? Chaos::EObjectStateType::Sleeping : GCParticle->ObjectState();
+			const Chaos::EObjectStateType State = bIsSleeping ? Chaos::EObjectStateType::Sleeping : ObjectState;
 			switch (State)
 			{
-			case EObjectStateType::Uninitialized:
-				ensure(false); // we should probably not be here 
-				break;
-			case EObjectStateType::Static:
-				AddToList(GCParticle, StaticGeometryCollectionArray);
-				break;
-			case EObjectStateType::Kinematic:
-				AddToList(GCParticle, KinematicGeometryCollectionArray);
-				break;
-			case EObjectStateType::Dynamic:
-				AddToList(GCParticle, DynamicGeometryCollectionArray);
-				break;
-			case EObjectStateType::Sleeping:
-				AddToList(GCParticle, SleepingGeometryCollectionArray);
-				break;
+				case EObjectStateType::Uninitialized:
+					ensure(false); // we should probably not be here 
+					break;
+				case EObjectStateType::Static:
+					AddToList(RigidParticles, StaticGeometryCollectionArray);
+					break;
+				case EObjectStateType::Kinematic:
+					AddToList(RigidParticles, KinematicGeometryCollectionArray);
+					break;
+				case EObjectStateType::Dynamic:
+					AddToList(RigidParticles, DynamicGeometryCollectionArray);
+					break;
+				case EObjectStateType::Sleeping:
+					AddToList(RigidParticles, SleepingGeometryCollectionArray);
+					break;
+				default : 
+					break;
 			}
 		}
 	}
 
+	void InsertGeometryCollectionParticle(TPBDGeometryCollectionParticleHandle<FReal, 3>* GCParticle)
+	{
+		FPBDRigidParticleHandle* RigidParticle = static_cast<FPBDRigidParticleHandle*>(GCParticle);
+		InsertGeometryCollectionParticles(MakeArrayView<FPBDRigidParticleHandle*>(&RigidParticle, 1), 
+			GCParticle->Disabled(), GCParticle->Sleeping(), GCParticle->ObjectState());
+	}
+
+	void RemoveGeometryCollectionParticles(const TArrayView<FPBDRigidParticleHandle*>& GCParticles)
+	{
+		RemoveFromList(GCParticles, StaticGeometryCollectionArray);
+		RemoveFromList(GCParticles, KinematicGeometryCollectionArray);
+		RemoveFromList(GCParticles, SleepingGeometryCollectionArray);
+		RemoveFromList(GCParticles, DynamicGeometryCollectionArray);
+	}
+
 	void RemoveGeometryCollectionParticle(TPBDGeometryCollectionParticleHandle<FReal, 3>* GCParticle)
 	{
-		RemoveFromList(GCParticle, StaticGeometryCollectionArray);
-		RemoveFromList(GCParticle, KinematicGeometryCollectionArray);
-		RemoveFromList(GCParticle, SleepingGeometryCollectionArray);
-		RemoveFromList(GCParticle, DynamicGeometryCollectionArray);
+		FPBDRigidParticleHandle* RigidParticle = static_cast<FPBDRigidParticleHandle*>(GCParticle);
+		RemoveGeometryCollectionParticles(MakeArrayView<FPBDRigidParticleHandle*>(&RigidParticle, 1));
 	}
 
 	const auto& GetClusteredParticles() const { return *ClusteredParticles; }
@@ -1029,29 +1249,38 @@ private:
 
 	// Initialize the list mask for all particle containers and lists
 	CHAOS_API void SetContainerListMasks();
-
-	// Add a particle to the specified SOA/List and update its ListMask
+	
+	// Add a set of particles to the specified SOA/List and update their ListMasks
 	template<typename TListType, typename TParticleHandleType>
-	void AddToList(TParticleHandleType* Particle, TListType& List)
+	void AddToList(const TArrayView<TParticleHandleType*> Particles, TListType& List)
 	{
-		//if (!Particle->IsInAnyList(List.GetContainerListMask()))
+		check(List.GetContainerListMask() != EGeometryParticleListMask::None);
+
+		for (TParticleHandleType* Particle : Particles)
 		{
 			Particle->AddToLists(List.GetContainerListMask());
-			List.Insert(Particle);
 		}
+		List.Insert(Particles);
 	}
 
-	// Add a set of particles to the specified SOA/List and update their ListMasks
+	// Add a particle to the specified SOA/List and update its ListMask
 	template<typename TListType, typename TParticleHandleType>
-	void AddToList(const TArray<TParticleHandleType*> Particles, TListType& List)
+	void AddToList(TParticleHandleType* Particle, TListType& List)
+	{
+		AddToList(MakeArrayView<TParticleHandleType*>(&Particle, 1), List);
+	}
+	
+	// Remove a set of particles from the specified SOA/List and update their ListMask
+	template<typename TListType, typename TParticleHandleType>
+	void RemoveFromList(const TArrayView<TParticleHandleType*>& Particles, TListType& List)
 	{
 		check(List.GetContainerListMask() != EGeometryParticleListMask::None);
 
 		for (TParticleHandleType* Particle : Particles)
 		{
-			Particle->AddToLists(List.GetContainerListMask());
+			Particle->RemoveFromLists(List.GetContainerListMask());
 		}
-		List.Insert(Particles);
+		List.Remove(Particles);
 	}
 
 	// Remove a particle from the specified SOA/List and update its ListMask
@@ -1059,12 +1288,8 @@ private:
 	void RemoveFromList(TParticleHandleType* Particle, TListType& List)
 	{
 		check(List.GetContainerListMask() != EGeometryParticleListMask::None);
-
-		//if (Particle->IsInAnyList(List.GetContainerListMask()))
-		{
-			Particle->RemoveFromLists(List.GetContainerListMask());
-			List.Remove(Particle);
-		}
+		Particle->RemoveFromLists(List.GetContainerListMask());
+		List.Remove(Particle);
 	}
 
 	// Check that the particles in an SOA have the appropriate bit set in their ListMask
@@ -1145,7 +1370,7 @@ private:
 		return ReturnHandles;
 	}
 	
-	void AddToActiveArray(const TArray<FPBDRigidParticleHandle*>& Particles)
+	void AddToActiveArray(const TArrayView<FPBDRigidParticleHandle*>& Particles)
 	{
 		AddToList(Particles, ActiveParticlesMapArray);
 
@@ -1163,81 +1388,108 @@ private:
 
 	void AddToActiveArray(FPBDRigidParticleHandle* Particle)
 	{
-		AddToList(Particle, ActiveParticlesMapArray);
-
-		if (bResimulating)
-		{
-			AddToList(Particle, ResimActiveParticlesMapArray);
-		}
-
-		//dirty contains Active so make sure no duplicates
-		RemoveFromList(Particle, TransientDirtyMapArray);
+		TArrayView<FPBDRigidParticleHandle*> RigidParticles(&Particle, 1);
+		AddToActiveArray(RigidParticles);
 	}
 
-	void RemoveFromActiveArray(FPBDRigidParticleHandle* Particle, bool bStillDirty)
+	void RemoveFromActiveArray(TArrayView<FPBDRigidParticleHandle*>& RigidParticles, bool bStillDirty)
 	{
-		RemoveFromList(Particle, ActiveParticlesMapArray);
+		RemoveFromList(RigidParticles, ActiveParticlesMapArray);
 
 		if (bResimulating)
 		{
-			RemoveFromList(Particle, ResimActiveParticlesMapArray);
+			RemoveFromList(RigidParticles, ResimActiveParticlesMapArray);
 		}
 
-		if(bStillDirty)
+		if (bStillDirty)
 		{
-			if(!MovingKinematicsMapArray.Contains(Particle))
-			{ 
-				//no longer active, but still dirty
-				AddToList(Particle, TransientDirtyMapArray);
+			int32 ParticleOffset = 0;
+			
+			// The particles in [0-ParticleOffset) will be the enabled ones switching to the new state
+			Private::ReorderParticleView(RigidParticles, [this](FPBDRigidParticleHandle* RigidParticle) -> bool 
+				{ return !MovingKinematicsMapArray.Contains(RigidParticle);},ParticleOffset);	
+		
+			const int32 NumTransients = ParticleOffset;
+			if (NumTransients > 0)
+			{
+				// Mark transient particles dirty
+				TArrayView<FPBDRigidParticleHandle*> TransientParticles(&RigidParticles[0], NumTransients);
+				AddToList(TransientParticles, TransientDirtyMapArray);
 			}
 		}
 		else
 		{
 			//might have already been removed from active from a previous call
 			//but now removing and don't want it dirty either
-			RemoveFromList(Particle, TransientDirtyMapArray);
+			RemoveFromList(RigidParticles, TransientDirtyMapArray);
 		}
 	}
+
+	void RemoveFromActiveArray(FPBDRigidParticleHandle* Particle, bool bStillDirty)
+	{
+		TArrayView<FPBDRigidParticleHandle*> RigidParticles(&Particle, 1);
+		RemoveFromActiveArray(RigidParticles, bStillDirty);
+	}
 	
 	//should be called whenever particles are added / removed / reordered
 	CHAOS_API void UpdateViews();
 
 	void InsertClusteredParticle(FPBDRigidClusteredParticleHandle* ClusteredParticle)
 	{
-		if (!ClusteredParticle->Disabled())
+		FPBDRigidParticleHandle* RigidParticle = static_cast<FPBDRigidParticleHandle*>(ClusteredParticle);
+		TArrayView<FPBDRigidParticleHandle*> ClusteredParticlesView(&RigidParticle, 1);
+		
+		InsertClusteredParticles(ClusteredParticlesView, ClusteredParticle->Disabled(), ClusteredParticle->ObjectState());
+	}
+
+	void InsertClusteredParticles(const TArray<FPBDRigidClusteredParticleHandle*>& Particles)
+	{
+		// We still need to loop over each particle since each one of them could have different object state 
+		// and could potentially be disabled
+		for (FPBDRigidClusteredParticleHandle* ClusteredParticle : Particles)
 		{
-			switch (ClusteredParticle->ObjectState())
-			{
-			case EObjectStateType::Uninitialized:
-				ensure(false); // we should probably not be here 
-				break;
-			case EObjectStateType::Static:
-				AddToList(ClusteredParticle, StaticClusteredMapArray);
-				break;
-			case EObjectStateType::Kinematic:
-				AddToList(ClusteredParticle, KinematicClusteredMapArray);
-				break;
-			case EObjectStateType::Dynamic:
-			case EObjectStateType::Sleeping:
-				AddToList(ClusteredParticle, DynamicClusteredMapArray);
-				break;
-			}
+			InsertClusteredParticle(ClusteredParticle);
 		}
 	}
 
-	void InsertClusteredParticles(const TArray<FPBDRigidClusteredParticleHandle*>& ClusteredParticleArray)
+	void InsertClusteredParticles(const TArrayView<FPBDRigidParticleHandle*>& RigidParticles, const bool bIsDisabled, const EObjectStateType ObjectState)
 	{
-		for (FPBDRigidClusteredParticleHandle* ClusteredParticle : ClusteredParticleArray)
+		if (!bIsDisabled)
 		{
-			InsertClusteredParticle(ClusteredParticle);
+			switch (ObjectState)
+			{
+				case EObjectStateType::Uninitialized:
+					ensure(false); // we should probably not be here 
+					break;
+				case EObjectStateType::Static:
+					AddToList(RigidParticles, StaticClusteredMapArray);
+					break;
+				case EObjectStateType::Kinematic:
+					AddToList(RigidParticles, KinematicClusteredMapArray);
+					break;
+				case EObjectStateType::Dynamic:
+				case EObjectStateType::Sleeping:
+					AddToList(RigidParticles, DynamicClusteredMapArray);
+					break;
+				default : 
+					break;
+			}
 		}
 	}
 
 	void RemoveClusteredParticle(FPBDRigidClusteredParticleHandle* ClusteredParticle)
 	{
-		RemoveFromList(ClusteredParticle, StaticClusteredMapArray);
-		RemoveFromList(ClusteredParticle, KinematicClusteredMapArray);
-		RemoveFromList(ClusteredParticle, DynamicClusteredMapArray);
+		FPBDRigidParticleHandle* RigidParticle = static_cast<FPBDRigidParticleHandle*>(ClusteredParticle);
+		TArrayView<FPBDRigidParticleHandle*> ClusteredParticlesView(&RigidParticle, 1);
+		
+		RemoveClusteredParticles(ClusteredParticlesView);
+	}
+
+	void RemoveClusteredParticles(const TArrayView<FPBDRigidParticleHandle*>& RigidParticles)
+	{
+		RemoveFromList(RigidParticles, StaticClusteredMapArray);
+		RemoveFromList(RigidParticles, KinematicClusteredMapArray);
+		RemoveFromList(RigidParticles, DynamicClusteredMapArray);
 	}
 
 	//Organized by SOA type
diff --git a/Engine/Source/Runtime/Experimental/Chaos/Public/Chaos/ParticleHandle.h b/Engine/Source/Runtime/Experimental/Chaos/Public/Chaos/ParticleHandle.h
--- a/Engine/Source/Runtime/Experimental/Chaos/Public/Chaos/ParticleHandle.h
+++ b/Engine/Source/Runtime/Experimental/Chaos/Public/Chaos/ParticleHandle.h
@@ -1631,7 +1631,7 @@ public:
 	const TPBDRigidClusteredParticleHandleImp<T, d, true>* Handle() const { return PBDRigidClusteredParticles->Handle(ParticleIdx); }
 	TPBDRigidClusteredParticleHandleImp<T, d, true>* Handle() { return PBDRigidClusteredParticles->Handle(ParticleIdx); }
 
-	static constexpr EParticleType StaticType() { return EParticleType::Rigid; }
+	static constexpr EParticleType StaticType() { return EParticleType::Clustered; }
 
 	int32 TransientParticleIndex() const { return ParticleIdx; }
 
diff --git a/Engine/Source/Runtime/Experimental/Chaos/Public/Chaos/Particles.h b/Engine/Source/Runtime/Experimental/Chaos/Public/Chaos/Particles.h
--- a/Engine/Source/Runtime/Experimental/Chaos/Public/Chaos/Particles.h
+++ b/Engine/Source/Runtime/Experimental/Chaos/Public/Chaos/Particles.h
@@ -93,6 +93,13 @@ namespace Chaos
 			MoveToOtherArrayCollection(Idx, Other);
 			IncrementDirtyValidation();
 		}
+		
+		/** Move particle from one SOA to another one given a list of particle indices */
+		void MoveToOtherParticles(const TArrayView<int32>& Indices, TParticles<T, d>& Other)
+		{
+			MoveToOtherArrayCollection(Indices, Other);
+			AddDirtyValidation(Indices.Num());
+		}
 
 		void Resize(const int32 Num)
 		{
@@ -213,6 +220,14 @@ namespace Chaos
 #endif
 		}
 
+		/** Increase the dirty validation count by NumDirty */
+		void AddDirtyValidation(const int32 NumDirty)
+		{
+#if PARTICLE_ITERATOR_RANGED_FOR_CHECK
+			MDirtyValidationCount += NumDirty;
+#endif
+		}
+
 		inline friend FArchive& operator<<(FArchive& Ar, TParticles<T, d>& InParticles)
 		{
 			InParticles.Serialize(Ar);
diff --git a/Engine/Source/Runtime/Experimental/Chaos/Public/RewindData.h b/Engine/Source/Runtime/Experimental/Chaos/Public/RewindData.h
--- a/Engine/Source/Runtime/Experimental/Chaos/Public/RewindData.h
+++ b/Engine/Source/Runtime/Experimental/Chaos/Public/RewindData.h
@@ -2266,6 +2266,18 @@ private:
 	// Properties for EResimFrameValidation::IslandValidation logic
 	TArray<const Private::FPBDIsland*> IslandValidationIslands;
 	TArray<const FGeometryParticleHandle*> IslandValidationIslandParticles;
+	
+	/** Frozen dynamic particle handles to be used when unfreezing  */
+	TArray<FPBDRigidParticleHandle*> FrozenDynamicHandles;
+	
+	/** Frozen sleeping particle handles to be used when unfreezing  */
+	TArray<FPBDRigidParticleHandle*> FrozenSleepingHandles;
+	
+	/** Frozen static particle handles to be used when unfreezing  */
+	TArray<FPBDRigidParticleHandle*> FrozenStaticHandles;
+	
+	/** Frozen kinematic particle handles to be used when unfreezing  */
+	TArray<FPBDRigidParticleHandle*> FrozenKinematicHandles;
 
 	TSet<const FGeometryParticleHandle*> ParticlesToForceResimAsFollower;
 
diff --git a/Engine/Source/Runtime/Experimental/ChaosCore/Public/Chaos/ArrayCollection.h b/Engine/Source/Runtime/Experimental/ChaosCore/Public/Chaos/ArrayCollection.h
--- a/Engine/Source/Runtime/Experimental/ChaosCore/Public/Chaos/ArrayCollection.h
+++ b/Engine/Source/Runtime/Experimental/ChaosCore/Public/Chaos/ArrayCollection.h
@@ -3,6 +3,7 @@
 
 #include "Chaos/Array.h"
 #include "Chaos/ArrayCollectionArrayBase.h"
+#include "Chaos/ArrayCollectionArray.h"
 
 namespace Chaos
 {
@@ -144,6 +145,72 @@ protected:
 		++Other.MSize;
 		--MSize;
 	}
+	/** Move collection data from one array to another one */
+	void MoveToOtherArrayCollection(const TArrayView<int32>& Indices, TArrayCollection& Other)
+	{
+		if (Indices.Num() <= static_cast<int32>(MSize))
+		{
+			for(int32 Index : Indices)
+			{
+				if (static_cast<uint32>(Index) >= MSize)
+				{
+					return;
+				}
+			}
+			if(MArrays.Num() == Other.MArrays.Num())
+			{ 
+				bool bCanAddIndices = false;
+
+				// For each collection arrays add data to the new array
+				for (int32 ArrayIdx = 0; ArrayIdx < MArrays.Num(); ++ArrayIdx)
+				{
+					if(MArrays[ArrayIdx] && Other.MArrays[ArrayIdx])
+					{ 
+						MArrays[ArrayIdx]->MoveToOtherArray(Indices, *Other.MArrays[ArrayIdx], false);
+						bCanAddIndices = true;
+					}
+				}
+				// Update the old collection size if at least one array has been moved
+				if(bCanAddIndices)
+				{
+					Other.MSize += Indices.Num();
+				}
+			}
+				
+			// For each collection arrays remove data from the old array
+			RemoveAtSwapLambda(static_cast<uint32>(MSize), Indices, [this](const int32 RemovedIndex)
+			{
+				// Lambda function to remove at swap the data for a given index
+				for (int32 ArrayIdx = 0; ArrayIdx < MArrays.Num(); ++ArrayIdx)
+				{
+					if(MArrays[ArrayIdx])
+					{ 
+						// The remove at swap from the TArrayCollectionArray is not shrinking the array which is 
+						// necessary for the algorithm to work. Array could potentially be shrinked just after.
+						MArrays[ArrayIdx]->RemoveAtSwap(RemovedIndex);
+							
+					}
+				}
+			});
+
+			bool bCanRemoveIndices = false;
+
+			// Lambda function to shrink the collection arrays a the end of removal
+			for (int32 ArrayIdx = 0; ArrayIdx < MArrays.Num(); ++ArrayIdx)
+			{
+				if(MArrays[ArrayIdx])
+				{ 
+					MArrays[ArrayIdx]->Shrink();
+					bCanRemoveIndices = true;
+				}
+			}
+			// Update the new collection size if at least one array has been moved
+			if(bCanRemoveIndices)
+			{
+				MSize -= Indices.Num();
+			}
+		}
+	}
 
 private:
 	TArray<TArrayCollectionArrayBase*> MArrays;
diff --git a/Engine/Source/Runtime/Experimental/ChaosCore/Public/Chaos/ArrayCollectionArray.h b/Engine/Source/Runtime/Experimental/ChaosCore/Public/Chaos/ArrayCollectionArray.h
--- a/Engine/Source/Runtime/Experimental/ChaosCore/Public/Chaos/ArrayCollectionArray.h
+++ b/Engine/Source/Runtime/Experimental/ChaosCore/Public/Chaos/ArrayCollectionArray.h
@@ -5,11 +5,74 @@
 
 #include "Chaos/Array.h"
 #include "Chaos/ArrayCollectionArrayBase.h"
-
+#include "Containers/BitArray.h"
 #include <algorithm>
 
 namespace Chaos
 {
+
+/** Remove at swap function working on an unsorted array of indices to remove. Sorting the 
+ * indices in place would be an option but the complexity being O(Nlog(N)), it is probably too slow to be used for 
+ * a lot of indices to remove. This algorithm requires allocating a mapping buffer to keep track of where the moved elements 
+ * are and that the RemoveFunc is not modifying the original memory order */
+template<typename RemoveIndexFunc>
+FORCEINLINE void RemoveAtSwapLambda(const int32 ArraySize, const TArrayView<int32>& Indices, const RemoveIndexFunc& RemoveFunc)
+{
+	if (Indices.IsEmpty() || (Indices.Num() > ArraySize))
+	{
+		return;
+	}
+#if !(UE_BUILD_SHIPPING || UE_BUILD_TEST)
+	TBitArray<> ValidIndices;
+	ValidIndices.Init(false,ArraySize);
+	for (int32 RemovedIndex : Indices)
+	{
+		if ((RemovedIndex < ArraySize) && (RemovedIndex >= 0))
+		{
+			// Check if the list of indices to remove is unique or not
+			check(!ValidIndices[RemovedIndex]);
+			ValidIndices[RemovedIndex] = true;
+		}
+	}
+#endif
+	// The mapping array is twice the size M of the indices to be removed among the array of size N elements
+	// [forward indices : array index position of the last M element position | reverse indices : element index value of the last M array positions]
+	TArray<int32> IndexMapping;
+	IndexMapping.Init(INDEX_NONE, 2 * Indices.Num());
+
+	// Since we are storing only that last M elements we need to store the forward/reverse offset to come back to global indices
+	const int32 ForwardOffset = ArraySize - Indices.Num();
+	const int32 ReverseOffset = ForwardOffset - Indices.Num();
+
+	// The goal here is to perform successive RemoveAtSwap on a list of indices without sorting them
+	// since doing naively sequential calls could leads to wrong results if not sorted correctly
+	int32 EndIndex = ArraySize - 1;
+	for (int32 RemovedIndex : Indices)
+	{
+		if ((RemovedIndex < ArraySize) && (RemovedIndex >= 0) && IndexMapping.IsValidIndex(EndIndex - ReverseOffset))
+		{
+			const int32 ForwardMapping = (RemovedIndex < ForwardOffset) ? RemovedIndex : IndexMapping[RemovedIndex - ForwardOffset];
+			const int32 ReverseMapping = IndexMapping[EndIndex - ReverseOffset];
+
+			const int32 ForwardIndex = (ForwardMapping != INDEX_NONE) ? ForwardMapping : RemovedIndex;
+			const int32 ReverseIndex = (ReverseMapping != INDEX_NONE) ? ReverseMapping : EndIndex;
+
+			if (ReverseIndex >= ForwardOffset)
+			{
+				IndexMapping[ReverseIndex - ForwardOffset] = ForwardIndex;
+			}
+		
+			// If ForwardIndex >= ForwardOffset then ForwardIndex >= ReverseOffset and the IndexMapping access is valid
+			if (ForwardIndex >= ForwardOffset)
+			{
+				IndexMapping[ForwardIndex - ReverseOffset] = ReverseIndex;
+			}
+			RemoveFunc(ForwardIndex);
+			--EndIndex;
+		}
+	}
+}
+	
 template<class T>
 class TArrayCollectionArray : public TArrayCollectionArrayBase, public TArray<T>
 {
@@ -91,6 +154,27 @@ public:
 		TArray<T>::RemoveAtSwap(Idx, 1, AllowShrinkOnRemove);
 	}
 
+	/** Remove a list of elements from the array */
+	FORCEINLINE void RemoveAtSwap(const TArrayView<int32>& Indices) override
+	{
+		RemoveAtSwapLambda(Num(), Indices, [this](const int32 RemovedIndex)
+		{
+			TArray<T>::RemoveAtSwap(RemovedIndex, 1, EAllowShrinking::No);
+		});
+		
+		if (AllowShrinkOnRemove == EAllowShrinking::Yes)
+		{
+			TArray<T>::Shrink();
+		}
+	}
+
+	/** Shrink the array size once the removal has been done. */
+	FORCEINLINE void Shrink() override
+	{
+		TArray<T>::Shrink();
+	}
+
+	/** Move only one element from a collection to another one */
 	FORCEINLINE void MoveToOtherArray(const int Idx, TArrayCollectionArrayBase& Other)
 	{
 		//todo: add developer check to make sure this is ok?
@@ -99,6 +183,28 @@ public:
 		TArray<T>::RemoveAtSwap(Idx, 1, AllowShrinkOnRemove);
 	}
 
+	/** Move a list of elements from a collection to another one */
+	FORCEINLINE void MoveToOtherArray(const TArrayView<int32>& Indices, TArrayCollectionArrayBase& Other, const bool bRemoveElements)
+	{
+		TArrayCollectionArray<T>& OtherTArray = static_cast<TArrayCollectionArray<T>&>(Other);
+		OtherTArray.Reserve(OtherTArray.Num() + Indices.Num());
+	
+		// Add elements to the new collection arrays 
+		for(int32 Idx : Indices)
+		{
+			if (TArray<T>::IsValidIndex(Idx))
+			{
+				OtherTArray.Emplace(MoveTemp(TArray<T>::operator [](Idx)));
+			}
+		}
+
+		// Remove elements from the old collection arrays
+		if(bRemoveElements)
+		{ 
+			RemoveAtSwap(Indices);
+		}
+	}
+
 	FORCEINLINE uint64 SizeOfElem() const override
 	{
 		return sizeof(T);
diff --git a/Engine/Source/Runtime/Experimental/ChaosCore/Public/Chaos/ArrayCollectionArrayBase.h b/Engine/Source/Runtime/Experimental/ChaosCore/Public/Chaos/ArrayCollectionArrayBase.h
--- a/Engine/Source/Runtime/Experimental/ChaosCore/Public/Chaos/ArrayCollectionArrayBase.h
+++ b/Engine/Source/Runtime/Experimental/ChaosCore/Public/Chaos/ArrayCollectionArrayBase.h
@@ -2,6 +2,7 @@
 #pragma once
 
 #include "CoreTypes.h"
+#include "Containers/ArrayView.h"
 
 namespace Chaos
 {
@@ -12,7 +13,10 @@ class TArrayCollectionArrayBase
 	virtual void Resize(const int Num) = 0;
 	virtual void RemoveAt(const int Idx, const int Count) = 0;
 	virtual void RemoveAtSwap(const int Idx) = 0;
+	virtual void RemoveAtSwap(const TArrayView<int32>& Indices) = 0;
 	virtual uint64 SizeOfElem() const = 0;
 	virtual void MoveToOtherArray(const int Idx, TArrayCollectionArrayBase& Other) = 0;
+	virtual void MoveToOtherArray(const TArrayView<int32>& Indices, TArrayCollectionArrayBase& Other, const bool bRemoveElements) = 0;
+	virtual void Shrink() = 0;
 };
 }
