Exposed host group name and added on-demand file extensions.

#rnx

#Synced-CL 49823132

diff --git a/Engine/Source/Runtime/Experimental/IoStore/OnDemandCore/Internal/IO/Serialization/OnDemandContainerToc.h b/Engine/Source/Runtime/Experimental/IoStore/OnDemandCore/Internal/IO/Serialization/OnDemandContainerToc.h
--- a/Engine/Source/Runtime/Experimental/IoStore/OnDemandCore/Internal/IO/Serialization/OnDemandContainerToc.h
+++ b/Engine/Source/Runtime/Experimental/IoStore/OnDemandCore/Internal/IO/Serialization/OnDemandContainerToc.h
@@ -29,6 +29,14 @@ namespace UE::IoStore { struct FOnDemandToc; }
 namespace UE::IoStore::Serialization
 {
 
+///////////////////////////////////////////////////////////////////////////////
+struct FOnDemandFileExt
+{
+	static constexpr const TCHAR* Toc		= TEXT(".uondemandtoc");
+	static constexpr const TCHAR* CloudToc	= TEXT(".iochunktoc");
+	static constexpr const TCHAR* Chunk		= TEXT(".iochunk");
+};
+
 ///////////////////////////////////////////////////////////////////////////////
 struct FOnDemandTocSignature
 {
@@ -320,6 +328,7 @@ public:
 	FUtf8StringView			BuildVersion() const { return GetString(Header->BuildVersion); }
 	FUtf8StringView			TargetPlatform() const { return GetString(Header->TargetPlatform); }
 	FUtf8StringView			ChunksDirectory() const { return GetString(Header->ChunksDirectory); }
+	FUtf8StringView			HostGroupName() const { return GetString(Header->HostGroupName); }
 	FUtf8StringView			CompressionFormat() const { return GetString(Header->CompressionFormat); }
 	UE_API FUtf8StringView	GetString(FOnDemandStringEntry StringEntry) const;
 
@@ -328,7 +337,7 @@ public:
 	UE_API TIoStatusOr<FOnDemandContainerTocView> ReadContainer(
 		const FOnDemandContainerEntry& Container,
 		FOnDemandTocStorage& Storage,
-		EOnDemandTocReaderOptions Options);
+		EOnDemandTocReaderOptions Options) const;
 	UE_API static TIoStatusOr<FOnDemandTocReader> Read(const FString& Filename);
 
 private:
diff --git a/Engine/Source/Runtime/Experimental/IoStore/OnDemandCore/Private/Serialization/OnDemandContainerToc.cpp b/Engine/Source/Runtime/Experimental/IoStore/OnDemandCore/Private/Serialization/OnDemandContainerToc.cpp
--- a/Engine/Source/Runtime/Experimental/IoStore/OnDemandCore/Private/Serialization/OnDemandContainerToc.cpp
+++ b/Engine/Source/Runtime/Experimental/IoStore/OnDemandCore/Private/Serialization/OnDemandContainerToc.cpp
@@ -339,7 +339,7 @@ FUtf8StringView FOnDemandTocReader::GetString(FOnDemandStringEntry StringEntry)
 TIoStatusOr<FOnDemandContainerTocView> FOnDemandTocReader::ReadContainer(
 	const FOnDemandContainerEntry& Container,
 	FOnDemandTocStorage& Storage,
-	EOnDemandTocReaderOptions Options)
+	EOnDemandTocReaderOptions Options) const
 {
 	const FOnDemandContainerEntry* Existing = Algo::FindByPredicate(
 		ContainerEntries, [&Container](const FOnDemandContainerEntry& Entry) { return &Entry == &Container; });
