Extended FOnDemandTocStorage to return whether the memory is file mapped or not.

#rnx

#Synced-CL 49823132

diff --git a/Engine/Source/Runtime/Experimental/IoStore/OnDemandCore/Internal/IO/Serialization/OnDemandContainerToc.h b/Engine/Source/Runtime/Experimental/IoStore/OnDemandCore/Internal/IO/Serialization/OnDemandContainerToc.h
--- a/Engine/Source/Runtime/Experimental/IoStore/OnDemandCore/Internal/IO/Serialization/OnDemandContainerToc.h
+++ b/Engine/Source/Runtime/Experimental/IoStore/OnDemandCore/Internal/IO/Serialization/OnDemandContainerToc.h
@@ -93,6 +93,14 @@ struct FOnDemandTocFooter
 };
 static_assert(sizeof(FOnDemandTocFooter) == 16);
 
+////////////////////////////////////////////////////////////////////////////////
+enum class FOnDemandTocStorageType : uint8
+{
+	Invalid = 0,
+	Memory,
+	MemoryMapped
+};
+
 ////////////////////////////////////////////////////////////////////////////////
 class FOnDemandTocStorage
 {
@@ -118,8 +126,9 @@ public:
 
 	UE_API FOnDemandTocStorage& operator=(FOnDemandTocStorage&&);
 
-	UE_API FMemoryView GetView() const;
-	bool IsEmpty() const { return HeaderData.GetSize() == 0; }
+	UE_API FMemoryView				GetView() const;
+	UE_API FOnDemandTocStorageType	StorageType() const;
+	bool							IsEmpty() const { return HeaderData.GetSize() == 0; }
 
 private:
 	FOnDemandTocStorage(const FOnDemandTocStorage&) = delete;
diff --git a/Engine/Source/Runtime/Experimental/IoStore/OnDemandCore/Private/Serialization/OnDemandContainerToc.cpp b/Engine/Source/Runtime/Experimental/IoStore/OnDemandCore/Private/Serialization/OnDemandContainerToc.cpp
--- a/Engine/Source/Runtime/Experimental/IoStore/OnDemandCore/Private/Serialization/OnDemandContainerToc.cpp
+++ b/Engine/Source/Runtime/Experimental/IoStore/OnDemandCore/Private/Serialization/OnDemandContainerToc.cpp
@@ -80,6 +80,20 @@ FMemoryView FOnDemandTocStorage::GetView() const
 	return FMemoryView();
 }
 
+FOnDemandTocStorageType FOnDemandTocStorage::StorageType() const
+{
+	if (ContainerData.IsType<FMappedFileStorage>())
+	{
+		return FOnDemandTocStorageType::MemoryMapped;
+	}
+	else if (ContainerData.IsType<FMemoryStorage>())
+	{
+		return FOnDemandTocStorageType::Memory;
+	}
+
+	return FOnDemandTocStorageType::Invalid;
+}
+
 ////////////////////////////////////////////////////////////////////////////////
 class FOnDemandTocWriter::FImpl
 {
