Dataflow : Add a way to compile ( and cache) graphs from the dataflow asset

#Synced-CL 49892100

diff --git a/Engine/Source/Runtime/Experimental/Dataflow/Engine/Private/Dataflow/DataflowObject.cpp b/Engine/Source/Runtime/Experimental/Dataflow/Engine/Private/Dataflow/DataflowObject.cpp
--- a/Engine/Source/Runtime/Experimental/Dataflow/Engine/Private/Dataflow/DataflowObject.cpp
+++ b/Engine/Source/Runtime/Experimental/Dataflow/Engine/Private/Dataflow/DataflowObject.cpp
@@ -6,6 +6,8 @@
 #include "Dataflow/DataflowNodeParameters.h"
 #include "Dataflow/DataflowObjectInterface.h"
 #include "Dataflow/DataflowSubGraph.h"
+#include "Dataflow/DataflowCompiler.h"
+#include "Dataflow/DataflowCompiledGraph.h"
 
 #if WITH_EDITOR
 #include "EdGraph/EdGraphPin.h"
@@ -462,6 +464,34 @@ void UDataflow::RefreshEdNode(TObjectPtr<UDataflowEdNode> EdNode)
 	}
 }
 
+TSharedPtr<const UE::Dataflow::FCompiledGraph> UDataflow::CompileGraphIfNeeded()
+{
+	const FGuid GraphTopolyGuid = Dataflow ? Dataflow->GetGraphTopologyGuid() : FGuid();
+	if (LastCompiledTopologyGUID != GraphTopolyGuid)
+	{
+		CompiledGraphs.Reset();
+
+		UE::Dataflow::FCompilerParameters Params;
+		Params.DataflowObjectToCompile = this;
+
+		UE::Dataflow::FCompilerResults Results;
+		if (UE::Dataflow::FCompiler::Compile(Params, Results))
+		{
+			CompiledGraphs = Results.GetCompiledGraphs();
+		}
+		else
+		{
+			ensureMsgf(false, TEXT("Graph compilation failed, check the result error log"));
+		}
+		// store the guid even if the compilation was not successful so that we don't keep on trying to recompile
+		LastCompiledTopologyGUID = GraphTopolyGuid;
+	}
+	const TSharedRef<UE::Dataflow::FCompiledGraph>* OutCompiledGraph = CompiledGraphs.Find(this);
+	return OutCompiledGraph 
+		? OutCompiledGraph->ToSharedPtr()
+		: nullptr;
+}
+
 #if WITH_EDITOR
 bool UDataflow::CanEditChange(const FProperty* InProperty) const
 {
diff --git a/Engine/Source/Runtime/Experimental/Dataflow/Engine/Private/Dataflow/DataflowObjectInterface.cpp b/Engine/Source/Runtime/Experimental/Dataflow/Engine/Private/Dataflow/DataflowObjectInterface.cpp
--- a/Engine/Source/Runtime/Experimental/Dataflow/Engine/Private/Dataflow/DataflowObjectInterface.cpp
+++ b/Engine/Source/Runtime/Experimental/Dataflow/Engine/Private/Dataflow/DataflowObjectInterface.cpp
@@ -18,17 +18,10 @@ namespace UE::Dataflow
 			{
 				if (TSharedPtr<FGraph> Graph = DataflowObject->GetDataflow())
 				{
-					UE::Dataflow::FCompilerParameters Params;
-					Params.DataflowObjectToCompile = DataflowObject;
-
-					UE::Dataflow::FCompilerResults Results;
-					if (UE::Dataflow::FCompiler::Compile(Params, Results))
+					if (TSharedPtr<const UE::Dataflow::FCompiledGraph> CompiledGraph = DataflowObject->CompileGraphIfNeeded())
 					{
-						if (TSharedPtr<const UE::Dataflow::FCompiledGraph> CompiledGraph = Results.GetCompiledGraph(DataflowObject))
-						{
-							Evaluator->Start(Graph.ToSharedRef(), CompiledGraph.ToSharedRef(), FGuid(), OnPostEvaluation);
-							return true;
-						}
+						Evaluator->Start(Graph.ToSharedRef(), CompiledGraph.ToSharedRef(), FGuid(), OnPostEvaluation);
+						return true;
 					}
 				}
 			}
diff --git a/Engine/Source/Runtime/Experimental/Dataflow/Engine/Public/Dataflow/DataflowCompiler.h b/Engine/Source/Runtime/Experimental/Dataflow/Engine/Public/Dataflow/DataflowCompiler.h
--- a/Engine/Source/Runtime/Experimental/Dataflow/Engine/Public/Dataflow/DataflowCompiler.h
+++ b/Engine/Source/Runtime/Experimental/Dataflow/Engine/Public/Dataflow/DataflowCompiler.h
@@ -59,6 +59,8 @@ namespace UE::Dataflow
 			return CompiledGraphPtr ? CompiledGraphPtr->ToSharedPtr() : TSharedPtr<const FCompiledGraph>();
 		}
 
+		const TMap<const UObject*, TSharedRef<FCompiledGraph>> GetCompiledGraphs() const { return CompiledGraphs; }
+
 	private:
 		friend struct FCompiler;
 
diff --git a/Engine/Source/Runtime/Experimental/Dataflow/Engine/Public/Dataflow/DataflowObject.h b/Engine/Source/Runtime/Experimental/Dataflow/Engine/Public/Dataflow/DataflowObject.h
--- a/Engine/Source/Runtime/Experimental/Dataflow/Engine/Public/Dataflow/DataflowObject.h
+++ b/Engine/Source/Runtime/Experimental/Dataflow/Engine/Public/Dataflow/DataflowObject.h
@@ -27,10 +27,11 @@ class UObject;
 class UDataflowEdNode;
 class UDataflowSubGraph;
 class UMaterial;
-namespace UE::Dataflow { class FGraph; }
-
 namespace UE::Dataflow
-{
+{ 
+	class FGraph;
+	struct FCompiledGraph;
+
 	enum class ESubGraphChangedReason: uint8
 	{
 		Created,
@@ -227,6 +228,12 @@ public:
 	DATAFLOWENGINE_API void RefreshEdNode(TObjectPtr<UDataflowEdNode> EdNode);
 	DATAFLOWENGINE_API void RefreshEdNodeByGuid(const FGuid NodeGuid);
 
+	/** 
+	* Compile the dataflow graph if it has not been compiled since last change to the graph 
+	* returns a shared pointer to the compiled graph
+	*/
+	DATAFLOWENGINE_API 	TSharedPtr<const UE::Dataflow::FCompiledGraph> CompileGraphIfNeeded();
+
 #if WITH_EDITORONLY_DATA
 
 	/*
@@ -283,6 +290,13 @@ private:
 	UPROPERTY()
 	TArray<TObjectPtr<UDataflowSubGraph>> DataflowSubGraphs; 
 
+	/** Stores the topology GUID used fro the most recent compilation (see CompiledGraphs) */
+	FGuid LastCompiledTopologyGUID;
+
+	/** Contained the most recent compiled dataflow graph (see GetGraphTopologyGuid )*/
+	TMap<const UObject*, TSharedRef<UE::Dataflow::FCompiledGraph>> CompiledGraphs;
+
+
 private:
 	/** Broadcasts a notification just before the dataflow is destroyed. */
 	FMeshDeformerBeginDestroyEvent BeginDestroyEvent;
