Iad.RecordReplay initial implementation

#Synced-CL 49969322

diff --git a/Engine/Source/Runtime/Experimental/IoStore/OnDemand/Private/DebugCommands.cpp b/Engine/Source/Runtime/Experimental/IoStore/OnDemand/Private/DebugCommands.cpp
--- a/Engine/Source/Runtime/Experimental/IoStore/OnDemand/Private/DebugCommands.cpp
+++ b/Engine/Source/Runtime/Experimental/IoStore/OnDemand/Private/DebugCommands.cpp
@@ -11,6 +11,7 @@
 #include "LatencyTesting.h"
 #include "Misc/AssertionMacros.h"
 #include "OnDemandIoStore.h"
+#include "OnDemandContentInstaller.h"
 #include "String/Numeric.h"
 
 #if !UE_BUILD_SHIPPING
@@ -59,6 +60,13 @@ void FOnDemandDebugCommands::BindConsoleCommands()
 			TEXT(""),
 			FConsoleCommandDelegate::CreateRaw(this, &FOnDemandDebugCommands::RunLatencyTests),
 			ECVF_Default));
+
+	DynamicConsoleCommands.Emplace(
+		IConsoleManager::Get().RegisterConsoleCommand(
+			TEXT("iad.RecordReplayStop"),
+			TEXT(""),
+			FConsoleCommandDelegate::CreateRaw(this, &FOnDemandDebugCommands::RecordReplayStop),
+			ECVF_Default));
 #endif // !NO_CVARS
 }
 
@@ -180,6 +188,26 @@ void FOnDemandDebugCommands::RunLatencyTests() const
 	UE_CLOG(bLoggedHostGroup, LogIas, Log, TEXT("--------------------------------------------------------------------------------"));
 }
 
+void FOnDemandDebugCommands::RecordReplayStop()
+{
+	if (FOnDemandContentInstallReplayRecorder* Recorder = FOnDemandContentInstallReplayRecorder::Get())
+	{
+		TValueOrError<FOnDemandContentInstallReplay, void> Result = Recorder->StopRecording();
+		if (Result.HasError())
+		{
+			UE_LOG(LogIoStoreOnDemand, Error, TEXT("Internal Error occured while recording replay!"));
+			return;
+		}
+
+		FOnDemandContentInstallReplay Replay = Result.StealValue();
+
+		// TODO: save replay to disk
+
+		int a= 0;
+		++a;
+	}
+}
+
 } // namespace UE::IoStore
 
 #endif // !UE_BUILD_SHIPPING
diff --git a/Engine/Source/Runtime/Experimental/IoStore/OnDemand/Private/DebugCommands.h b/Engine/Source/Runtime/Experimental/IoStore/OnDemand/Private/DebugCommands.h
--- a/Engine/Source/Runtime/Experimental/IoStore/OnDemand/Private/DebugCommands.h
+++ b/Engine/Source/Runtime/Experimental/IoStore/OnDemand/Private/DebugCommands.h
@@ -31,6 +31,8 @@ private:
 	void RunRequestTest(const TArray<FString>& Args) const;
 	void RunLatencyTests() const;
 
+	void RecordReplayStop();
+
 	FOnDemandIoStore* OnDemandIoStore = nullptr;
 
 	TArray<IConsoleCommand*> DynamicConsoleCommands;
diff --git a/Engine/Source/Runtime/Experimental/IoStore/OnDemand/Private/OnDemandContentInstaller.cpp b/Engine/Source/Runtime/Experimental/IoStore/OnDemand/Private/OnDemandContentInstaller.cpp
--- a/Engine/Source/Runtime/Experimental/IoStore/OnDemand/Private/OnDemandContentInstaller.cpp
+++ b/Engine/Source/Runtime/Experimental/IoStore/OnDemand/Private/OnDemandContentInstaller.cpp
@@ -354,6 +354,37 @@ TOptional<UE::UnifiedError::FError> FOnDemandContentInstaller::FRequest::Consume
 }
 
 ////////////////////////////////////////////////////////////////////////////////
+FOnDemandContentInstaller::FRequest::FInstall::FInstall(
+	UPTRINT InstallerRequest, 
+	FOnDemandInstallArgs&& InArgs, 
+	FOnDemandInstallCompleted&& InOnCompleted, 
+	FOnDemandInstallProgressed&& InOnProgress)
+	: Args(MoveTemp(InArgs))
+	, OnCompleted(MoveTemp(InOnCompleted))
+	, OnProgress(MoveTemp(InOnProgress)) 
+	, Request(MakeShared<FOnDemandInternalInstallRequest, ESPMode::ThreadSafe>(InstallerRequest))
+{
+#if !UE_BUILD_SHIPPING
+	// Ideally, this would be in FOnDemandInstallRequest::FOnDemandInstallRequest but that constructor
+	// doesn't have access to all the data because the type erasure we do on the request.
+	if (FOnDemandContentInstallReplayRecorder* Recorder = FOnDemandContentInstallReplayRecorder::Get())
+	{
+		Recorder->RecordInstall(
+			UPTRINT(Request.Get()),
+			Args.ContentHandle.IsValid() ? Args.ContentHandle.Handle->HandleId() : 0,
+			Args.MountId,
+			Args.TagSets,
+			Args.PackageIds,
+			Args.Priority,
+			Args.Options);
+	}
+
+	// FOnDemandInstallRequest::~FOnDemandInstallRequest will call RecordInstallDestroyed()
+	// In general we want to record the public API calls, not the lifetime of internal data structures.
+	// See comments in FOnDemandInstallRequest::FOnDemandInstallRequest
+#endif // !UE_BUILD_SHIPPING
+}
+
 FOnDemandContentInstaller::FOnDemandContentInstaller(FOnDemandIoStore& InIoStore)
 	: IoStore(InIoStore)
 	, InstallerPipe(TEXT("IoStoreOnDemandInstallerPipe"))
@@ -373,11 +404,10 @@ FSharedInternalInstallRequest FOnDemandContentInstaller::EnqueueInstallRequest(
 	FRequest* Request = nullptr;
 	{
 		TUniqueLock Lock(Mutex);
-		Request = RequestAllocator.Construct(MoveTemp(Args), MoveTemp(OnCompleted), MoveTemp(OnProgress));
+		Request = RequestAllocator.Alloc();
 	}
 
-	FSharedInternalInstallRequest InstallRequest = MakeShared<FOnDemandInternalInstallRequest, ESPMode::ThreadSafe>(UPTRINT(Request));
-	Request->AsInstall().Request = InstallRequest;
+	new(Request) FRequest(MoveTemp(Args), MoveTemp(OnCompleted), MoveTemp(OnProgress));
 
 	FOnDemandContentInstallerStats::OnRequestEnqueued();
 
@@ -386,7 +416,8 @@ FSharedInternalInstallRequest FOnDemandContentInstaller::EnqueueInstallRequest(
 		[this, Request] { ProcessInstallRequest(*Request); },
 		UE::Tasks::ETaskPriority::BackgroundLow);
 
-	return InstallRequest;
+	FRequest::FInstall& InstallRequest = Request->AsInstall();
+	return InstallRequest.Request;
 }
 
 void FOnDemandContentInstaller::EnqueuePurgeRequest(FOnDemandPurgeArgs&& Args, FOnDemandPurgeCompleted&& OnCompleted)
diff --git a/Engine/Source/Runtime/Experimental/IoStore/OnDemand/Private/OnDemandContentInstaller.h b/Engine/Source/Runtime/Experimental/IoStore/OnDemand/Private/OnDemandContentInstaller.h
--- a/Engine/Source/Runtime/Experimental/IoStore/OnDemand/Private/OnDemandContentInstaller.h
+++ b/Engine/Source/Runtime/Experimental/IoStore/OnDemand/Private/OnDemandContentInstaller.h
@@ -95,12 +95,10 @@ class FOnDemandContentInstaller
 		struct FInstall
 		{
 			FInstall(
+				UPTRINT InstallerRequest,
 				FOnDemandInstallArgs&& InArgs,
 				FOnDemandInstallCompleted&& InOnCompleted,
-				FOnDemandInstallProgressed&& InOnProgress)
-					: Args(MoveTemp(InArgs))
-					, OnCompleted(MoveTemp(InOnCompleted))
-					, OnProgress(MoveTemp(InOnProgress)) { }
+				FOnDemandInstallProgressed&& InOnProgress);
 
 			FOnDemandInstallArgs						Args;
 			FOnDemandInstallCompleted					OnCompleted;
@@ -157,7 +155,7 @@ class FOnDemandContentInstaller
 
 		FRequest(FOnDemandInstallArgs&& Args, FOnDemandInstallCompleted&& OnCompleted, FOnDemandInstallProgressed&& OnProgress)
 		{
-			Variant.Emplace<FInstall>(MoveTemp(Args), MoveTemp(OnCompleted), MoveTemp(OnProgress));
+			Variant.Emplace<FInstall>(UPTRINT(this), MoveTemp(Args), MoveTemp(OnCompleted), MoveTemp(OnProgress));
 		}
 
 		FRequest(FOnDemandPurgeArgs&& Args, FOnDemandPurgeCompleted&& OnCompleted)
diff --git a/Engine/Source/Runtime/Experimental/IoStore/OnDemand/Private/OnDemandIoStore.cpp b/Engine/Source/Runtime/Experimental/IoStore/OnDemand/Private/OnDemandIoStore.cpp
--- a/Engine/Source/Runtime/Experimental/IoStore/OnDemand/Private/OnDemandIoStore.cpp
+++ b/Engine/Source/Runtime/Experimental/IoStore/OnDemand/Private/OnDemandIoStore.cpp
@@ -538,6 +538,16 @@ FIoStatus FOnDemandIoStore::InitializePostHotfix()
 #endif
 	}
 
+#if !UE_BUILD_SHIPPING
+	if (FParse::Param(FCommandLine::Get(), TEXT("Iad.RecordReplay")))
+	{
+		// If we don't start recording at startup, then we can miss content handles being created.
+		// Recording should only be started via command line; however, it should be possible to stop recording at any time.
+		static FOnDemandContentInstallReplayRecorder Recorder;
+		FOnDemandContentInstallReplayRecorder::StartRecording(&Recorder);
+	}
+#endif // !UE_BUILD_SHIPPING
+
 	return FIoStatus::Ok;
 }
 
@@ -1656,7 +1666,7 @@ void FOnDemandIoStore::CreateContainersFromToc(
 		Container->BlockHashes			= MoveTemp(ContainerEntry.BlockHashes);
 		Container->ContainerId			= ContainerEntry.ContainerId;
 		Container->HostGroupName		= Toc.Header.HostGroupName.IsEmpty() ? FOnDemandHostGroup::DefaultName : FName(Toc.Header.HostGroupName);
-		Container->RelativeUrl			= UE::FIoRelativeUrl::From(Container->ChunksDirectory);
+		Container->RelativeUrl			= FIoRelativeUrl::From(Container->ChunksDirectory);
 
 		Container->CompressionFormats.Reserve(1);
 		Container->CompressionFormats.Add(CompressionFormat);
diff --git a/Engine/Source/Runtime/Experimental/IoStore/OnDemand/Private/OnDemandIoStore.h b/Engine/Source/Runtime/Experimental/IoStore/OnDemand/Private/OnDemandIoStore.h
--- a/Engine/Source/Runtime/Experimental/IoStore/OnDemand/Private/OnDemandIoStore.h
+++ b/Engine/Source/Runtime/Experimental/IoStore/OnDemand/Private/OnDemandIoStore.h
@@ -121,7 +121,7 @@ struct FOnDemandContainer
 	TArrayView<FIoChunkId>					ChunkIds;
 	TArrayView<FOnDemandChunkEntry> 		ChunkEntries;
 	TArray<FOnDemandChunkEntryReferences>	ChunkEntryReferences;
-	UE::FIoRelativeUrl						RelativeUrl;
+	FIoRelativeUrl							RelativeUrl;
 	FIoContainerId							ContainerId;
 	uint32									BlockSize = 0;
 	FName									HostGroupName = FOnDemandHostGroup::DefaultName;
@@ -251,7 +251,7 @@ struct FOnDemandChunkInfo
 	inline TConstArrayView<uint32>			Blocks() const;
 	inline TConstArrayView<FIoBlockHash>	BlockHashes() const;
 	FAnsiStringView							ChunksDirectory() const { return SharedContainer->ChunksDirectory; }
-	UE::FIoRelativeUrl						RelativeUrl() const { return SharedContainer->RelativeUrl; }
+	FIoRelativeUrl							RelativeUrl() const { return SharedContainer->RelativeUrl; }
 	const FOnDemandChunkEntry&				ChunkEntry() const { return Entry; }
 	inline const FIASHostGroup&				HostGroup() const;
 	FName									HostGroupName() const { return SharedContainer->HostGroupName; }
diff --git a/Engine/Source/Runtime/Experimental/IoStore/OnDemandCore/Internal/IO/IoStoreOnDemandInternals.h b/Engine/Source/Runtime/Experimental/IoStore/OnDemandCore/Internal/IO/IoStoreOnDemandInternals.h
--- a/Engine/Source/Runtime/Experimental/IoStore/OnDemandCore/Internal/IO/IoStoreOnDemandInternals.h
+++ b/Engine/Source/Runtime/Experimental/IoStore/OnDemandCore/Internal/IO/IoStoreOnDemandInternals.h
@@ -2,6 +2,8 @@
 
 #pragma once
 
+#include "Async/Mutex.h"
+#include "Containers/PagedArray.h"
 #include "IO/IoStoreOnDemand.h"
 #include "HAL/PlatformTime.h"
 
@@ -42,13 +44,9 @@ struct FOnDemandWeakContentHandle
 class FOnDemandInternalContentHandle
 {
 public:
-	FOnDemandInternalContentHandle()
-		: DebugName(TEXT("NoName"))
-	{ }
+	FOnDemandInternalContentHandle();
 
-	FOnDemandInternalContentHandle(FSharedString InDebugName)
-		: DebugName(InDebugName)
-	{ }
+	FOnDemandInternalContentHandle(FSharedString InDebugName);
 
 	UE_API ~FOnDemandInternalContentHandle();
 
@@ -83,6 +81,147 @@ public:
 	virtual uint32 GetTimeToFirstByteMs() const = 0;
 };
 
+#if !UE_BUILD_SHIPPING
+struct FOnDemandContentInstallReplay
+{
+	// TODO: record mounting/unmounting containers
+
+	struct FContentHandleEventData
+	{
+		FString DebugName;
+	};
+
+	struct FContentHandleDestroyedEventData
+	{
+		int64 ContentHandleEventId = -1;
+	};
+
+	struct FInstallEventData
+	{
+		FString MountId;
+		TArray<FString> TagSets;
+		TArray<FPackageId> PackageIds;
+		int64 ContentHandleEventId = -1;
+		int32 Priority = 0;
+		EOnDemandInstallOptions Options = EOnDemandInstallOptions::None;
+	};
+
+	struct FCancelInstallEventData
+	{
+		int64 InstallEventId = -1;
+	};
+
+	struct FUpdatePriorityEventData
+	{
+		int64 InstallEventId = -1;
+		int32 NewPriority = 0;
+	};
+
+	struct FInstallDestroyedEventData
+	{
+		int64 InstallEventId = 0;
+	};
+
+	using FEventVariant = TVariant<
+		FContentHandleEventData,
+		FContentHandleDestroyedEventData,
+		FInstallEventData,
+		FCancelInstallEventData,
+		FUpdatePriorityEventData,
+		FInstallDestroyedEventData>;
+
+	struct FEvent
+	{
+		int64 TimeOffset = 0; // Time since start of replay
+		FEventVariant EventData;
+	};
+
+	TArray<FEvent> Events;
+};
+
+class FOnDemandContentInstallReplayRecorder
+{
+private:
+	static FOnDemandContentInstallReplayRecorder* Instance;
+
+public:
+	UE_API static FOnDemandContentInstallReplayRecorder* Get();
+
+	UE_API static void StartRecording(FOnDemandContentInstallReplayRecorder* InInstance);
+
+	UE_API TValueOrError<FOnDemandContentInstallReplay, void> StopRecording();
+
+	// Records the public content handle, no the internal implementation
+	UE_API void RecordContentHandle(UPTRINT HandleId, FString&& DebugName);
+
+	UE_API void RecordContentHandleDestroyed(UPTRINT HandleId);
+
+	UE_API void RecordInstall(
+		UPTRINT RequestId,
+		UPTRINT HandleId,
+		const FString& MountId,
+		const TArray<FString>& TagSets,
+		const TArray<FPackageId>& PackageIds,
+		int32 Priority,
+		EOnDemandInstallOptions Options);
+
+	UE_API void RecordInstall(
+		UPTRINT RequestId,
+		UPTRINT HandleId,
+		FString&& MountId,
+		TArray<FString>&& TagSets,
+		TArray<FPackageId>&& PackageIds,
+		int32 Priority,
+		EOnDemandInstallOptions Options);
+
+	UE_API void RecordCancel(UPTRINT RequestId);
+
+	UE_API void RecordUpdatePriority(UPTRINT RequestId, int32 NewPriority);
+
+	UE_API void RecordInstallDestroyed(UPTRINT RequestId);
+
+	bool IsValid() const { return bValid; }
+
+private:
+	void StartRecordingInternal();
+
+	void Reset();
+
+	int64 GetTimeOffset() const;
+
+	template<typename EventType>
+	int64 AddEvent(EventType&& InEvent)
+	{
+		const int64 TimeOffset = GetTimeOffset();
+		const int64 EventIndex = Events.Add(
+			FOnDemandContentInstallReplay::FEvent
+			{
+				.TimeOffset = TimeOffset,
+				.EventData = FOnDemandContentInstallReplay::FEventVariant(
+					TInPlaceType<EventType>(), MoveTemp(InEvent))
+			}
+		);
+		return EventIndex;
+	}
+
+private:
+
+	FMutex Mutex;
+
+	struct FContentHandleRecord
+	{
+		int64 EventIndex = -1;
+		uint64 RefCount = 0;
+	};
+
+	TMap<UPTRINT, FContentHandleRecord> ContentHandleEventToIndex;
+	TMap<UPTRINT, int64> InstallEventToIndex;
+	TPagedArray<FOnDemandContentInstallReplay::FEvent, 4096> Events;
+	FDateTime StartTime;
+	bool bValid = true;
+};
+#endif // !UE_BUILD_SHIPPING
+
 } // namespace UE::IoStore
 
 ////////////////////////////////////////////////////////////////////////////////////////////////////////////////////////////////////////////////////////////////
diff --git a/Engine/Source/Runtime/Experimental/IoStore/OnDemandCore/Private/IoStoreOnDemandInternals.cpp b/Engine/Source/Runtime/Experimental/IoStore/OnDemandCore/Private/IoStoreOnDemandInternals.cpp
--- a/Engine/Source/Runtime/Experimental/IoStore/OnDemandCore/Private/IoStoreOnDemandInternals.cpp
+++ b/Engine/Source/Runtime/Experimental/IoStore/OnDemandCore/Private/IoStoreOnDemandInternals.cpp
@@ -2,6 +2,7 @@
 
 #include "IO/IoStoreOnDemandInternals.h"
 
+#include "Async/UniqueLock.h"
 #include "Logging/StructuredLog.h"
 
 namespace UE::IoStore
@@ -28,6 +29,16 @@ FOnDemandWeakContentHandle FOnDemandWeakContentHandle::FromUnsafeHandle(UPTRINT
 	};
 }
 
+FOnDemandInternalContentHandle::FOnDemandInternalContentHandle()
+	: FOnDemandInternalContentHandle(FSharedString(TEXT("NoName")))
+{
+}
+
+FOnDemandInternalContentHandle::FOnDemandInternalContentHandle(FSharedString InDebugName)
+	: DebugName(MoveTemp(InDebugName))
+{
+}
+
 ///////////////////////////////////////////////////////////////////////////////
 FOnDemandInternalContentHandle::~FOnDemandInternalContentHandle()
 {
@@ -44,12 +55,93 @@ FOnDemandContentHandle::FOnDemandContentHandle()
 
 FOnDemandContentHandle::~FOnDemandContentHandle()
 {
+#if !UE_BUILD_SHIPPING
+	if (IsValid())
+	{
+		if (FOnDemandContentInstallReplayRecorder* Recorder = FOnDemandContentInstallReplayRecorder::Get())
+		{
+			Recorder->RecordContentHandleDestroyed(Handle->HandleId());
+		}
+	}
+#endif // !UE_BUILD_SHIPPING
+}
+
+FOnDemandContentHandle::FOnDemandContentHandle(const FOnDemandContentHandle& Other)
+	: Handle(Other.Handle)
+{
+#if !UE_BUILD_SHIPPING
+	if (IsValid())
+	{
+		if (FOnDemandContentInstallReplayRecorder* Recorder = FOnDemandContentInstallReplayRecorder::Get())
+		{
+			Recorder->RecordContentHandle(Handle->HandleId(), FString()); // Only need to record debug name on first creation
+		}
+	}
+#endif // !UE_BUILD_SHIPPING
+}
+
+FOnDemandContentHandle& FOnDemandContentHandle::operator=(const FOnDemandContentHandle& Other)
+{
+	if (this == &Other)
+	{
+		return *this;
+	}
+
+#if !UE_BUILD_SHIPPING
+	FOnDemandContentInstallReplayRecorder* Recorder = FOnDemandContentInstallReplayRecorder::Get();
+	if (IsValid() && Recorder)
+	{
+		Recorder->RecordContentHandleDestroyed(Handle->HandleId());
+	}
+#endif // !UE_BUILD_SHIPPING
+
+	Handle = Other.Handle;
+
+#if !UE_BUILD_SHIPPING
+	if (IsValid() && Recorder)
+	{
+		Recorder->RecordContentHandle(Handle->HandleId(), FString()); // Only need to record debug name on first creation
+	}
+#endif // !UE_BUILD_SHIPPING
+
+	return *this;
+}
+
+FOnDemandContentHandle& FOnDemandContentHandle::operator=(FOnDemandContentHandle&& Other)
+{
+	if (this == &Other)
+	{
+		return *this;
+	}
+
+#if !UE_BUILD_SHIPPING
+	if (IsValid())
+	{
+		if (FOnDemandContentInstallReplayRecorder* Recorder = FOnDemandContentInstallReplayRecorder::Get())
+		{
+			Recorder->RecordContentHandleDestroyed(Handle->HandleId());
+		}
+	}
+#endif // !UE_BUILD_SHIPPING
+
+	Handle = MoveTemp(Other.Handle);
+
+	return *this;
 }
 
 FOnDemandContentHandle FOnDemandContentHandle::Create()
 {
 	FOnDemandContentHandle NewHandle;
 	NewHandle.Handle = MakeShared<FOnDemandInternalContentHandle, ESPMode::ThreadSafe>();
+
+#if !UE_BUILD_SHIPPING
+	FSharedString DebugName = NewHandle.Handle->DebugName;
+	if (FOnDemandContentInstallReplayRecorder* Recorder = FOnDemandContentInstallReplayRecorder::Get())
+	{
+		Recorder->RecordContentHandle(NewHandle.Handle->HandleId(), FString::ConstructFromPtrSize(*DebugName, DebugName.Len()));
+	}
+#endif // !UE_BUILD_SHIPPING
+
 	return NewHandle;
 }
 
@@ -57,6 +149,14 @@ FOnDemandContentHandle FOnDemandContentHandle::Create(FSharedString DebugName)
 {
 	FOnDemandContentHandle NewHandle;
 	NewHandle.Handle = MakeShared<FOnDemandInternalContentHandle, ESPMode::ThreadSafe>(DebugName);
+
+#if !UE_BUILD_SHIPPING
+	if (FOnDemandContentInstallReplayRecorder* Recorder = FOnDemandContentInstallReplayRecorder::Get())
+	{
+		Recorder->RecordContentHandle(NewHandle.Handle->HandleId(), FString::ConstructFromPtrSize(*DebugName, DebugName.Len()));
+	}
+#endif // !UE_BUILD_SHIPPING
+
 	return NewHandle;
 }
 
@@ -79,10 +179,22 @@ FOnDemandInstallRequest::FOnDemandInstallRequest(FWeakOnDemandIoStore InIoStore,
 	: IoStore(InIoStore)
 	, Request(InternalRequest)
 {
+	// Unfortunately, request creation is recorded by FOnDemandContentInstaller::FRequest::FInstall::FInstall()
+	// because some of the information needed for the replay event is type erased and we don't have access to it here.
+	// Luckily, this is a move only type, so there are no other places that RecordInstall() needs to be called.
 }
 
 FOnDemandInstallRequest::~FOnDemandInstallRequest()
 {
+#if !UE_BUILD_SHIPPING
+	if (Request)
+	{
+		if (FOnDemandContentInstallReplayRecorder* Recorder = FOnDemandContentInstallReplayRecorder::Get())
+		{
+			Recorder->RecordInstallDestroyed(UPTRINT(Request.Get()));
+		}
+	}
+#endif // !UE_BUILD_SHIPPING
 }
 
 FOnDemandInstallRequest::FOnDemandInstallRequest(FOnDemandInstallRequest&& Other)
@@ -110,6 +222,13 @@ void FOnDemandInstallRequest::Cancel()
 		if (TSharedPtr<IOnDemandIoStore, ESPMode::ThreadSafe> Pinned = IoStore.Pin(); Pinned.IsValid())
 		{
 			Pinned->CancelInstallRequest(Request);
+
+#if !UE_BUILD_SHIPPING
+			if (FOnDemandContentInstallReplayRecorder* Recorder = FOnDemandContentInstallReplayRecorder::Get())
+			{
+				Recorder->RecordCancel(UPTRINT(Request.Get()));
+			}
+#endif // !UE_BUILD_SHIPPING
 		}
 	}
 }
@@ -121,9 +240,234 @@ void FOnDemandInstallRequest::UpdatePriority(int32 NewPriority)
 		if (TSharedPtr<IOnDemandIoStore, ESPMode::ThreadSafe> Pinned = IoStore.Pin(); Pinned.IsValid())
 		{
 			Pinned->UpdateInstallRequestPriority(Request, NewPriority);
+
+#if !UE_BUILD_SHIPPING
+			if (FOnDemandContentInstallReplayRecorder* Recorder = FOnDemandContentInstallReplayRecorder::Get())
+			{
+				Recorder->RecordUpdatePriority(UPTRINT(Request.Get()), NewPriority);
+			}
+#endif // !UE_BUILD_SHIPPING
+		}
+	}
+}
+
+#if !UE_BUILD_SHIPPING
+FOnDemandContentInstallReplayRecorder* FOnDemandContentInstallReplayRecorder::Instance = nullptr;
+
+FOnDemandContentInstallReplayRecorder* FOnDemandContentInstallReplayRecorder::Get()
+{
+	return FOnDemandContentInstallReplayRecorder::Instance;
+}
+
+void FOnDemandContentInstallReplayRecorder::StartRecording(FOnDemandContentInstallReplayRecorder* InInstance)
+{
+	// This should only be called by init, so no need to take the lock yet
+	check(Instance == nullptr);
+	Instance = InInstance;
+
+	Instance->StartRecordingInternal();
+}
+
+void FOnDemandContentInstallReplayRecorder::StartRecordingInternal()
+{
+	TUniqueLock Lock(Mutex);
+
+	Reset();
+	StartTime = FDateTime::UtcNow();
+}
+
+TValueOrError<FOnDemandContentInstallReplay, void> FOnDemandContentInstallReplayRecorder::StopRecording()
+{
+	TUniqueLock Lock(Mutex);
+
+	if (!IsValid())
+	{
+		Reset();
+		return MakeError();
+	}
+
+	FOnDemandContentInstallReplay Replay;
+	MoveTemp(Events).ToArray(Replay.Events);
+
+	Reset();
+	return MakeValue(MoveTemp(Replay));
+}
+
+void FOnDemandContentInstallReplayRecorder::RecordContentHandle(UPTRINT HandleId, FString&& DebugName)
+{
+	if (ensure(HandleId) == false)
+	{
+		// still valid but this shouldn't happen.
+		return;
+	}
+
+	TUniqueLock Lock(Mutex);
+
+	const int64 EventIndex = AddEvent(
+		FOnDemandContentInstallReplay::FContentHandleEventData{ MoveTemp(DebugName) });
+
+	FContentHandleRecord& Record = ContentHandleEventToIndex.FindOrAdd(HandleId, FContentHandleRecord{ EventIndex, 0 });
+	++Record.RefCount;
+}
+
+void FOnDemandContentInstallReplayRecorder::RecordContentHandleDestroyed(UPTRINT HandleId)
+{
+	TUniqueLock Lock(Mutex);
+
+	FContentHandleRecord* Record = ContentHandleEventToIndex.Find(HandleId);
+	if (ensure(Record) == false)
+	{
+		bValid = false;
+		return;
+	}
+
+	if (ensure(Record->RefCount != 0) == false)
+	{
+		bValid = false;
+		return;
+	}
+
+	--Record->RefCount;
+
+	if (Record->RefCount == 0)
+	{
+		int64 ContentHandleEventIndex = Record->EventIndex;
+		ContentHandleEventToIndex.Remove(HandleId);
+
+		AddEvent(FOnDemandContentInstallReplay::FContentHandleDestroyedEventData{ ContentHandleEventIndex });
+	}
+}
+
+void FOnDemandContentInstallReplayRecorder::RecordInstall(
+	UPTRINT RequestId, 
+	UPTRINT HandleId, 
+	const FString& MountId, 
+	const TArray<FString>& TagSets, 
+	const TArray<FPackageId>& PackageIds, 
+	int32 Priority, 
+	EOnDemandInstallOptions Options)
+{
+	RecordInstall(
+		RequestId,
+		HandleId,
+		CopyTemp(MountId),
+		CopyTemp(TagSets),
+		CopyTemp(PackageIds),
+		Priority,
+		Options
+	);
+}
+
+void FOnDemandContentInstallReplayRecorder::RecordInstall(
+	UPTRINT RequestId,
+	UPTRINT HandleId,
+	FString&& MountId,
+	TArray<FString>&& TagSets,
+	TArray<FPackageId>&& PackageIds,
+	int32 Priority,
+	EOnDemandInstallOptions Options)
+{
+	if (ensure(RequestId) == false)
+	{
+		// still valid but this shouldn't happen.
+		return;
+	}
+
+	TUniqueLock Lock(Mutex);
+
+	int64 ContentHandleEventIndex = -1;
+	if (HandleId)
+	{
+		// If the content handle is valid, it should exist in the lookup table
+		FContentHandleRecord* Record = ContentHandleEventToIndex.Find(HandleId);
+		if (ensure(Record) == false)
+		{
+			bValid = false;
+			return;
 		}
+
+		ContentHandleEventIndex = Record->EventIndex;
 	}
+
+	// If the content handle is invalid, the install request will fail at runtime but the game
+	// still made the request, so we will record it.
+
+	const int64 EventIndex = AddEvent(
+		FOnDemandContentInstallReplay::FInstallEventData
+		{
+			.MountId = MoveTemp(MountId),
+			.TagSets = MoveTemp(TagSets),
+			.PackageIds = MoveTemp(PackageIds),
+			.ContentHandleEventId = ContentHandleEventIndex,
+			.Priority = Priority,
+			.Options = Options
+		}
+	);
+
+	InstallEventToIndex.Add(RequestId, EventIndex);
+}
+
+void FOnDemandContentInstallReplayRecorder::RecordCancel(UPTRINT RequestId)
+{
+	TUniqueLock Lock(Mutex);
+
+	const int64 InstallEventIndex = InstallEventToIndex.FindRef(RequestId, -1);
+	if (ensure(InstallEventIndex >= 0) == false)
+	{
+		bValid = false;
+		return;
+	}
+
+	AddEvent(FOnDemandContentInstallReplay::FCancelInstallEventData{ InstallEventIndex });
+}
+
+void FOnDemandContentInstallReplayRecorder::RecordUpdatePriority(UPTRINT RequestId, int32 NewPriority)
+{
+	TUniqueLock Lock(Mutex);
+
+	const int64 InstallEventIndex = InstallEventToIndex.FindRef(RequestId, -1);
+	if (ensure(InstallEventIndex >= 0) == false)
+	{
+		bValid = false;
+		return;
+	}
+
+	AddEvent(FOnDemandContentInstallReplay::FUpdatePriorityEventData
+	{
+		.InstallEventId = InstallEventIndex,
+		.NewPriority = NewPriority
+	});
+}
+
+void FOnDemandContentInstallReplayRecorder::RecordInstallDestroyed(UPTRINT RequestId)
+{
+	TUniqueLock Lock(Mutex);
+
+	int64 InstallEventIndex = -1;
+	if (ensure(InstallEventToIndex.RemoveAndCopyValue(RequestId, InstallEventIndex)) == false)
+	{
+		bValid = false;
+		return;
+	}
+
+	AddEvent(FOnDemandContentInstallReplay::FInstallDestroyedEventData{ InstallEventIndex });
+}
+
+void FOnDemandContentInstallReplayRecorder::Reset()
+{
+	ContentHandleEventToIndex.Empty();
+	InstallEventToIndex.Empty();
+	Events.Empty();
+	StartTime = FDateTime();
+	bValid = true;
+}
+
+int64 FOnDemandContentInstallReplayRecorder::GetTimeOffset() const
+{
+	const int64 TimeOffset = (FDateTime::UtcNow() - StartTime).GetTicks();
+	return TimeOffset;
 }
+#endif // !UE_BUILD_SHIPPING
 
 } // namespace UE::IoStore
 
diff --git a/Engine/Source/Runtime/Experimental/IoStore/OnDemandCore/Public/IO/IoStoreOnDemand.h b/Engine/Source/Runtime/Experimental/IoStore/OnDemandCore/Public/IO/IoStoreOnDemand.h
--- a/Engine/Source/Runtime/Experimental/IoStore/OnDemandCore/Public/IO/IoStoreOnDemand.h
+++ b/Engine/Source/Runtime/Experimental/IoStore/OnDemandCore/Public/IO/IoStoreOnDemand.h
@@ -91,6 +91,13 @@ public:
 	UE_API FOnDemandContentHandle();
 	/** Destroy the handle and release any referenced content. */ 
 	UE_API ~FOnDemandContentHandle();
+
+	UE_API FOnDemandContentHandle(const FOnDemandContentHandle& Other);
+	FOnDemandContentHandle(FOnDemandContentHandle&& Other) = default;
+
+	UE_API FOnDemandContentHandle& operator=(const FOnDemandContentHandle& Other);
+	UE_API FOnDemandContentHandle& operator=(FOnDemandContentHandle&& Other);
+
 	/** Destroy the handle and release any referenced content. */ 
 	void Reset() { Handle.Reset(); }
 	/** Returns whether the handle is valid. */
@@ -214,7 +221,7 @@ struct FOnDemandMountResult
 using FOnDemandMountCompleted = TUniqueFunction<void(FOnDemandMountResult)>;
 
 /** Options for controlling the behavior of the install request. */
-enum class EOnDemandInstallOptions
+enum class EOnDemandInstallOptions : uint16
 {
 	/** No additional options. */
 	None						= 0,
