Added unit tests for converting legacy on-demand container TOC to the new format.

#rnx

#Synced-CL 49920676

diff --git a/Engine/Source/Runtime/Experimental/IoStore/OnDemandCore/Private/Serialization/OnDemandContainerToc.cpp b/Engine/Source/Runtime/Experimental/IoStore/OnDemandCore/Private/Serialization/OnDemandContainerToc.cpp
--- a/Engine/Source/Runtime/Experimental/IoStore/OnDemandCore/Private/Serialization/OnDemandContainerToc.cpp
+++ b/Engine/Source/Runtime/Experimental/IoStore/OnDemandCore/Private/Serialization/OnDemandContainerToc.cpp
@@ -20,6 +20,18 @@
 #include "Templates/UniquePtr.h"
 #include "Templates/UnrealTemplate.h"
 
+#ifndef WITH_IOSTORE_ONDEMAND_TESTS
+#define WITH_IOSTORE_ONDEMAND_TESTS 0
+#endif
+
+#if WITH_IOSTORE_ONDEMAND_TESTS
+#include "TestHarness.h"
+#include "TestMacros/Assertions.h"
+#include <catch2/generators/catch_generators.hpp>
+#include "IO/IoChunkId.h"
+#include "IO/PackageId.h"
+#endif
+
 namespace UE::IoStore::Serialization
 {
 
@@ -558,4 +570,198 @@ TIoStatusOr<FOnDemandTocReader> FOnDemandTocReader::Read(const FString& Filename
 
 } // V1 
 
+///////////////////////////////////////////////////////////////////////////////
+#if WITH_IOSTORE_ONDEMAND_TESTS
+
+class FTmpDirectoryScope
+{
+public:
+	explicit FTmpDirectoryScope(const FString& InDir)
+		: Ifm(IFileManager::Get())
+		, Dir(InDir)
+	{
+		const bool bTree			= true;
+		const bool bRequireExists	= false;
+		Ifm.DeleteDirectory(*Dir, bRequireExists, bTree);
+		Ifm.MakeDirectory(*Dir, bTree);
+	}
+
+	~FTmpDirectoryScope()
+	{
+		const bool bTree			= true;
+		const bool bRequireExists	= false;
+		Ifm.DeleteDirectory(*Dir, bRequireExists, bTree);
+	}
+private:
+	IFileManager& Ifm;
+	FString Dir;
+};
+
+TEST_CASE("IoStore::OnDemand::Serialization::TOC::V1", "[IoStoreOnDemand][Serialization]")
+{
+	using namespace UE::IoStore::Serialization;
+	using namespace UE::IoStore::Serialization::V1;
+
+	const FString TestBaseDir = "TocTestTmpDir";
+
+	SECTION("SetBlockInfoFromOffsetAndCount")
+	{
+		const FOnDemandChunkBlockInfo Nfo = FOnDemandChunkBlockInfo::FromOffsetAndCount(46, 2);
+		CHECK(Nfo.Size() == 0);
+		CHECK(Nfo.Hash() == 0);
+		CHECK(Nfo.Offset() == 46);
+		CHECK(Nfo.Count() == 2);
+	}
+
+	SECTION("SetBlockInfoFromSizeAndHash")
+	{
+		const FOnDemandChunkBlockInfo Nfo = FOnDemandChunkBlockInfo::FromSizeAndHash(46, 2);
+		CHECK(Nfo.Size() == 46);
+		CHECK(Nfo.Hash() == 2);
+		CHECK(Nfo.Offset() == ~uint32(0));
+		CHECK(Nfo.Count() == 1);
+	}
+
+	SECTION("ConvertFromLegacyToc")
+	{
+		const FString TocFilePath = TestBaseDir / TEXT("test.uondemandtoc");
+		FTmpDirectoryScope _(TestBaseDir);
+		UE::IoStore::FOnDemandToc LegacyToc;
+
+		// Create a legacy TOC
+		{
+			auto AddEntry = [](FOnDemandTocContainerEntry& Container, FOnDemandTocEntry Entry)
+			{
+				check(Entry.BlockCount > 0);
+				Entry.BlockOffset = Container.BlockSizes.Num();
+				for (uint32 Idx = 0; Idx < Entry.BlockCount; ++Idx)
+				{
+					Container.BlockSizes.Add(Idx + 1);
+					Container.BlockHashes.Add(Idx + 1);
+				}
+				Container.Entries.Add(Entry);
+			};
+
+			LegacyToc.Header.BlockSize = 64 << 10;
+			LegacyToc.Header.CompressionFormat = TEXT("Oodle");
+			LegacyToc.Header.ChunksDirectory = TEXT("iostore/test");
+			LegacyToc.Header.HostGroupName = TEXT("myhost");
+
+			FOnDemandTocContainerEntry& Container = LegacyToc.Containers.AddDefaulted_GetRef();
+			Container.ContainerName = TEXT("test");
+			Container.ContainerId = FIoContainerId::FromName(FName(*Container.ContainerName));
+
+			AddEntry(Container, FOnDemandTocEntry
+			{
+				.ChunkId	= CreatePackageDataChunkId(FPackageId::FromName(FName(TEXT("MyAssetA")))),
+				.RawSize	= 42,
+				.BlockCount = 1
+			});
+			AddEntry(Container, FOnDemandTocEntry
+			{
+				.ChunkId	= CreatePackageDataChunkId(FPackageId::FromName(FName(TEXT("MyAssetB")))),
+				.RawSize	= 42,
+				.BlockCount = 2
+			});
+			AddEntry(Container, FOnDemandTocEntry
+			{
+				.ChunkId	= CreatePackageDataChunkId(FPackageId::FromName(FName(TEXT("MyAssetC")))),
+				.RawSize	= 42,
+				.BlockCount = 3
+			});
+
+			FOnDemandTocTagSet& TagSet = LegacyToc.TagSets.AddDefaulted_GetRef();
+			TagSet.Tag = TEXT("TheTag");
+			FOnDemandTocTagSetPackageList& List = TagSet.Packages.AddDefaulted_GetRef();
+			List.ContainerIndex = 0;
+			List.PackageIndicies.Add(1);
+		}
+
+		// Serialize the legacy TOC to the latest format
+		{
+
+			FOnDemandTocWriter Writer;
+			TIoStatusOr<uint64> Status = Writer.Write(LegacyToc, TocFilePath);
+			CHECK(Status.IsOk());
+		}
+
+		// Assert
+		FOnDemandTocReader Reader = FOnDemandTocReader::Read(TocFilePath).ConsumeValueOrDie();
+		CHECK(FString(Reader.ChunksDirectory()) == LegacyToc.Header.ChunksDirectory);
+		CHECK(FString(Reader.CompressionFormat()) == LegacyToc.Header.CompressionFormat);
+		CHECK(FString(Reader.HostGroupName()) == LegacyToc.Header.HostGroupName);
+
+		for (int32 ContainerIndex = 0; const FOnDemandTocContainerEntry& LegacyContainer : LegacyToc.Containers)
+		{
+			TConstArrayView<uint32>			LegacyContainerBlockSizes = LegacyContainer.BlockSizes;
+			TConstArrayView<FIoBlockHash>	LegacyContainerBlockHashes = LegacyContainer.BlockHashes;
+
+			const FOnDemandContainerEntry& ContainerEntry = Reader.Containers()[ContainerIndex];
+
+			FOnDemandTocStorage Storage;
+			FOnDemandContainerTocView Container =
+				Reader.ReadContainer(ContainerEntry, Storage, EOnDemandTocReaderOptions::None)
+					.ConsumeValueOrDie();
+
+			// Check header
+			CHECK(FString(Container.Header.ContainerName()) == LegacyContainer.ContainerName);
+			CHECK(Container.Header.ContainerId() == LegacyContainer.ContainerId);
+			CHECK(Container.Header.BlockSize() == LegacyToc.Header.BlockSize);
+
+			// Check entries
+			for (const FOnDemandTocEntry& LegacyEntry : LegacyContainer.Entries)
+			{
+				const int32 EntryIndex = Algo::LowerBound(Container.ChunkIds, LegacyEntry.ChunkId);
+				CHECK(EntryIndex != INDEX_NONE);
+				CHECK(EntryIndex < Container.ChunkIds.Num());
+
+				const FOnDemandChunkEntry& Entry = Container.ChunkEntries[EntryIndex];
+				CHECK(Entry.RawSize == LegacyEntry.RawSize);
+				CHECK(Entry.BlockInfo.Count() == LegacyEntry.BlockCount);
+
+				// Check compression block size and hash
+				TConstArrayView<uint32> LegacyBlockSizes = LegacyContainerBlockSizes.Mid(LegacyEntry.BlockOffset, LegacyEntry.BlockCount);
+				TConstArrayView<uint32> BlockSizes = Entry.BlockInfo.Count() > 1
+					? Container.BlockSizes.Mid(Entry.BlockInfo.Offset(), Entry.BlockInfo.Count())
+					: TConstArrayView<uint32>(Entry.BlockInfo.Data(), 1);
+				
+				for (int32 Idx = 0; Idx < BlockSizes.Num(); ++Idx)
+				{
+					CHECK(BlockSizes[Idx] == LegacyBlockSizes[Idx]);
+				}
+
+				TConstArrayView<uint32> LegacyBlockHashes = LegacyContainerBlockHashes.Mid(LegacyEntry.BlockOffset, LegacyEntry.BlockCount);
+				TConstArrayView<uint32> BlockHashes = Entry.BlockInfo.Count() > 1
+					? Container.BlockHashes.Mid(Entry.BlockInfo.Offset(), Entry.BlockInfo.Count())
+					: TConstArrayView<uint32>(Entry.BlockInfo.Data() + 1, 1);
+				
+				for (int32 Idx = 0; Idx < BlockSizes.Num(); ++Idx)
+				{
+					CHECK(BlockHashes[Idx] == LegacyBlockHashes[Idx]);
+				}
+
+				// Check tags
+				for (const FOnDemandTagSetEntry& TagEntry : Container.TagSets)
+				{
+					FString TagName = FString(Container.Header.GetString(TagEntry.Tag));
+					const FOnDemandTocTagSet* LegacyTagSet = Algo::FindBy(LegacyToc.TagSets, TagName, &FOnDemandTocTagSet::Tag);
+					CHECK(LegacyTagSet != nullptr);
+					const FOnDemandTocTagSetPackageList* List = Algo::FindBy(LegacyTagSet->Packages, ContainerIndex, &FOnDemandTocTagSetPackageList::ContainerIndex);
+					CHECK(List != nullptr);
+
+					TConstArrayView<uint32> Indices = Container.TagSetIndices.Mid(TagEntry.Offset, TagEntry.Count);
+					for (int32 Idx = 0; Idx < Indices.Num(); ++Idx)
+					{
+						CHECK(Indices[Idx] == List->PackageIndicies[Idx]);
+					}
+				}
+			}
+
+			ContainerIndex++;
+		}
+	}
+}
+
+#endif // WITH_IOSTORE_ONDEMAND_TESTS
+
 } // UE::IoStore::Serialization
