Update I/O store on-demand to use zero-copy container TOC file(s).

- updated container file mounting to use the new container TOC serialization format.
- removed support for passing a pre-serialized container TOC through the public API to enable memory-mapped access.
- removed support for downloading container TOCs during mounting, as this path is no longer used (and can be re-enabled if needed).
- removed support for adding loose files to the container TOC, since it is no longer used.

#rnx

#Synced-CL 49940788

diff --git a/Engine/Source/Developer/IoStore/OnDemandUtilities/Private/ChunkPlugin.cpp b/Engine/Source/Developer/IoStore/OnDemandUtilities/Private/ChunkPlugin.cpp
--- a/Engine/Source/Developer/IoStore/OnDemandUtilities/Private/ChunkPlugin.cpp
+++ b/Engine/Source/Developer/IoStore/OnDemandUtilities/Private/ChunkPlugin.cpp
@@ -11,6 +11,7 @@
 #include "IO/IoStatus.h"
 #include "IO/IoStore.h"
 #include "IO/IoStoreOnDemand.h"
+#include "IO/Serialization/OnDemandContainerToc.h"
 #include "Misc/Compression.h"
 #include "Misc/DateTime.h"
 #include "Misc/FileHelper.h"
@@ -629,34 +630,24 @@ int32 WriteOnDemandToc(FOnDemandToc& OnDemandToc, const FString& ContainerName,
 
 	IFileManager& FileMgr = IFileManager::Get();
 
-	FString Filename = FPathViews::SetExtension(ContainerName, TEXT(".uondemandtoc"));
+	FString Filename = FPathViews::SetExtension(ContainerName, UE::IoStore::Serialization::FOnDemandFileExt::Toc);
 	Filename.ToLowerInline();
 
 	const FString TocPath = OutputFolder / Filename;
 
-	if (TUniquePtr<FArchive> Ar(FileMgr.CreateFileWriter(*TocPath)); Ar.IsValid())
+	UE::IoStore::Serialization::FOnDemandTocWriter Writer;
+	TIoStatusOr<uint64> Status = Writer.Write(OnDemandToc, TocPath);
+	if (Status.IsOk())
 	{
-		*Ar << OnDemandToc;
-		Ar->Close();
+		UE_LOG(LogIoStoreOnDemand, Display, TEXT("Saved on-demand container TOC '%s' (%.2lf KiB)"),
+			*TocPath, double(Status.ConsumeValueOrDie()) / 1024.0);
+		return 0;
+	}
+	UE_LOG(LogIoStoreOnDemand, Error, TEXT("Failed to serialize on-demand container TOC '%s', reason: %s"),
+		*TocPath, *Status.Status().ToString());
 
-		if (Ar->IsError())
-		{
-			UE_LOG(LogIoStoreOnDemand, Error, TEXT("Failed to serialize TOC '%s'"), *TocPath);
-			return -1;
-		}
-		else
-		{
-			const int64 TocSize = Ar->Tell();
-			UE_LOG(LogIoStoreOnDemand, Display, TEXT("Writing file '%s' (%.2lf KiB)"), *TocPath, double(TocSize) / 1024);
 
-			return 0;
-		}
-	}
-	else
-	{
-		UE_LOG(LogIoStoreOnDemand, Display, TEXT("Failed writing file '%s'"), *TocPath);
-		return -1;
-	}
+	return -1;
 }
 
 ////////////////////////////////////////////////////////////////////////////////
@@ -825,60 +816,7 @@ static int32 ChunkPluginCommandEntry(const FContext& Context)
 	// Write additional file(s)
 	if (bIncludeSigPak)
 	{
-		TRACE_CPUPROFILER_EVENT_SCOPE(ChunkPluginCommand::IncludeSigPak);
-
-		const FStringView AllowedExt[]
-		{
-			TEXTVIEW("pak"),
-			TEXTVIEW("sig"),
-		};
-
-		TStringBuilder<256> Sb;
-		TArray<FString> AdditionalFiles;
-		FileMgr.FindFiles(AdditionalFiles, *ContainerFolder, TEXT("*.*"));
-
-		UE_LOG(LogIoStoreOnDemand, Display, TEXT("Serializing %d additional file(s)"), AdditionalFiles.Num());
-		for (const FString& Filename : AdditionalFiles)
-		{
-			const bool bIncludeDot = false;
-			const FStringView Ext = FPathViews::GetExtension(Filename, bIncludeDot);
-
-			if (const FStringView* Result =
-				Algo::FindByPredicate(AllowedExt, [&Ext](const FStringView& E) { return Ext == E; });
-				Result == nullptr)
-			{
-				continue;
-			}
-
-			const FString FullPath = ContainerFolder / Filename;
-			TArray<uint8> FileData;
-			if (!FFileHelper::LoadFileToArray(FileData, *FullPath))
-			{
-				UE_LOG(LogIoStoreOnDemand, Error, TEXT("Failed reading file '%s'"), *FullPath);
-				return -1;
-			}
-
-			FMemoryView Chunk = MakeMemoryView(FileData.GetData(), FileData.Num());
-			const FIoHash ChunkHash = FIoHash::HashBuffer(Chunk);
-
-			FIoBuffer Buffer(FIoBuffer::Clone, Chunk);
-
-			FIoStatus WriteStatus = Input.ChunkWriter->WriteChunk(ChunksRelativeFolder, Buffer, ChunkHash);
-
-			if (WriteStatus.IsOk() == false)
-			{
-				UE_LOG(LogIoStoreOnDemand, Error, TEXT("%s"), *WriteStatus.ToString());
-				return -1;
-			}
-
-			UE_LOG(LogIoStoreOnDemand, Display, TEXT("Adding additional file '%s'"), *Filename);
-			Output.InstallOnDemandToc.AdditionalFiles.Add(FOnDemandTocAdditionalFile
-			{
-				.Hash = ChunkHash,
-				.Filename = Filename,
-				.FileSize = IntCastChecked<uint64>(FileData.Num())
-			});
-		}
+		UE_LOG(LogIoStoreOnDemand, Warning, TEXT("Adding additional file(s) to the on-demand container TOC is no longer supported"));
 	}
 
 	const int64 Timestamp = FDateTime::UtcNow().ToUnixTimestamp();
diff --git a/Engine/Source/Developer/IoStore/OnDemandUtilities/Private/Upload.cpp b/Engine/Source/Developer/IoStore/OnDemandUtilities/Private/Upload.cpp
--- a/Engine/Source/Developer/IoStore/OnDemandUtilities/Private/Upload.cpp
+++ b/Engine/Source/Developer/IoStore/OnDemandUtilities/Private/Upload.cpp
@@ -11,6 +11,7 @@
 #include "IO/IoChunkEncoding.h"
 #include "IO/IoStore.h"
 #include "IO/IoStoreOnDemand.h"
+#include "IO/Serialization/OnDemandContainerToc.h"
 #include "Misc/ConfigCacheIni.h"
 #include "Misc/FileHelper.h"
 #include "Misc/KeyChainUtilities.h"
@@ -54,7 +55,8 @@ struct FUploadParams
 	bool bDeletePakFiles = true;
 	bool bPerContainerTocs = false;
 	bool bIgnoreContainerFlags = false;
-	bool bIgnoreContainerHeader = false; 
+	bool bIgnoreContainerHeader = false;
+	bool bLegacyToc = false;
 
 	/** If we should write out the .iochunktoc to disk as well as uploading it. */
 	bool bWriteTocToDisk = false;
@@ -96,17 +98,20 @@ FIoStatus FUploadParams::Validate() const
 	return FIoStatus::Ok;
 }
 
+////////////////////////////////////////////////////////////////////////////////
+struct FTocInfo
+{
+	FIoHash Hash = FIoHash::Zero;
+	FString Path;
+	uint64	Size = 0;
+};
+
 ////////////////////////////////////////////////////////////////////////////////
 struct FUploadResult
 {
-	/** Hash of the toc when written as a binary blob */
-	FIoHash TocHash;
-	/** Url of the service that the toc was uploaded too */
-	FString ServiceUrl;
-	/** Path of the toc on the service */
-	FString TocPath;
-	/** Size (in bytes) of the toc when written as a binary blob */
-	uint64 TocSize = 0;
+	FTocInfo	FileToc;
+	FTocInfo	CloudToc;
+	FString		ServiceUrl;
 };
 
 ////////////////////////////////////////////////////////////////////////////////
@@ -134,6 +139,7 @@ static FUploadParams BuildUploadParams(const FContext& Context)
 	Ret.bWriteTocToDisk			= Context.Get<bool>(TEXT("-WriteTocToDisk"),				Ret.bWriteTocToDisk);
 	Ret.bPerContainerTocs		= Context.Get<bool>(TEXT("-PerContainerTocs"),				Ret.bPerContainerTocs);
 	Ret.bIgnoreContainerHeader	= Context.Get<bool>(TEXT("-IgnoreContainerHeader"),			Ret.bIgnoreContainerHeader);
+	Ret.bLegacyToc				= Context.Get<bool>(TEXT("-LegacyToc"),						Ret.bLegacyToc);
 	Ret.MaxConcurrentUploads	= Context.Get<int32>(TEXT("-MaxConcurrentUploads"),			Ret.MaxConcurrentUploads);
 	Ret.MaxTocListCount			= Context.Get<int32>(TEXT("-MaxTocListCount"),				Ret.MaxTocListCount);
 	Ret.MaxTocDownloadCount		= Context.Get<int32>(TEXT("-MaxTocDownloadCount"),			Ret.MaxTocDownloadCount);
@@ -203,7 +209,7 @@ static bool SaveArrayToFile(TArrayView64<uint8> Data, const TCHAR* Filename, uin
 }
 
 ////////////////////////////////////////////////////////////////////////////////
-static FIoStatus WriteContainerFiles(FOnDemandToc& OnDemandToc, const TMap<FIoHash, FString>& UTocPaths)
+static FIoStatus WriteContainerFiles(FOnDemandToc& OnDemandToc, const TMap<FIoHash, FString>& UTocPaths, bool bLegacyToc)
 {
 	UE_LOG(LogIas, Display, TEXT("Attempting to write out %d '%s' files"), OnDemandToc.Containers.Num(), *FOnDemandToc::FileExt);
 
@@ -232,21 +238,26 @@ static FIoStatus WriteContainerFiles(FOnDemandToc& OnDemandToc, const TMap<FIoHa
 			return FIoStatus(EIoErrorCode::FileOpenFailed, FString::Printf(TEXT("Failed to open '%s' for write"), *OutputPath));
 		}
 
-		// TODO: We should consider adding a hash of the FOnDemandToc that can be computed at runtime on the loaded structure
-		// (to avoid running over the file twice) to verify that nothing was corrupted.
-
-		(*Ar) << ContainerToc;
-
-		FOnDemandTocSentinel Sentinel;
-		(*Ar) << Sentinel;
+		if (bLegacyToc)
+		{
+			(*Ar) << ContainerToc;
+			FOnDemandTocSentinel Sentinel;
+			(*Ar) << Sentinel;
+		}
+		else
+		{
+			UE::IoStore::Serialization::FOnDemandTocWriter Writer;
+			if (TIoStatusOr<uint64> Status = Writer.Write(ContainerToc, *Ar); !Status.IsOk())
+			{
+				return Status.Status();
+			}
+		}
 
 		if (!Ar->Close())
 		{
 			return FIoStatus(EIoErrorCode::WriteError, FString::Printf(TEXT("Failed to write to '%s'"), *OutputPath));
 		}
 
-		// TODO: Should write an end of file sentinel here for a simple/easy file corruption check
-
 		UE_LOG(LogIas, Display, TEXT("Wrote ondemand container file '%s' to disk"), *OutputPath);
 
 		// Move the container entry back to OnDemandToc in case we want to use the data structure
@@ -342,7 +353,7 @@ static void WriteConfigFile(
 		Sb << TEXT("ServiceUrl=\"") << UploadResult.ServiceUrl << TEXT("\"\r\n");
 	}
 
-	Sb << TEXT("TocPath=\"") << UploadResult.TocPath << TEXT("\"\r\n");
+	Sb << TEXT("TocPath=\"") << UploadResult.FileToc.Path << TEXT("\"\r\n");
 
 	// Temporary solution to get replays working with encrypted on demand content
 	{
@@ -784,11 +795,13 @@ static TIoStatusOr<FUploadResult> UploadContainerFiles(
 
 	FUploadResult UploadResult;
 	{
+		using FOnDemandFileExt = UE::IoStore::Serialization::FOnDemandFileExt;
+
 		FLargeMemoryWriter Ar;
 		Ar << OnDemandToc;
 
-		UploadResult.TocHash = FIoHash::HashBuffer(Ar.GetView());
-		const FString TocFilename = LexToString(UploadResult.TocHash) + TEXT(".iochunktoc");
+		UploadResult.CloudToc.Hash = FIoHash::HashBuffer(Ar.GetView());
+		const FString TocFilename = LexToString(UploadResult.CloudToc.Hash) + FOnDemandFileExt::CloudToc;
 
 		TStringBuilder<256> Key;
 		if (!UploadParams.BucketPrefix.IsEmpty())
@@ -798,8 +811,8 @@ static TIoStatusOr<FUploadResult> UploadContainerFiles(
 		Key << TocFilename;
 
 		UploadResult.ServiceUrl = ServiceUrl;
-		UploadResult.TocPath = OnDemandToc.Header.ChunksDirectory / TocFilename;
-		UploadResult.TocSize = Ar.TotalSize();
+		UploadResult.CloudToc.Path = OnDemandToc.Header.ChunksDirectory / TocFilename;
+		UploadResult.CloudToc.Size = Ar.TotalSize();
 
 		const FS3PutObjectResponse Response = Client.TryPutObject(FS3PutObjectRequest{UploadParams.Bucket, Key.ToString(), Ar.GetView()});
 		if (Response.IsOk())
@@ -812,23 +825,40 @@ static TIoStatusOr<FUploadResult> UploadContainerFiles(
 			return FIoStatus(EIoErrorCode::WriteError, TEXT("Failed to upload TOC"));
 		}
 
+		FLargeMemoryWriter FileAr;
+		if (UploadParams.bLegacyToc)
+		{
+			FileAr.Serialize(Ar.GetData(), Ar.TotalSize());
+		}
+		else
+		{
+			UE::IoStore::Serialization::FOnDemandTocWriter Writer;
+			if (TIoStatusOr<uint64> Status = Writer.Write(OnDemandToc, FileAr); !Status.IsOk())
+			{
+				return Status.Status();
+			}
+		}
+
+		UploadResult.FileToc.Hash = FIoHash::HashBuffer(FileAr.GetView());
+		UploadResult.FileToc.Path = OnDemandToc.Header.ChunksDirectory / LexToString(UploadResult.FileToc.Hash) + FOnDemandFileExt::CloudToc;
+		UploadResult.FileToc.Size = FileAr.TotalSize();
+
 		// Write a single global.uondemandtoc next to global.utoc or in the Cloud directory (legacy)
 		{
 			TStringBuilder<512> OnDemandTocFilePath;
 			if (UploadParams.GlobalContainerPath.IsEmpty() == false)
 			{
-				FPathViews::Append(OnDemandTocFilePath, UploadParams.GlobalContainerPath, TEXT("global.uondemandtoc"));
+				FPathViews::Append(OnDemandTocFilePath, UploadParams.GlobalContainerPath, FString(TEXT("global")) + FOnDemandFileExt::Toc);
 			}
 			else if (UploadParams.bWriteTocToDisk)
 			{
-				FPathViews::Append(OnDemandTocFilePath, UploadParams.TocOutputDir, UploadResult.TocHash);
-				OnDemandTocFilePath << TEXT(".iochunktoc");
+				FPathViews::Append(OnDemandTocFilePath, UploadParams.TocOutputDir, LexToString(UploadResult.FileToc.Hash) + FOnDemandFileExt::CloudToc);
 			}
 
 			if (OnDemandTocFilePath.Len() > 0)
 			{
 				FPathViews::NormalizeFilename(OnDemandTocFilePath);
-				if (SaveArrayToFile(TArrayView64<uint8>(Ar.GetData(), Ar.TotalSize()), OnDemandTocFilePath.ToString()))
+				if (SaveArrayToFile(TArrayView64<uint8>(FileAr.GetData(), FileAr.TotalSize()), OnDemandTocFilePath.ToString()))
 				{
 					UE_LOG(LogIas, Display, TEXT("Saved on-demand TOC file '%s'"), OnDemandTocFilePath.ToString());
 				}
@@ -843,7 +873,7 @@ static TIoStatusOr<FUploadResult> UploadContainerFiles(
 	if (bWritePerContainerToc)
 	{
 		// Write out separate .uondemandtoc files, one per.utoc containing ondemand data.
-		FIoStatus Result = WriteContainerFiles(OnDemandToc, UTocPaths);
+		FIoStatus Result = WriteContainerFiles(OnDemandToc, UTocPaths, UploadParams.bLegacyToc);
 		if (!Result.IsOk())
 		{
 			return Result;
@@ -911,9 +941,9 @@ static TIoStatusOr<FUploadResult> UploadContainerFiles(
 		UE_LOG(LogIas, Display, TEXT("%-15s: %s"), TEXT("TargetPlatform"), *UploadParams.TargetPlatform);
 		UE_LOG(LogIas, Display, TEXT("%-15s: %s"), TEXT("BuildVersion"), *UploadParams.BuildVersion);
 		UE_LOG(LogIas, Display, TEXT("%-15s: %s"), TEXT("HostGroupName"), *UploadParams.HostGroupName);
-		UE_LOG(LogIas, Display, TEXT("%-15s: %s"), TEXT("TOC path"), *UploadResult.TocPath);
+		UE_LOG(LogIas, Display, TEXT("%-15s: %s"), TEXT("TOC path"), *UploadResult.CloudToc.Path);
 		UE_LOG(LogIas, Display, TEXT("%-15s: %s"), TEXT("TOC flags"), *LexToString(UploadParams.TocFlags));
-		UE_LOG(LogIas, Display, TEXT("%-15s: %.2lf KiB"), TEXT("TOC size"), double(UploadResult.TocSize) / 1024.0);
+		UE_LOG(LogIas, Display, TEXT("%-15s: %.2lf KiB"), TEXT("TOC size"), double(UploadResult.CloudToc.Size) / 1024.0);
 		UE_LOG(LogIas, Display, TEXT("%-15s: %.2lf second(s)"), TEXT("Duration"), Duration);
 		UE_LOG(LogIas, Display, TEXT(""));
 		
@@ -1034,6 +1064,7 @@ static FCommand UploadCommand(
 		TArgument<bool>(TEXT("-StreamOnDemand"),				TEXT("Set the content to be streamed on-demand")),
 		TArgument<bool>(TEXT("-InstallOnDemand"),				TEXT("Set the content to be installed on-demand")),
 		TArgument<bool>(TEXT("-IgnoreContainerHeader"),			TEXT("Do not serialize the container header as part of the TOC file.")),
+		TArgument<bool>(TEXT("-LegacyToc"),						TEXT("Whether to serialize the TOC using the legacy format.")),
 		TArgument<int32>(TEXT("-MaxConcurrentUploads"),			TEXT("Number of simultaneous uploads")),
 		TArgument<int32>(TEXT("-MaxTocListCount"),				TEXT("Maximum number of TOC file(s) to list from the bucket")),
 		TArgument<int32>(TEXT("-MaxTocDownloadCount"),			TEXT("Maximum number of TOC file(s) to download")),
diff --git a/Engine/Source/Runtime/Experimental/IoStore/OnDemand/Private/OnDemandConfig.cpp b/Engine/Source/Runtime/Experimental/IoStore/OnDemand/Private/OnDemandConfig.cpp
--- a/Engine/Source/Runtime/Experimental/IoStore/OnDemand/Private/OnDemandConfig.cpp
+++ b/Engine/Source/Runtime/Experimental/IoStore/OnDemand/Private/OnDemandConfig.cpp
@@ -251,30 +251,6 @@ TIoStatusOr<FOnDemandIoStoreConfig> TryParseConfig(const TCHAR* CommandLine)
 #endif
 
 	FOnDemandIoStoreConfig OutConfig;
-#if !UE_BUILD_SHIPPING
-	FString UrlParam;
-	if (FParse::Value(CommandLine, TEXT("Ias.TocUrl="), UrlParam))
-	{
-		FStringView UrlView(UrlParam);
-		if (UrlView.StartsWith(TEXTVIEW("http://")) && UrlView.EndsWith(TEXTVIEW(".iochunktoc")))
-		{
-			int32 Delim = INDEX_NONE;
-			if (UrlView.RightChop(7).FindChar(TEXT('/'), Delim))
-			{
-				OutConfig.ServiceUrls.Add(FString(UrlView.Left(7 +  Delim)));
-				OutConfig.TocPath = UrlView.RightChop(OutConfig.ServiceUrls[0].Len() + 1);
-			}
-		}
-
-		if (OutConfig.IsValid() == false)
-		{
-			FIoStatus Status = FIoStatusBuilder(EIoErrorCode::InvalidParameter) << TEXT("Failed to parse '-Ias.TocUrl='");
-			return Status;
-		}
-
-		return OutConfig;
-	}
-#endif
 
 	// Check for global TOC
 	bool bGlobalToc = false;
diff --git a/Engine/Source/Runtime/Experimental/IoStore/OnDemand/Private/OnDemandInstallCache.cpp b/Engine/Source/Runtime/Experimental/IoStore/OnDemand/Private/OnDemandInstallCache.cpp
--- a/Engine/Source/Runtime/Experimental/IoStore/OnDemand/Private/OnDemandInstallCache.cpp
+++ b/Engine/Source/Runtime/Experimental/IoStore/OnDemand/Private/OnDemandInstallCache.cpp
@@ -2975,7 +2975,7 @@ FResult FOnDemandInstallCache::Verify()
 		if (FileHandle->Seek(int64(ChunkLocation.Location.BlockOffset)) == false)
 		{
 			UE_LOG(LogIoStoreOnDemand, Error, TEXT("Chunk %d/%d SEEK FAILED, Container='%s', ChunkId='%s', ChunkSize=%" INT64_FMT ", Hash='%s', Block=%u, BlockOffset=%u"),
-				ChunkIndex + 1, ChunkCount, *Container->Name, *LexToString(ChunkId), ChunkSize, *LexToString(ChunkEntry.Hash),
+				ChunkIndex + 1, ChunkCount, *Container->Name(), *LexToString(ChunkId), ChunkSize, *LexToString(ChunkEntry.Hash),
 				ChunkLocation.Location.BlockId.Id, ChunkLocation.Location.BlockOffset);
 
 			ReadErrorCount++;
@@ -2986,7 +2986,7 @@ FResult FOnDemandInstallCache::Verify()
 		if (FileHandle->Read(reinterpret_cast<uint8*>(Chunk.GetData()), ChunkSize) == false)
 		{
 			UE_LOG(LogIoStoreOnDemand, Error, TEXT("Chunk %d/%d READ FAILED, Container='%s', ChunkId='%s', ChunkSize=%" INT64_FMT ", Hash='%s', Block=%u, BlockOffset=%u"),
-				ChunkIndex + 1, ChunkCount, *Container->Name, *LexToString(ChunkId), ChunkSize, *LexToString(ChunkEntry.Hash),
+				ChunkIndex + 1, ChunkCount, *Container->Name(), *LexToString(ChunkId), ChunkSize, *LexToString(ChunkEntry.Hash),
 				ChunkLocation.Location.BlockId.Id, ChunkLocation.Location.BlockOffset);
 
 			ReadErrorCount++;
@@ -2999,13 +2999,13 @@ FResult FOnDemandInstallCache::Verify()
 		if (ChunkHash == ChunkEntry.Hash)
 		{
 			UE_LOG(LogIoStoreOnDemand, VeryVerbose, TEXT("Chunk %d/%d OK, Container='%s', ChunkId='%s', ChunkSize=%" INT64_FMT ", Hash='%s', Block=%u, BlockOffset=%u"),
-				ChunkIndex + 1, ChunkCount, *Container->Name, *LexToString(ChunkId), ChunkSize, *LexToString(ChunkEntry.Hash),
+				ChunkIndex + 1, ChunkCount, *Container->Name(), *LexToString(ChunkId), ChunkSize, *LexToString(ChunkEntry.Hash),
 				ChunkLocation.Location.BlockId.Id, ChunkLocation.Location.BlockOffset);
 		}
 		else
 		{
 			UE_LOG(LogIoStoreOnDemand, Error, TEXT("Chunk %d/%d CORRUPT, Container='%s', ChunkId='%s', ChunkSize=%" INT64_FMT ", Hash='%s', ActualHash='%s', Block=%u, BlockOffset=%u"),
-				ChunkIndex + 1, ChunkCount, *Container->Name, *LexToString(ChunkId), ChunkSize, *LexToString(ChunkEntry.Hash), *LexToString(ChunkHash),
+				ChunkIndex + 1, ChunkCount, *Container->Name(), *LexToString(ChunkId), ChunkSize, *LexToString(ChunkEntry.Hash), *LexToString(ChunkHash),
 				ChunkLocation.Location.BlockId.Id, ChunkLocation.Location.BlockOffset);
 
 			CorruptChunkCount++;
@@ -3227,7 +3227,7 @@ void FOnDemandInstallCache::AddReferencesToBlocks(
 				if (!BlockInfo)
 				{
 					UE_LOG(LogIoStoreOnDemand, Error, TEXT("Failed to find CAS block info for referenced chunk, ChunkId='%s', Container='%s'"),
-						*LexToString(Container->ChunkIds[EntryIndex]), *Container->Name);
+						*LexToString(Container->ChunkIds[EntryIndex]), *Container->Name());
 				}
 				else
 				{
@@ -3244,7 +3244,7 @@ void FOnDemandInstallCache::AddReferencesToBlocks(
 				}
 
 				UE_CLOG(!bChunkIsPending, LogIoStoreOnDemand, Error, TEXT("Failed to find CAS location or pending chunk for chunk reference, ChunkId='%s', Container='%s'"),
-					*LexToString(Container->ChunkIds[EntryIndex]), *Container->Name);
+					*LexToString(Container->ChunkIds[EntryIndex]), *Container->Name());
 
 				ensure(bChunkIsPending);
 			}
diff --git a/Engine/Source/Runtime/Experimental/IoStore/OnDemand/Private/OnDemandIoStore.cpp b/Engine/Source/Runtime/Experimental/IoStore/OnDemand/Private/OnDemandIoStore.cpp
--- a/Engine/Source/Runtime/Experimental/IoStore/OnDemand/Private/OnDemandIoStore.cpp
+++ b/Engine/Source/Runtime/Experimental/IoStore/OnDemand/Private/OnDemandIoStore.cpp
@@ -7,7 +7,6 @@
 #include "Algo/RemoveIf.h"
 #include "Algo/Transform.h"
 #include "Async/UniqueLock.h"
-#include "Containers/RingBuffer.h"
 #include "Containers/StringConv.h"
 #include "Containers/Ticker.h"
 #include "DebugCommands.h"
@@ -61,6 +60,13 @@ static FAutoConsoleVariableRef CVar_PackageStreamingEnabled(
 	TEXT("Enables streaming packages without pre-installing into the install cache.")
 );
 
+bool GIaxMemoryMapTocEnabled = true;
+static FAutoConsoleVariableRef CVar_IaxMemoryMappedTocEnabled(
+	TEXT("iax.MemoryMappedTocEnabled"),
+	GIaxMemoryMapTocEnabled,
+	TEXT("Whether container TOC memory-mapping is enabled.")
+);
+
 ///////////////////////////////////////////////////////////////////////////////
 bool IsPackageStreamingEnabled()
 {
@@ -76,47 +82,6 @@ bool IsPackageStreamingEnabled()
 namespace Private
 {
 
-///////////////////////////////////////////////////////////////////////////////
-static FAnsiStringBuilderBase& GetChunkUrl(
-	const FAnsiStringView& Host,
-	const FOnDemandContainer& Container,
-	const FOnDemandChunkEntry& Entry,
-	FAnsiStringBuilderBase& OutUrl)
-{
-	OutUrl.Reset();
-	if (Host.IsEmpty() == false)
-	{
-		OutUrl << Host;
-	}
-
-	if (!Container.ChunksDirectory.IsEmpty())
-	{
-		OutUrl << "/" << Container.ChunksDirectory;
-	}
-
-	const FString HashString = LexToString(Entry.Hash);
-	OutUrl << "/" << HashString.Left(2) << "/" << HashString << ".iochunk";
-
-	return OutUrl;
-}
-
-static FAnsiStringBuilderBase& GetChunkUrl(
-	const FOnDemandContainer& Container,
-	const FOnDemandChunkEntry& Entry,
-	FAnsiStringBuilderBase& OutUrl)
-{
-	OutUrl.Reset();
-	if (!Container.ChunksDirectory.IsEmpty())
-	{
-		OutUrl << "/" << Container.ChunksDirectory;
-	}
-
-	const FString HashString = LexToString(Entry.Hash);
-	OutUrl << "/" << HashString.Left(2) << "/" << HashString << ".iochunk";
-
-	return OutUrl;
-}
-
 ///////////////////////////////////////////////////////////////////////////////
 static TIoStatusOr<FIoBuffer> DecodeChunk(const FOnDemandChunkInfo& ChunkInfo, FMemoryView EncodedChunk)
 {
@@ -184,10 +149,50 @@ static TIoStatusOr<FSharedContainerHeader> DeserializeContainerHeader(const FOnD
 	return OutHeader; 
 }
 
+///////////////////////////////////////////////////////////////////////////////
+static TIoStatusOr<FSharedContainerHeader> DeserializeContainerHeader(FSharedOnDemandContainer Container, FMemoryView EncodedHeaderChunk)
+{
+	// Only serialize package dependencies for containers used for installing content or when stream-all-the-things is enabled
+	const bool bShouldSerializeHeader = EnumHasAnyFlags(Container->Flags, EOnDemandContainerFlags::InstallOnDemand) || IsPackageStreamingEnabled();
+	if (bShouldSerializeHeader == false || EncodedHeaderChunk.IsEmpty())
+	{
+		FSharedContainerHeader Empty = MakeShared<FIoContainerHeader>();
+		return Empty;
+	}
+
+	const FIoChunkId ChunkId			= CreateContainerHeaderChunkId(Container->ContainerId());
+	const FOnDemandChunkInfo ChunkInfo	= FOnDemandChunkInfo::Find(Container, ChunkId); ChunkInfo.IsValid();
+
+	// Create an empty header if the container does not have package dependencies
+	if (ChunkInfo.IsValid() == false)
+	{
+		FSharedContainerHeader Empty = MakeShared<FIoContainerHeader>();
+		return Empty;
+	}
+
+	return Private::DeserializeContainerHeader(
+		ChunkInfo,
+		EncodedHeaderChunk,
+		EnumHasAnyFlags(Container->Flags, EOnDemandContainerFlags::WithSoftReferences));
+}
+
+///////////////////////////////////////////////////////////////////////////////
+static void LogMounted(FSharedOnDemandContainer Container)
+{
+	using namespace UE::IoStore::Serialization;
+
+	TStringBuilder<128> Sb;
+	Sb << Container->Flags;
+	UE_LOG(LogIoStoreOnDemand, Log, TEXT("Mounting container '%s', Entries=%d, Flags='%s', HostGroup='%s', TocSize=%.2lf KiB, Storage='%s'"),
+		*Container->Name(), Container->ChunkEntries.Num(), Sb.ToString(), *Container->HostGroupName.ToString(),
+		double(Container->Storage.GetView().GetSize()) / 1024.0,
+		Container->Storage.StorageType() == FOnDemandTocStorageType::MemoryMapped ? TEXT("MemoryMapped") : TEXT("Memory"));
+}
+
 } // namespace UE::IoStore::Private
 
 ///////////////////////////////////////////////////////////////////////////////
-const FOnDemandChunkEntry FOnDemandChunkEntry::Null = {};
+const FOnDemandChunkEntry FOnDemandChunkInfo::NullEntry = {};
 
 ///////////////////////////////////////////////////////////////////////////////
 static FString OnDemandContainerUniqueName(FStringView MountId, FStringView Name)
@@ -197,7 +202,7 @@ static FString OnDemandContainerUniqueName(FStringView MountId, FStringView Name
 
 FString FOnDemandContainer::UniqueName() const
 {
-	return OnDemandContainerUniqueName(MountId, Name);
+	return OnDemandContainerUniqueName(MountId, Name());
 }
 
 ///////////////////////////////////////////////////////////////////////////////
@@ -481,18 +486,7 @@ FIoStatus FOnDemandIoStore::InitializePostHotfix()
 	}
 	else if (!Config.ServiceUrls.IsEmpty() && !Config.TocPath.IsEmpty())
 	{
-		TIoStatusOr<FOnDemandHostGroup> HostGroup = FOnDemandHostGroup::Create(Config.ServiceUrls);
-		if (HostGroup.IsOk() == false)
-		{
-			return HostGroup.Status();
-		}
-		MountArgs.Emplace(FOnDemandMountArgs
-		{
-			.MountId = TEXT("global"),
-			.TocRelativeUrl = Config.TocPath,
-			.HostGroup = HostGroup.ConsumeValueOrDie(),
-			.Options = EOnDemandMountOptions::StreamOnDemand
-		});
+		UE_LOG(LogIoStoreOnDemand, Warning, TEXT("Mounting on-demand container TOC from URL is no longer supporeted"));
 	}
 
 #if !UE_BUILD_SHIPPING
@@ -686,16 +680,12 @@ FIoStatus FOnDemandIoStore::Unmount(FStringView MountId)
 			{
 				ensureMsgf(Container->ChunkEntryReferences.IsEmpty(), 
 					TEXT("Container is still referenced when unmounting, ContainerName='%s', MountId='%s'"), 
-					*Container->Name, *Container->MountId);
+					*Container->Name(), *Container->MountId);
 
 				UE_LOG(LogIoStoreOnDemand, Log, TEXT("Unmounting container, ContainerName='%s', MountId='%s'"),
-					*Container->Name, *Container->MountId);
+					*Container->Name(), *Container->MountId);
 
 				bRemoved = true;
-
-				const FString ContainerName = Container->UniqueName();
-				PendingContainerHeaders.Remove(ContainerName);
-
 				return true;
 			}
 
@@ -1266,11 +1256,20 @@ FIoStatus FOnDemandIoStore::ProcessMountRequest(FMountRequest& MountRequest)
 
 	if (Args.MountId.IsEmpty())
 	{
-		return FIoStatusBuilder(EIoErrorCode::InvalidParameter) << TEXT("Invalid mount ID");
+		Args.MountId = Args.FilePath;
+		if (Args.MountId.IsEmpty())
+		{
+			return FIoStatus(EIoErrorCode::InvalidParameter, TEXT("Invalid mount ID"));
+		}
+	}
+
+	if (Args.FilePath.IsEmpty())
+	{
+		return FIoStatus(EIoErrorCode::InvalidParameter, TEXT("Invalid TOC file path"));
 	}
 
 	EIoErrorCode MountResult = EIoErrorCode::Ok;
-	TArray<FString> ExistingContainers;
+	TArray<FIoContainerId> ExistingContainers;
 
 	// Find containers matching the mount ID
 	{
@@ -1282,7 +1281,7 @@ FIoStatus FOnDemandIoStore::ProcessMountRequest(FMountRequest& MountRequest)
 				continue;
 			}
 			
-			ExistingContainers.Add(Container->Name);
+			ExistingContainers.Add(Container->ContainerId());
 			
 			if (EnumHasAnyFlags(Container->Flags, EOnDemandContainerFlags::PendingEncryptionKey))
 			{
@@ -1296,73 +1295,13 @@ FIoStatus FOnDemandIoStore::ProcessMountRequest(FMountRequest& MountRequest)
 		}
 	}
 
-	// Containers haven't been created yet, do it now
 	TArray<FSharedOnDemandContainer> RequestedContainers;
-	const FStringView TocPath = FPathViews::GetExtension(Args.TocRelativeUrl).IsEmpty() ? Args.TocRelativeUrl : FPathViews::GetPath(Args.TocRelativeUrl);
-
-	if (Args.Toc)
+	if (FIoStatus Status = CreateContainersFromToc(Args, ExistingContainers, RequestedContainers); !Status.IsOk())
 	{
-		CreateContainersFromToc(Args.MountId, TocPath, *Args.Toc, ExistingContainers, RequestedContainers);
-	}
-	else if (Args.FilePath.IsEmpty() == false)
-	{
-		UE_LOG(LogIoStoreOnDemand, Log, TEXT("Loading TOC from file '%s'"), *Args.FilePath);
-
-		// TODO: Enable validation when the sentinal is included in all serialization paths
-		const bool bValidate = false;
-		TIoStatusOr<FOnDemandToc> TocStatus = FOnDemandToc::LoadFromFile(Args.FilePath, bValidate);
-		if (!TocStatus.IsOk())
-		{
-			return TocStatus.Status();
-		}
-
-		Args.Toc = MakeUnique<FOnDemandToc>(TocStatus.ConsumeValueOrDie());
-
-		CreateContainersFromToc(Args.MountId, TocPath, *Args.Toc, ExistingContainers, RequestedContainers);
-	}
-	else if (Args.HostGroup.IsEmpty() == false && Args.TocRelativeUrl.IsEmpty() == false)
-	{
-		TStringBuilder<512> Url;
-		Url.Append(Args.HostGroup.PrimaryHost());
-		if (Args.TocRelativeUrl[0] != TEXT('/'))
-		{
-			Url.AppendChar(TEXT('/'));
-		}
-		Url.Append(Args.TocRelativeUrl);
-
-		UE_LOG(LogIoStoreOnDemand, Log, TEXT("Loading TOC from URL '%s'"), Url.ToString());
-
-		auto AnsiUrl = StringCast<ANSICHAR>(Url.GetData(), Url.Len());
-		TIoStatusOr<FMultiEndpointHttpClientResponse> Response = FMultiEndpointHttpClient::Get(AnsiUrl, FMultiEndpointHttpClientConfig
-		{ 
-			.MaxRetryCount	= 2,
-			.Redirects		= EHttpRedirects::Follow
-		});
-
-		if (!Response.IsOk())
-		{
-			FIoStatus Status = FIoStatusBuilder(EIoErrorCode::ReadError) << TEXT("Failed to fetch TOC from URL");
-			return Status;
-		}
-
-		FMemoryReaderView Ar(Response.ValueOrDie().Body.GetView());
-		Args.Toc = MakeUnique<FOnDemandToc>();
-		Ar << *Args.Toc;
-
-		if (Ar.IsError() || Ar.IsCriticalError())
-		{
-			Args.Toc.Reset();
-			FIoStatus Status = FIoStatusBuilder(EIoErrorCode::ReadError) << TEXT("Failed to serialize TOC from HTTP response");
-			return Status;
-		}
-
-		CreateContainersFromToc(Args.MountId, TocPath, *Args.Toc, ExistingContainers, RequestedContainers);
+		return Status;
 	}
-	
-	TAnsiStringBuilder<512> ChunkUrlBuilder;
-	TMap<FString, FIoBuffer> RequestedContainerHeaders;
 
-	for (const FSharedOnDemandContainer& Container : RequestedContainers)
+	for (FSharedOnDemandContainer& Container : RequestedContainers)
 	{
 		if (EnumHasAnyFlags(Container->Flags, EOnDemandContainerFlags::StreamOnDemand | EOnDemandContainerFlags::InstallOnDemand) == false)
 		{
@@ -1381,7 +1320,7 @@ FIoStatus FOnDemandIoStore::ProcessMountRequest(FMountRequest& MountRequest)
 				EnumAddFlags(Container->Flags, EOnDemandContainerFlags::PendingHostGroup);
 				MountResult = EIoErrorCode::PendingHostGroup;
 				UE_LOG(LogIoStoreOnDemand, Log, TEXT("Deferring container '%s' until host group '%s' becomes available"),
-					*Container->Name, *Container->HostGroupName.ToString());
+					*Container->Name(), *Container->HostGroupName.ToString());
 			}
 			else
 			{
@@ -1392,82 +1331,27 @@ FIoStatus FOnDemandIoStore::ProcessMountRequest(FMountRequest& MountRequest)
 		if (EnumHasAnyFlags(Container->Flags, EOnDemandContainerFlags::Encrypted) &&
 			Container->EncryptionKey.IsValid() == false)
 		{
-			FGuid KeyGuid;
-			ensure(FGuid::Parse(Container->EncryptionKeyGuid, KeyGuid));
-			if (FEncryptionKeyManager::Get().TryGetKey(KeyGuid, Container->EncryptionKey) == false)
+			if (FEncryptionKeyManager::Get().TryGetKey(Container->EncryptionKeyGuid(), Container->EncryptionKey) == false)
 			{
 				EnumAddFlags(Container->Flags, EOnDemandContainerFlags::PendingEncryptionKey);
 				MountResult = EIoErrorCode::PendingEncryptionKey;
 				UE_LOG(LogIoStoreOnDemand, Log, TEXT("Deferring container '%s' until encryption key '%s' becomes available"),
-					*Container->Name, *Container->EncryptionKeyGuid);
+					*Container->Name(), *LexToString(Container->EncryptionKeyGuid()));
 			}
 		}
 
-		// Try fetch and deserialize the container header if the container is used for installing content or if package streaming is enabled
-		if (EnumHasAnyFlags(Container->Flags, EOnDemandContainerFlags::InstallOnDemand) || IsPackageStreamingEnabled())
+		// Serialize the container header with package dependencies if the encryption key is available
+		if (EnumHasAnyFlags(Container->Flags, EOnDemandContainerFlags::PendingEncryptionKey) == false)
 		{
-			const FIoChunkId ChunkId			= CreateContainerHeaderChunkId(Container->ContainerId);
-			const FOnDemandChunkInfo ChunkInfo	= FOnDemandChunkInfo::Find(Container, ChunkId);
-
-			if (ChunkInfo.IsValid())
+			FIoBuffer ContainerHeaderChunk = MoveTemp(Container->ContainerHeaderChunk);
+			TIoStatusOr<FSharedContainerHeader> Header = Private::DeserializeContainerHeader(Container, ContainerHeaderChunk.GetView());
+			if (Header.IsOk())
 			{
-				const FOnDemandTocContainerEntry* ContainerEntry = Args.Toc->Containers.FindByPredicate(
-						[&Container](const FOnDemandTocContainerEntry& E) { return E.ContainerId == Container->ContainerId; });
-				UE_CLOG(ContainerEntry == nullptr, LogIoStoreOnDemand, Error, TEXT("Failed to find TOC container entry for container '%s'"), *Container->Name);
-
-				if (ContainerEntry != nullptr)
-				{
-					FIoBuffer HeaderBuffer;
-					if (ContainerEntry->Header.IsEmpty() == false)
-					{
-						HeaderBuffer = FIoBuffer(FIoBuffer::Wrap, ContainerEntry->Header.GetData(), ContainerEntry->Header.Num());
-					}
-					else
-					{
-						if (EnumHasAnyFlags(Container->Flags, EOnDemandContainerFlags::PendingHostGroup))
-						{
-							continue;
-						}
-
-						if (Container->HostGroup.IsResolved() == false)
-						{
-							return FIoStatusBuilder(EIoErrorCode::InvalidParameter) << TEXT("No valid endpoint(s)");
-						}
-
-						UE_LOG(LogIoStoreOnDemand, Warning, TEXT("Fetching container header, ContainerName='%s', ChunkId='%s'"),
-								*Container->Name, *LexToString(ChunkId));
-
-						const FAnsiStringView PrimaryHost = Container->HostGroup.GetUnderlyingHostGroup().PrimaryHost();
-						const uint32 RetryCount = 2;
-						TIoStatusOr<FMultiEndpointHttpClientResponse> Response = FMultiEndpointHttpClient::Get(
-								Private::GetChunkUrl(PrimaryHost, *Container, ChunkInfo.ChunkEntry(), ChunkUrlBuilder).ToView(),
-								FMultiEndpointHttpClientConfig{.MaxRetryCount = RetryCount});
-
-						if (Response.IsOk() == false)
-						{
-							return Response.Status();
-						}
-						HeaderBuffer = Response.ConsumeValueOrDie().Body;
-					}
-
-					if (EnumHasAnyFlags(Container->Flags, EOnDemandContainerFlags::PendingEncryptionKey))
-					{
-						HeaderBuffer.EnsureOwned();
-						RequestedContainerHeaders.Add(Container->UniqueName(), HeaderBuffer);
-						continue;
-					}
-
-					TIoStatusOr<FSharedContainerHeader> Header = Private::DeserializeContainerHeader(
-							ChunkInfo,
-							HeaderBuffer.GetView(),
-							EnumHasAnyFlags(Container->Flags, EOnDemandContainerFlags::WithSoftReferences));
-
-					if (Header.IsOk() == false)
-					{
-						return Header.Status();
-					}
-					Container->Header = Header.ConsumeValueOrDie();
-				}
+				Container->Header = Header.ConsumeValueOrDie();
+			}
+			else
+			{
+				return Header.Status();
 			}
 		}
 
@@ -1476,18 +1360,13 @@ FIoStatus FOnDemandIoStore::ProcessMountRequest(FMountRequest& MountRequest)
 			continue;
 		}
 
-		TStringBuilder<128> Sb;
-		Sb << Container->Flags;
-		UE_LOG(LogIoStoreOnDemand, Log, TEXT("Mounting container '%s', Entries=%d, Flags='%s', HostGroup='%s'"),
-			*Container->Name, Container->ChunkEntries.Num(), Sb.ToString(), *Container->HostGroupName.ToString());
+		Private::LogMounted(Container);
 		EnumAddFlags(Container->Flags, EOnDemandContainerFlags::Mounted);
-		Container->EncryptionKeyGuid.Empty();
 	}
 
 	{
 		UE::TUniqueLock Lock(ContainerMutex);
 		Containers.Append(MoveTemp(RequestedContainers));
-		PendingContainerHeaders.Append(MoveTemp(RequestedContainerHeaders));
 	}
 
 	if (ensureMsgf(PackageStoreBackend, TEXT("PackageStoreBackend is null, is this a prefork server?")))
@@ -1495,7 +1374,7 @@ FIoStatus FOnDemandIoStore::ProcessMountRequest(FMountRequest& MountRequest)
 		PackageStoreBackend->NeedsUpdate();
 	}
 
-	return MountResult; 
+	return MountResult;
 }
 
 void FOnDemandIoStore::CompleteMountRequest(FMountRequest& Request, FOnDemandMountResult&& MountResult)
@@ -1536,51 +1415,29 @@ void FOnDemandIoStore::OnEncryptionKeyAdded(const FGuid& Id, const FAES::FAESKey
 			continue;
 		}
 
-		FGuid KeyGuid;
-		ensure(FGuid::Parse(Container->EncryptionKeyGuid, KeyGuid));
-
-		if (FEncryptionKeyManager::Get().TryGetKey(KeyGuid, Container->EncryptionKey) == false)
+		if (FEncryptionKeyManager::Get().TryGetKey(Container->EncryptionKeyGuid(), Container->EncryptionKey) == false)
 		{
 			continue;
 		}
 
-		UE_LOG(LogIoStoreOnDemand, Log, TEXT("Mounting container (found encryption key) '%s', Entries=%d, Flags='%s'"),
-				*Container->Name, Container->ChunkEntries.Num(), *LexToString(Container->Flags));
-
+		Private::LogMounted(Container);
 		EnumRemoveFlags(Container->Flags, EOnDemandContainerFlags::PendingEncryptionKey);
 		EnumAddFlags(Container->Flags, EOnDemandContainerFlags::Mounted);
-		
-		Container->EncryptionKeyGuid.Empty();
 
-		const FString ContainerName		= Container->UniqueName();
-		const uint32 ContainerNameHash	= GetTypeHash(ContainerName);
-
-		if (FIoBuffer* EncodedHeader = PendingContainerHeaders.FindByHash(ContainerNameHash, ContainerName))
 		{
-			const FIoChunkId ChunkId			= CreateContainerHeaderChunkId(Container->ContainerId);
-			const FOnDemandChunkInfo ChunkInfo	= FOnDemandChunkInfo::Find(Container, ChunkId);
-
-			if (ensure(ChunkInfo.IsValid()))
+			FIoBuffer ContainerHeaderChunk = MoveTemp(Container->ContainerHeaderChunk);
+			TIoStatusOr<FSharedContainerHeader> Header = Private::DeserializeContainerHeader(Container, ContainerHeaderChunk.GetView());
+			if (Header.IsOk())
 			{
-				TIoStatusOr<FSharedContainerHeader> Header = Private::DeserializeContainerHeader(
-					ChunkInfo,
-					EncodedHeader->GetView(),
-					EnumHasAnyFlags(Container->Flags, EOnDemandContainerFlags::WithSoftReferences));
-
-				if (ensure(Header.IsOk()))
-				{
-					Container->Header = Header.ConsumeValueOrDie();
-					bAddedContainerHeaders = true;
-				}
-				else
-				{
-					EnumRemoveFlags(Container->Flags, EOnDemandContainerFlags::Mounted);
-					UE_LOG(LogIoStoreOnDemand, Error, TEXT("Failed to deserialize header when mounting container '%s', Entries=%d, Flags='%s'"),
-						*Container->Name, Container->ChunkEntries.Num(), *LexToString(Container->Flags));
-				}
+				Container->Header = Header.ConsumeValueOrDie();
+				bAddedContainerHeaders = true;
+			}
+			else
+			{
+				EnumRemoveFlags(Container->Flags, EOnDemandContainerFlags::Mounted);
+				UE_LOG(LogIoStoreOnDemand, Error, TEXT("Failed to deserialize header when mounting container '%s', Entries=%d, Flags='%s'"),
+					*Container->Name(), Container->ChunkEntries.Num(), *LexToString(Container->Flags));
 			}
-
-			PendingContainerHeaders.RemoveByHash(ContainerNameHash, ContainerName);
 		}
 	}
 
@@ -1600,40 +1457,51 @@ void FOnDemandIoStore::OnHostGroupRegistered(const FName& HostGroup)
 			continue;
 		}
 
+		Private::LogMounted(Container);
 		EnumRemoveFlags(Container->Flags, EOnDemandContainerFlags::PendingHostGroup);
 		EnumAddFlags(Container->Flags, EOnDemandContainerFlags::Mounted);
-		UE_LOG(LogIoStoreOnDemand, Log, TEXT("Mounting container '%s', Entries=%d, Flags='%s', HostGroup='%s'"),
-			*Container->Name, Container->ChunkEntries.Num(), *LexToString(Container->Flags), *HostGroup.ToString());
 	}
 }
 
-void FOnDemandIoStore::CreateContainersFromToc(
-	FStringView MountId,
-	FStringView TocPath,
-	FOnDemandToc& Toc,
-	const TArray<FString>& ExistingContainers,
+FIoStatus FOnDemandIoStore::CreateContainersFromToc(
+	FOnDemandMountArgs& MountArgs,
+	TConstArrayView<FIoContainerId> Existing,
 	TArray<FSharedOnDemandContainer>& Out)
 {
-	static FName AssetClassName("OnDemandIoStore");
+	using namespace UE::IoStore::Serialization::V1;
+
+	check(!MountArgs.MountId.IsEmpty());
+	check(!MountArgs.FilePath.IsEmpty());
+
+	const static FName AssetClassName("OnDemandIoStore");
 	LLM_SCOPE_BYNAME(TEXT("FileSystem/OnDemandIoStore"));
-	LLM_TAGSET_SCOPE(FName(MountId), ELLMTagSet::Assets);
+	LLM_TAGSET_SCOPE(FName(MountArgs.MountId), ELLMTagSet::Assets);
 	LLM_TAGSET_SCOPE(AssetClassName, ELLMTagSet::AssetClasses);
-	UE_TRACE_METADATA_SCOPE_ASSET_FNAME(FName(MountId), AssetClassName, FName(TocPath));
-	const FOnDemandTocHeader& Header = Toc.Header;
-	const FName CompressionFormat(Header.CompressionFormat);
+	UE_TRACE_METADATA_SCOPE_ASSET_FNAME(FName(MountArgs.MountId), AssetClassName, FName(MountArgs.FilePath));
+
+	UE_LOG(LogIoStoreOnDemand, Log, TEXT("Loading TOC from file '%s'"), *MountArgs.FilePath);
+	TIoStatusOr<FOnDemandTocReader> MaybeReader = FOnDemandTocReader::Read(*MountArgs.FilePath);
+	if (MaybeReader.IsOk() == false)
+	{
+		return MaybeReader.Status();
+	}
+
+	FOnDemandTocReader Reader = MaybeReader.ConsumeValueOrDie();
+
+	const FStringView TocPath = FPathViews::GetExtension(MountArgs.TocRelativeUrl).IsEmpty()
+		? MountArgs.TocRelativeUrl
+		: FPathViews::GetPath(MountArgs.TocRelativeUrl);
 
 	TStringBuilder<128> Sb;
 	FStringView ChunksDirectory;
 	{
 		if (TocPath.IsEmpty() == false)
 		{
-			//Algo::Transform(TocPath, AppendChars(Sb), FChar::ToLower);
 			FPathViews::Append(Sb, TocPath);
 		}
 		else
 		{
-			// Algo::Transform(Toc.Header.ChunksDirectory, AppendChars(Sb), FChar::ToLower);
-			FPathViews::Append(Sb, Toc.Header.ChunksDirectory);
+			FPathViews::Append(Sb, FString(Reader.ChunksDirectory()));
 		}
 		FPathViews::Append(Sb, TEXT("chunks"));
 
@@ -1648,93 +1516,57 @@ void FOnDemandIoStore::CreateContainersFromToc(
 		}
 	}
 
-	const EOnDemandTocFlags TocFlags = static_cast<EOnDemandTocFlags>(Toc.Header.Flags);
-	for (FOnDemandTocContainerEntry& ContainerEntry : Toc.Containers)
+	for (const FOnDemandContainerEntry& ContainerEntry : Reader.Containers())
 	{
-		if (ExistingContainers.Contains(ContainerEntry.ContainerName))
+		if (Existing.Contains(ContainerEntry.ContainerId))
 		{
 			continue;
 		}
 
-		FSharedOnDemandContainer Container = MakeShared<FOnDemandContainer>();
-		Container->Name					= MoveTemp(ContainerEntry.ContainerName);
-		Container->MountId				= MountId;
-		Container->ChunksDirectory		= StringCast<ANSICHAR>(ChunksDirectory.GetData(), ChunksDirectory.Len());
-		Container->EncryptionKeyGuid	= MoveTemp(ContainerEntry.EncryptionKeyGuid);
-		Container->BlockSize			= Header.BlockSize;
-		Container->BlockSizes			= MoveTemp(ContainerEntry.BlockSizes);
-		Container->BlockHashes			= MoveTemp(ContainerEntry.BlockHashes);
-		Container->ContainerId			= ContainerEntry.ContainerId;
-		Container->HostGroupName		= Toc.Header.HostGroupName.IsEmpty() ? FOnDemandHostGroup::DefaultName : FName(Toc.Header.HostGroupName);
-		Container->RelativeUrl			= FIoRelativeUrl::From(Container->ChunksDirectory);
+		FSharedOnDemandContainer				Container = MakeShared<FOnDemandContainer>();
+		EOnDemandTocReaderOptions				Options = GIaxMemoryMapTocEnabled ? EOnDemandTocReaderOptions::MemoryMap : EOnDemandTocReaderOptions::None;
+		TIoStatusOr<FOnDemandContainerTocView>	MaybeView = Reader.ReadContainer(ContainerEntry, Container->Storage, Options);
 
-		Container->CompressionFormats.Reserve(1);
-		Container->CompressionFormats.Add(CompressionFormat);
+		if (MaybeView.IsOk() == false)
+		{
+			return MaybeView.Status();
+		}
 
-		const EIoContainerFlags ContainerFlags = static_cast<EIoContainerFlags>(ContainerEntry.ContainerFlags);
-		if (EnumHasAnyFlags(ContainerFlags, EIoContainerFlags::Encrypted))
+		Container->MountId				= MountArgs.MountId;
+		FOnDemandContainerTocView View	= MaybeView.ConsumeValueOrDie();
+		Container->HeaderView			= View.Header;
+		Container->ContainerHeaderChunk	= MoveTemp(View.ContainerHeaderChunk);
+		Container->ChunkIds				= View.ChunkIds;
+		Container->ChunkEntries			= View.ChunkEntries;
+		Container->BlockSizes			= View.BlockSizes;
+		Container->BlockHashes			= View.BlockHashes;
+		Container->TagSets				= View.TagSets;
+		Container->TagSetIndices		= View.TagSetIndices;
+		Container->RelativeUrl			= UE::FIoRelativeUrl::From(FAnsiString(ChunksDirectory));
+		Container->CompressionFormat	= FName(FString(Reader.CompressionFormat()));
+		Container->HostGroupName		= FName(Reader.HostGroupName());
+
+		if (EnumHasAnyFlags(View.Header.FileContainerFlags(), EIoContainerFlags::Encrypted))
 		{
 			EnumAddFlags(Container->Flags, EOnDemandContainerFlags::Encrypted);
 		}
-		if (EnumHasAnyFlags(TocFlags, EOnDemandTocFlags::InstallOnDemand))
+		if (EnumHasAnyFlags(View.Header.ContainerFlags(), EOnDemandContainerEntryFlags::InstallOnDemand))
 		{
 			EnumAddFlags(Container->Flags, EOnDemandContainerFlags::InstallOnDemand);
 		}
-		else if (EnumHasAnyFlags(TocFlags, EOnDemandTocFlags::StreamOnDemand))
+		else if (EnumHasAnyFlags(View.Header.ContainerFlags(), EOnDemandContainerEntryFlags::StreamOnDemand))
 		{
 			EnumAddFlags(Container->Flags, EOnDemandContainerFlags::StreamOnDemand);
 		}
 
-		//TODO: Sort before uploading
-		ContainerEntry.Entries.Sort([](const FOnDemandTocEntry& Lhs, const FOnDemandTocEntry& Rhs)
-		{
-			return Lhs.ChunkId < Rhs.ChunkId;
-		});
-
-		const int32	EntryCount		= ContainerEntry.Entries.Num();
-		const uint64 TotalEntrySize = (sizeof(FIoChunkId) + sizeof(FOnDemandTocEntry)) * EntryCount;
-		Container->ChunkEntryData	= MakeUnique<uint8[]>(TotalEntrySize);
-		Container->ChunkIds			= MakeArrayView<FIoChunkId>(reinterpret_cast<FIoChunkId*>(Container->ChunkEntryData.Get()), EntryCount);
-		Container->ChunkEntries		= MakeArrayView<FOnDemandChunkEntry>(
-			reinterpret_cast<FOnDemandChunkEntry*>(Container->ChunkIds.GetData() + EntryCount), EntryCount);
-
-		for (int32 EntryIndex = 0; EntryIndex < EntryCount; EntryIndex++)
-		{
-			const FOnDemandTocEntry& TocEntry = ContainerEntry.Entries[EntryIndex];
-
-			Container->ChunkIds[EntryIndex]		= TocEntry.ChunkId;
-			Container->ChunkEntries[EntryIndex]	= FOnDemandChunkEntry
-			{
-				.Hash					= TocEntry.Hash,
-				.RawSize				= uint32(TocEntry.RawSize),
-				.EncodedSize			= uint32(TocEntry.EncodedSize),
-				.BlockOffset			= TocEntry.BlockOffset,
-				.BlockCount				= TocEntry.BlockCount,
-				.CompressionFormatIndex	= 0
-			};
-		}
-
-		const uint32 ContainerIndex = Out.Num();
-		for (FOnDemandTocTagSet& TagSet : Toc.TagSets)
-		{
-			for (FOnDemandTocTagSetPackageList& ContainerPackageIndices : TagSet.Packages)
-			{
-				if (ContainerPackageIndices.ContainerIndex == ContainerIndex)
-				{
-					FOnDemandTagSet& NewTagSet = Container->TagSets.AddDefaulted_GetRef();
-					NewTagSet.Tag = TagSet.Tag;
-					NewTagSet.PackageIndicies = MoveTemp(ContainerPackageIndices.PackageIndicies);
-					break;
-				}
-			}
-		}
-
 		// Do not count output container memory as being allocated for utocs
 		LLM_TAGSET_SCOPE_CLEAR(ELLMTagSet::Assets);
 		LLM_TAGSET_SCOPE_CLEAR(ELLMTagSet::AssetClasses);
 		UE_TRACE_METADATA_CLEAR_SCOPE();
 		Out.Add(MoveTemp(Container));
 	}
+
+	return FIoStatus::Ok;
 }
 
 FIoStatus FOnDemandIoStore::SetupHostGroup(const FSharedOnDemandContainer& Container, const FOnDemandMountArgs& MountArgs)
@@ -1831,7 +1663,7 @@ FIoStatus FOnDemandIoStore::GetContainersForInstall(
 				check(EnumHasAnyFlags(Container->Flags, EOnDemandContainerFlags::Mounted) == false);
 				FIoStatus Status = FIoStatusBuilder(EIoErrorCode::PendingEncryptionKey)
 					<< TEXT("Trying to install content from encrypted container '")
-					<< Container->Name
+					<< Container->Name()
 					<< TEXT("'");
 				return Status;
 			}
@@ -1903,13 +1735,16 @@ FIoStatus FOnDemandIoStore::GetContainersAndPackagesForInstall(
 			const int32 NumBeforeTag = OutPackageIdsToInstall.Num();
 			for (const FOnDemandTagSet& TagSet : Container->TagSets)
 			{
-				if (TagSet.Tag == Tag)
+				//TODO: Use UTF8 strings in the public API
+				const FString TagSetTag = FString(Container->HeaderView.GetString(TagSet.Tag));
+				if (TagSetTag == Tag)
 				{
-					UE_LOG(LogIoStoreOnDemand, Log, TEXT("Installing %d package(s) with tag '%s' from container '%s'"),
-						TagSet.PackageIndicies.Num(), *TagSet.Tag, *Container->UniqueName());
+					UE_LOG(LogIoStoreOnDemand, Log, TEXT("Installing %u package(s) with tag '%s' from container '%s'"),
+						TagSet.Count, *Tag, *Container->UniqueName());
 
-					OutPackageIdsToInstall.Reserve(TagSet.PackageIndicies.Num());
-					for (const uint32 PackageIndex : TagSet.PackageIndicies)
+					OutPackageIdsToInstall.Reserve(TagSet.Count);
+					TConstArrayView<uint32> PackageIndicies = Container->TagSetIndices.Mid(TagSet.Offset, TagSet.Count);
+					for (const uint32 PackageIndex : PackageIndicies)
 					{
 						const FPackageId PackageId = Container->Header->PackageIds[IntCastChecked<int32>(PackageIndex)];
 						OutPackageIdsToInstall.Add(PackageId);
diff --git a/Engine/Source/Runtime/Experimental/IoStore/OnDemand/Private/OnDemandIoStore.h b/Engine/Source/Runtime/Experimental/IoStore/OnDemand/Private/OnDemandIoStore.h
--- a/Engine/Source/Runtime/Experimental/IoStore/OnDemand/Private/OnDemandIoStore.h
+++ b/Engine/Source/Runtime/Experimental/IoStore/OnDemand/Private/OnDemandIoStore.h
@@ -14,6 +14,7 @@
 #include "IO/IoChunkId.h"
 #include "IO/IoHash.h"
 #include "IO/IoStoreOnDemandInternals.h"
+#include "IO/Serialization/OnDemandContainerToc.h"
 #include "IasHostGroup.h"
 #include "Misc/AES.h"
 #include "Misc/EnumClassFlags.h"
@@ -52,27 +53,10 @@ enum class EOnDemandContainerFlags : uint8
 ENUM_CLASS_FLAGS(EOnDemandContainerFlags);
 
 ///////////////////////////////////////////////////////////////////////////////
-struct FOnDemandChunkEntry
-{
-	static const FOnDemandChunkEntry Null;
-
-	FIoHash	Hash;
-	uint32	RawSize = 0;
-	uint32	EncodedSize = 0;
-	uint32	BlockOffset = ~uint32(0);
-	uint32	BlockCount = 0;
-	uint8	CompressionFormatIndex = 0;
-
-	uint32	GetDiskSize() const { return Align(EncodedSize, FAES::AESBlockSize); }
-};
-static_assert(sizeof(FOnDemandChunkEntry) == 40);
-
-///////////////////////////////////////////////////////////////////////////////
-struct FOnDemandTagSet
-{
-	FString			Tag;
-	TArray<uint32>	PackageIndicies;
-};
+using FOnDemandChunkEntry				= UE::IoStore::Serialization::V1::FOnDemandChunkEntry;
+using FOnDemandTagSet					= UE::IoStore::Serialization::V1::FOnDemandTagSetEntry;
+using FOnDemandTocStorage				= UE::IoStore::Serialization::FOnDemandTocStorage;
+using FOnDemandContainerTocHeaderView	= UE::IoStore::Serialization::V1::FOnDemandContainerTocHeaderView;
 
 ///////////////////////////////////////////////////////////////////////////////
 struct FOnDemandChunkEntryReferences
@@ -94,7 +78,10 @@ struct FOnDemandContainerChunkEntryReferences
 struct FOnDemandContainer
 {
 	FString									UniqueName() const;
-
+	inline FString							Name() const;
+	FIoContainerId							ContainerId() const { return HeaderView.ContainerId(); }
+	const FGuid&							EncryptionKeyGuid() const { return HeaderView.EncryptionKeyGuid(); }
+	uint32									BlockSize() const { return HeaderView.BlockSize(); }
 	inline int32							GetNumChunks() const;
 	inline int32							FindChunkEntryIndex(const FIoChunkId& ChunkId) const;
 	inline const FOnDemandChunkEntry*		FindChunkEntry(const FIoChunkId& ChunkId, int32* OutIndex = nullptr) const;
@@ -107,27 +94,31 @@ struct FOnDemandContainer
 	bool									HasAllFlags(EOnDemandContainerFlags Contains) const { return EnumHasAllFlags(Flags, Contains); }
 
 	FAES::FAESKey							EncryptionKey;
+	FOnDemandTocStorage						Storage;
+	FOnDemandContainerTocHeaderView			HeaderView;
+	FIoBuffer								ContainerHeaderChunk;
 	FSharedContainerHeader					Header;
 	FIASHostGroup							HostGroup;
-	FString									EncryptionKeyGuid;
-	FString									Name;
 	FString									MountId;
-	FAnsiString								ChunksDirectory;
-	TArray<FName>							CompressionFormats;
-	TArray<uint32>							BlockSizes;
-	TArray<FIoBlockHash>					BlockHashes;
-	TArray<FOnDemandTagSet>					TagSets;
-	TUniquePtr<uint8[]>						ChunkEntryData;
-	TArrayView<FIoChunkId>					ChunkIds;
-	TArrayView<FOnDemandChunkEntry> 		ChunkEntries;
+	TConstArrayView<FIoChunkId>				ChunkIds;
+	TConstArrayView<FOnDemandChunkEntry> 	ChunkEntries;
+	TConstArrayView<uint32>					BlockSizes;
+	TConstArrayView<FIoBlockHash>			BlockHashes;
+	TConstArrayView<FOnDemandTagSet>		TagSets;
+	TConstArrayView<uint32>					TagSetIndices;
 	TArray<FOnDemandChunkEntryReferences>	ChunkEntryReferences;
 	FIoRelativeUrl							RelativeUrl;
-	FIoContainerId							ContainerId;
-	uint32									BlockSize = 0;
+	FName									CompressionFormat;
 	FName									HostGroupName = FOnDemandHostGroup::DefaultName;
 	EOnDemandContainerFlags 				Flags = EOnDemandContainerFlags::None;
 };
 
+///////////////////////////////////////////////////////////////////////////////
+inline FString FOnDemandContainer::Name() const
+{
+	return FString(HeaderView.ContainerName());
+}
+
 int32 FOnDemandContainer::GetNumChunks() const
 {
 	return ChunkIds.Num();
@@ -213,7 +204,7 @@ FAnsiString FOnDemandContainer::GetTestUrl()
 	HashString << ChunkEntries[0].Hash;
 
 	TAnsiStringBuilder<256> Url;
-	Url << "/" << ChunksDirectory
+	Url << "/" << RelativeUrl.GetView()
 		<< "/" << HashString.ToView().Left(2)
 		<< "/" << HashString << ANSITEXTVIEW(".iochunk");
 
@@ -223,34 +214,23 @@ FAnsiString FOnDemandContainer::GetTestUrl()
 ///////////////////////////////////////////////////////////////////////////////
 struct FOnDemandChunkInfo
 {
-	FOnDemandChunkInfo()
-		: Entry(FOnDemandChunkEntry::Null)
-	{ }
-
-	FOnDemandChunkInfo(const FOnDemandChunkInfo& Other)
-		: SharedContainer(Other.SharedContainer)
-		, Entry(Other.Entry)
-	{
-
-	}
-
-	FOnDemandChunkInfo(FOnDemandChunkInfo&& Other)
-		: SharedContainer(MoveTemp(Other.SharedContainer))
-		, Entry(Other.Entry)
-	{
-
-	}
-
+											FOnDemandChunkInfo() = default;
+											FOnDemandChunkInfo(const FOnDemandChunkInfo& Other)
+												: SharedContainer(Other.SharedContainer)
+												, Entry(Other.Entry) { }
+											FOnDemandChunkInfo(FOnDemandChunkInfo&& Other)
+												: SharedContainer(MoveTemp(Other.SharedContainer))
+												, Entry(Other.Entry) { }
 
 	const									FIoHash& Hash() const { return Entry.Hash; }
 	uint32									RawSize() const { return Entry.RawSize; }
 	uint32									EncodedSize() const { return Entry.EncodedSize; }
-	uint32									BlockSize() const { return SharedContainer->BlockSize; }
-	FName									CompressionFormat() const { return SharedContainer->CompressionFormats[Entry.CompressionFormatIndex]; }
+	uint32									DiskSize() const { return Entry.GetDiskSize(); }
+	uint32									BlockSize() const { return SharedContainer->BlockSize(); }
+	FName									CompressionFormat() const { return SharedContainer->CompressionFormat; }
 	FMemoryView								EncryptionKey() const { return FMemoryView(SharedContainer->EncryptionKey.Key, FAES::FAESKey::KeySize); }
 	inline TConstArrayView<uint32>			Blocks() const;
 	inline TConstArrayView<FIoBlockHash>	BlockHashes() const;
-	FAnsiStringView							ChunksDirectory() const { return SharedContainer->ChunksDirectory; }
 	FIoRelativeUrl							RelativeUrl() const { return SharedContainer->RelativeUrl; }
 	const FOnDemandChunkEntry&				ChunkEntry() const { return Entry; }
 	inline const FIASHostGroup&				HostGroup() const;
@@ -261,7 +241,7 @@ struct FOnDemandChunkInfo
 	
 	inline static FOnDemandChunkInfo		Find(FSharedOnDemandContainer Container, const FIoChunkId& ChunkId);
 
-	inline void								GetUrl(FAnsiStringBuilderBase& Url) const;
+	static const FOnDemandChunkEntry		NullEntry;
 
 private:
 	friend class FOnDemandIoStore;
@@ -273,19 +253,25 @@ private:
 	{ }
 
 	FSharedOnDemandContainer	SharedContainer;
-	const FOnDemandChunkEntry&	Entry;
+	const FOnDemandChunkEntry&	Entry = NullEntry;
 };
 
 TConstArrayView<uint32> FOnDemandChunkInfo::Blocks() const
 {
-	return TConstArrayView<uint32>(SharedContainer->BlockSizes.GetData() + Entry.BlockOffset, Entry.BlockCount);
+	// When the chunk entry has a single compression block the block size is inlined into the block info as the first 4 bytes
+	return Entry.BlockInfo.Count() > 1
+		? TConstArrayView<uint32>(SharedContainer->BlockSizes.GetData() + Entry.BlockInfo.Offset(), Entry.BlockInfo.Count())
+		: TConstArrayView<uint32>(Entry.BlockInfo.Data(), 1);
 }
 
 TConstArrayView<FIoBlockHash> FOnDemandChunkInfo::BlockHashes() const
 {
+	// When the chunk entry has a single compression block the block hash is inlined into the block info as the second 4 bytes
 	return SharedContainer->BlockHashes.IsEmpty()
 		? TConstArrayView<FIoBlockHash>()
-		: TConstArrayView<FIoBlockHash>(SharedContainer->BlockHashes.GetData() + Entry.BlockOffset, Entry.BlockCount);
+		: Entry.BlockInfo.Count() > 1
+			? TConstArrayView<FIoBlockHash>(SharedContainer->BlockHashes.GetData() + Entry.BlockInfo.Offset(), Entry.BlockInfo.Count())
+			: TConstArrayView<uint32>(Entry.BlockInfo.Data() + 1, 1);
 }
 
 const FIASHostGroup& FOnDemandChunkInfo::HostGroup() const
@@ -304,17 +290,6 @@ FOnDemandChunkInfo FOnDemandChunkInfo::Find(FSharedOnDemandContainer Container,
 	return FOnDemandChunkInfo();
 }
 
-void FOnDemandChunkInfo::GetUrl(FAnsiStringBuilderBase& Url) const
-{
-	TAnsiStringBuilder<41> HashString;
-	HashString << Hash();
-
-	Url.Reset();
-	Url << "/" << ChunksDirectory()
-		<< "/" << HashString.ToView().Left(2)
-		<< "/" << HashString << ANSITEXTVIEW(".iochunk");
-}
-
 ///////////////////////////////////////////////////////////////////////////////
 /** Flush last access args. */
 struct FOnDemandFlushLastAccessArgs
@@ -420,13 +395,11 @@ private:
 	void								CompleteMountRequest(FMountRequest& MountRequest, FOnDemandMountResult&& MountResult);
 	void								OnEncryptionKeyAdded(const FGuid& Id, const FAES::FAESKey& Key);
 	void								OnHostGroupRegistered(const FName& HostGroup);
-	static void							CreateContainersFromToc(
-											FStringView MountId,
-											FStringView TocPath,
-											FOnDemandToc& Toc,
-											const TArray<FString>& ExistingContainers,
+	static FIoStatus					CreateContainersFromToc(
+											FOnDemandMountArgs& MountArgs,
+											TConstArrayView<FIoContainerId> Existing,
 											TArray<FSharedOnDemandContainer>& Out);
-	static FIoStatus						SetupHostGroup(const FSharedOnDemandContainer& Container, const FOnDemandMountArgs& MountArgs);
+	static FIoStatus					SetupHostGroup(const FSharedOnDemandContainer& Container, const FOnDemandMountArgs& MountArgs);
 
 #if !UE_BUILD_SHIPPING
 	TUniquePtr<FOnDemandDebugCommands>		DebugCommands;
@@ -441,7 +414,6 @@ private:
 	FDelegateHandle							OnMountPakHandle;
 	FDelegateHandle							OnServerPostForkHandle;
 	TArray<FSharedOnDemandContainer>		Containers;
-	TMap<FString, FIoBuffer>				PendingContainerHeaders;
 	mutable UE::FMutex						ContainerMutex;
 
 	TArray<FSharedMountRequest>				MountRequests;
diff --git a/Engine/Source/Runtime/Experimental/IoStore/OnDemandCore/Public/IO/IoStoreOnDemand.h b/Engine/Source/Runtime/Experimental/IoStore/OnDemandCore/Public/IO/IoStoreOnDemand.h
--- a/Engine/Source/Runtime/Experimental/IoStore/OnDemandCore/Public/IO/IoStoreOnDemand.h
+++ b/Engine/Source/Runtime/Experimental/IoStore/OnDemandCore/Public/IO/IoStoreOnDemand.h
@@ -182,8 +182,6 @@ ENUM_CLASS_FLAGS(EOnDemandMountOptions);
  */
 struct FOnDemandMountArgs
 {
-	/** Mount an already serialized TOC. */
-	TUniquePtr<FOnDemandToc> Toc;
 	/** Mandatory ID to be used for unmounting all container file(s) included in the TOC. */
 	FString MountId;
 	/** Relative URL from the primary endpoint from where to download the TOC. */
