Implement IAD cache replay serialization.

#Synced-CL 49969322

diff --git a/Engine/Source/Runtime/Experimental/IoStore/OnDemand/Private/OnDemandInstallCache.cpp b/Engine/Source/Runtime/Experimental/IoStore/OnDemand/Private/OnDemandInstallCache.cpp
--- a/Engine/Source/Runtime/Experimental/IoStore/OnDemand/Private/OnDemandInstallCache.cpp
+++ b/Engine/Source/Runtime/Experimental/IoStore/OnDemand/Private/OnDemandInstallCache.cpp
@@ -1777,7 +1777,7 @@ TResult<FCasSnapshot> FCasSnapshot::Load(const FString& SnapshotFile, int64* Out
 		*OutFileSize = Ar->Tell();
 	}
 
-	return MakeValue(Snapshot);
+	return MakeValue(MoveTemp(Snapshot));
 }
 
 TResult<int64> FCasSnapshot::TryCreateAndResetJournal(const FString& SnapshotFile, const FString& JournalFile)
diff --git a/Engine/Source/Runtime/Experimental/IoStore/OnDemandCore/Internal/IO/IoStoreOnDemandInternals.h b/Engine/Source/Runtime/Experimental/IoStore/OnDemandCore/Internal/IO/IoStoreOnDemandInternals.h
--- a/Engine/Source/Runtime/Experimental/IoStore/OnDemandCore/Internal/IO/IoStoreOnDemandInternals.h
+++ b/Engine/Source/Runtime/Experimental/IoStore/OnDemandCore/Internal/IO/IoStoreOnDemandInternals.h
@@ -86,14 +86,57 @@ struct FOnDemandContentInstallReplay
 {
 	// TODO: record mounting/unmounting containers
 
+	// TODO: Record when installs are processed and exectuted ?
+
+	// TODO: Record an event for file last access
+
+	enum class EVersion : int32
+	{
+		Invalid = 0,
+		Initial,
+
+		LatestPlusOne,
+		Latest = LatestPlusOne - 1
+	};
+
+	static const FGuid CustomVersionGuid;
+
+	struct FHeader
+	{
+		static const uint8 MagicSequence[16];
+
+		bool			IsValid() const;
+		static int64	Size() { return sizeof(FHeader); }
+
+		uint8		Magic[16] = { 0 };
+		EVersion	Version = EVersion::Invalid;
+		uint8		Pad[12] = { 0 };
+	};
+	static_assert(sizeof(FHeader) == 32);
+
+	struct FFooter
+	{
+		static const uint8 MagicSequence[16];
+
+		static int64	Size() { return sizeof(FFooter); }
+		bool			IsValid() const;
+
+		uint8 Magic[16] = { 0 };
+	};
+	static_assert(sizeof(FFooter) == 16);
+
 	struct FContentHandleEventData
 	{
 		FString DebugName;
+
+		bool operator==(const FContentHandleEventData& Other) const = default;
 	};
 
 	struct FContentHandleDestroyedEventData
 	{
 		int64 ContentHandleEventId = -1;
+
+		bool operator==(const FContentHandleDestroyedEventData& Other) const = default;
 	};
 
 	struct FInstallEventData
@@ -104,25 +147,34 @@ struct FOnDemandContentInstallReplay
 		int64 ContentHandleEventId = -1;
 		int32 Priority = 0;
 		EOnDemandInstallOptions Options = EOnDemandInstallOptions::None;
+
+		bool operator==(const FInstallEventData& Other) const = default;
 	};
 
 	struct FCancelInstallEventData
 	{
 		int64 InstallEventId = -1;
+
+		bool operator==(const FCancelInstallEventData& Other) const = default;
 	};
 
 	struct FUpdatePriorityEventData
 	{
 		int64 InstallEventId = -1;
 		int32 NewPriority = 0;
+
+		bool operator==(const FUpdatePriorityEventData& Other) const = default;
 	};
 
 	struct FInstallDestroyedEventData
 	{
-		int64 InstallEventId = 0;
+		int64 InstallEventId = -1;
+
+		bool operator==(const FInstallDestroyedEventData& Other) const = default;
 	};
 
 	using FEventVariant = TVariant<
+		FEmptyVariantState,
 		FContentHandleEventData,
 		FContentHandleDestroyedEventData,
 		FInstallEventData,
@@ -136,6 +188,13 @@ struct FOnDemandContentInstallReplay
 		FEventVariant EventData;
 	};
 
+#if WITH_IOSTORE_ONDEMAND_TESTS
+	bool operator==(const FOnDemandContentInstallReplay& Other) const;
+#endif // WITH_IOSTORE_ONDEMAND_TESTS
+
+	static TResult<FOnDemandContentInstallReplay>	Load(const FString& Filename, int64* OutFileSize = nullptr);
+	static TResult<int64>							Save(const FOnDemandContentInstallReplay& Replay, const FString& Filename);
+
 	TArray<FEvent> Events;
 };
 
@@ -149,6 +208,7 @@ public:
 
 	UE_API static void StartRecording(FOnDemandContentInstallReplayRecorder* InInstance);
 
+	// TODO: return TResult
 	UE_API TValueOrError<FOnDemandContentInstallReplay, void> StopRecording();
 
 	// Records the public content handle, no the internal implementation
diff --git a/Engine/Source/Runtime/Experimental/IoStore/OnDemandCore/Private/IoStoreOnDemandInternals.cpp b/Engine/Source/Runtime/Experimental/IoStore/OnDemandCore/Private/IoStoreOnDemandInternals.cpp
--- a/Engine/Source/Runtime/Experimental/IoStore/OnDemandCore/Private/IoStoreOnDemandInternals.cpp
+++ b/Engine/Source/Runtime/Experimental/IoStore/OnDemandCore/Private/IoStoreOnDemandInternals.cpp
@@ -3,8 +3,15 @@
 #include "IO/IoStoreOnDemandInternals.h"
 
 #include "Async/UniqueLock.h"
+#include "HAL/FileManager.h"
 #include "Logging/StructuredLog.h"
 
+#if WITH_IOSTORE_ONDEMAND_TESTS
+#include "TestHarness.h"
+#include "TestMacros/Assertions.h"
+#include <catch2/generators/catch_generators.hpp>
+#endif
+
 namespace UE::IoStore
 {
 
@@ -251,7 +258,234 @@ void FOnDemandInstallRequest::UpdatePriority(int32 NewPriority)
 	}
 }
 
+////////////////////////////////////////////////////////////////////////////////
 #if !UE_BUILD_SHIPPING
+const FGuid FOnDemandContentInstallReplay::CustomVersionGuid(0x945E86CB, 0x453E4132, 0x9F036DF2, 0x0D9A930E);
+
+const uint8 FOnDemandContentInstallReplay::FHeader::MagicSequence[16] = { 'R', 'E', 'P', 'L', 'A', 'Y', 'I', 'N', 'S', 'T', 'A', 'L', 'L', 'H', 'D' };
+const uint8 FOnDemandContentInstallReplay::FFooter::MagicSequence[16] = { 'R', 'E', 'P', 'L', 'A', 'Y', 'I', 'N', 'S', 'T', 'A', 'L', 'L', 'F', 'T' };
+
+bool FOnDemandContentInstallReplay::FHeader::IsValid() const
+{
+	if (FMemory::Memcmp(&Magic, &FHeader::MagicSequence, sizeof(FHeader::MagicSequence)) != 0)
+	{
+		return false;
+	}
+
+	if (static_cast<uint32>(Version) > static_cast<uint32>(EVersion::Latest))
+	{
+		return false;
+	}
+
+	return true;
+}
+
+bool FOnDemandContentInstallReplay::FFooter::IsValid() const
+{
+	return FMemory::Memcmp(Magic, FFooter::MagicSequence, sizeof(FFooter::MagicSequence)) == 0;
+}
+
+FArchive& operator<<(FArchive& Ar, FOnDemandContentInstallReplay::FContentHandleEventData& EventData)
+{
+	Ar << EventData.DebugName;
+	return Ar;
+}
+
+FArchive& operator<<(FArchive& Ar, FOnDemandContentInstallReplay::FContentHandleDestroyedEventData& EventData)
+{
+	Ar << EventData.ContentHandleEventId;
+	return Ar;
+}
+
+FArchive& operator<<(FArchive& Ar, FOnDemandContentInstallReplay::FInstallEventData& EventData)
+{
+	Ar << EventData.MountId;
+	Ar << EventData.TagSets;
+	Ar << EventData.PackageIds;
+	Ar << EventData.ContentHandleEventId;
+	Ar << EventData.Priority;
+	Ar << EventData.Options;
+	return Ar;
+}
+
+FArchive& operator<<(FArchive& Ar, FOnDemandContentInstallReplay::FCancelInstallEventData& EventData)
+{
+	Ar << EventData.InstallEventId;
+	return Ar;
+}
+
+FArchive& operator<<(FArchive& Ar, FOnDemandContentInstallReplay::FUpdatePriorityEventData& EventData)
+{
+	Ar << EventData.InstallEventId;
+	Ar << EventData.NewPriority;
+	return Ar;
+}
+
+FArchive& operator<<(FArchive& Ar, FOnDemandContentInstallReplay::FInstallDestroyedEventData& EventData)
+{
+	Ar << EventData.InstallEventId;
+	return Ar;
+}
+
+FArchive& operator<<(FArchive& Ar, FOnDemandContentInstallReplay::FEvent& Event)
+{
+	Ar << Event.TimeOffset;
+	Ar << Event.EventData;
+
+	return Ar;
+}
+
+#if WITH_IOSTORE_ONDEMAND_TESTS
+struct FOnDemandContentInstallReplayEventDataEqualityVistor
+{
+	const FOnDemandContentInstallReplay::FEventVariant& EventData;
+
+	bool operator()(const FEmptyVariantState& OtherData) const
+	{
+		CHECK(false);
+		ensure(false);
+		return true;
+	}
+
+	template<typename TVistedType>
+	bool operator()(const TVistedType& OtherData) const
+	{
+		return EventData.Get<TVistedType>() == OtherData;
+	}
+};
+
+bool FOnDemandContentInstallReplay::operator==(const FOnDemandContentInstallReplay& Other) const
+{
+	if (Events.Num() != Other.Events.Num())
+	{
+		return false;
+	}
+
+	for (int i = 0; const FEvent& Event : Events)
+	{
+		const FEvent& OtherEvent = Other.Events[i++];
+
+		if (Event.TimeOffset != OtherEvent.TimeOffset)
+		{
+			return false;
+		}
+
+		if (Event.EventData.GetIndex() != OtherEvent.EventData.GetIndex())
+		{
+			return false;
+		}
+
+		FOnDemandContentInstallReplayEventDataEqualityVistor Vistor{ .EventData = Event.EventData };
+		bool bEventDataResult = Visit(Vistor, OtherEvent.EventData);
+		if (bEventDataResult == false)
+		{
+			return false;
+		}
+	}
+
+	return true;
+}
+#endif // WITH_IOSTORE_ONDEMAND_TESTS
+
+TResult<FOnDemandContentInstallReplay> FOnDemandContentInstallReplay::Load(const FString& Filename, int64* OutFileSize)
+{
+	using namespace UE::UnifiedError::IoStoreOnDemand;
+
+	IFileManager& Ifm = IFileManager::Get();
+
+	TUniquePtr<FArchive> Ar(Ifm.CreateFileReader(*Filename));
+	if (Ar.IsValid() == false)
+	{
+		return MakeError(InstallReplayLoadError::MakeError());
+	}
+
+	FHeader Header;
+	Ar->Serialize(reinterpret_cast<uint8*>(&Header), FHeader::Size());
+	if (Ar->IsError() || Header.IsValid() == false)
+	{
+		return MakeError(InstallReplayLoadError::MakeError());
+	}
+
+	Ar->SetCustomVersion(CustomVersionGuid, static_cast<int32>(Header.Version), "InstallReplayVersion");
+
+	FOnDemandContentInstallReplay Replay;
+	*Ar << Replay.Events;
+
+	FFooter Footer;
+	Ar->Serialize(reinterpret_cast<uint8*>(&Footer), FFooter::Size());
+	if (Ar->IsError() || Footer.IsValid() == false)
+	{
+		return MakeError(InstallReplayLoadError::MakeError());
+	}
+
+	if (OutFileSize != nullptr)
+	{
+		*OutFileSize = Ar->Tell();
+	}
+
+	return MakeValue(MoveTemp(Replay));
+}
+
+TResult<int64> FOnDemandContentInstallReplay::Save(const FOnDemandContentInstallReplay& Replay, const FString& Filename)
+{
+	using namespace UE::UnifiedError::IoStoreOnDemand;
+
+	IFileManager& Ifm = IFileManager::Get();
+
+	TUniquePtr<FArchive> Ar(Ifm.CreateFileWriter(*Filename));
+	if (Ar.IsValid() == false)
+	{
+		return MakeError(InstallReplaySaveError::MakeError());
+	}
+
+	FHeader Header;
+	FMemory::Memcpy(&Header.Magic, &FHeader::MagicSequence, sizeof(FHeader::MagicSequence));
+	Header.Version = EVersion::Latest;
+
+	Ar->SetCustomVersion(CustomVersionGuid, static_cast<int32>(Header.Version), "InstallReplayVersion");
+
+	Ar->Serialize(reinterpret_cast<uint8*>(&Header), FHeader::Size());
+	if (Ar->IsError())
+	{
+		const uint32 LastError = FPlatformMisc::GetLastError();
+		Ar.Reset();
+		Ifm.Delete(*Filename);
+		return MakeError(InstallReplaySaveError::MakeError());
+	}
+
+	FOnDemandContentInstallReplay& NonConst = const_cast<FOnDemandContentInstallReplay&>(Replay);
+	*Ar << NonConst.Events;
+
+	if (Ar->IsError())
+	{
+		Ar.Reset();
+		Ifm.Delete(*Filename);
+		return MakeError(InstallReplaySaveError::MakeError());
+	}
+
+	FFooter Footer;
+	FMemory::Memcpy(&Footer.Magic, &FFooter::MagicSequence, sizeof(FFooter::MagicSequence));
+	Ar->Serialize(reinterpret_cast<uint8*>(&Footer), FFooter::Size());
+	if (Ar->IsError())
+	{
+		const uint32 LastError = FPlatformMisc::GetLastError();
+		Ar.Reset();
+		Ifm.Delete(*Filename);
+		return MakeError(InstallReplaySaveError::MakeError());
+	}
+
+	const int64 FileSize = Ar->TotalSize();
+	if (Ar->Close() == false)
+	{
+		const uint32 LastError = FPlatformMisc::GetLastError();
+		Ar.Reset();
+		Ifm.Delete(*Filename);
+		return MakeError(InstallReplaySaveError::MakeError());
+	}
+
+	return MakeValue(FileSize);
+}
+
 FOnDemandContentInstallReplayRecorder* FOnDemandContentInstallReplayRecorder::Instance = nullptr;
 
 FOnDemandContentInstallReplayRecorder* FOnDemandContentInstallReplayRecorder::Get()
@@ -476,3 +710,130 @@ FString LexToString(const UE::IoStore::FOnDemandInternalContentHandle& Handle)
 {
 	return FString::Printf(TEXT("Id=0x%" UPTRINT_X_FMT " DebugName='%s'"), Handle.HandleId(), *Handle.DebugName);
 }
+
+////////////////////////////////////////////////////////////////////////////////
+#if WITH_IOSTORE_ONDEMAND_TESTS
+
+TEST_CASE("IoStore::OnDemand::InstallReplay", "[IoStoreOnDemand][InstallReplay]")
+{
+	class FTmpDirectoryScope
+	{
+	public:
+		explicit FTmpDirectoryScope(const FString& InDir)
+			: Ifm(IFileManager::Get())
+			, Dir(InDir)
+		{
+			const bool bTree = true;
+			const bool bRequireExists = false;
+			Ifm.DeleteDirectory(*Dir, bRequireExists, bTree);
+			Ifm.MakeDirectory(*Dir, bTree);
+		}
+
+		~FTmpDirectoryScope()
+		{
+			const bool bTree = true;
+			const bool bRequireExists = false;
+			Ifm.DeleteDirectory(*Dir, bRequireExists, bTree);
+		}
+	private:
+		IFileManager& Ifm;
+		FString Dir;
+	};
+
+	const FString TestBaseDir = "TestTmpDir";
+
+	SECTION("SaveLoadRoundtrip")
+	{
+		using namespace UE::IoStore;
+
+		FOnDemandContentInstallReplay ExpectedReplay;
+		ExpectedReplay.Events.Emplace(FOnDemandContentInstallReplay::FEvent
+		{
+			.TimeOffset = 1,
+			.EventData{ 
+				TInPlaceType<FOnDemandContentInstallReplay::FContentHandleEventData>(), 
+				FOnDemandContentInstallReplay::FContentHandleEventData
+				{
+					.DebugName = TEXT("ContentHandle")
+				}
+			}
+		});
+		ExpectedReplay.Events.Emplace(FOnDemandContentInstallReplay::FEvent
+		{
+			.TimeOffset = 2,
+			.EventData{ 
+				TInPlaceType<FOnDemandContentInstallReplay::FInstallEventData>(),
+				FOnDemandContentInstallReplay::FInstallEventData
+				{
+					.MountId = TEXT("MountId"),
+					.TagSets{ TEXT("TagA"), TEXT("TagB") },
+					.PackageIds{ FPackageId::FromValue(0xABAB) },
+					.ContentHandleEventId = 0,
+					.Priority = 1,
+					.Options = EOnDemandInstallOptions::InstallSoftReferences
+				}
+			}
+		});
+		ExpectedReplay.Events.Emplace(FOnDemandContentInstallReplay::FEvent
+		{
+			.TimeOffset = 3,
+			.EventData{
+				TInPlaceType<FOnDemandContentInstallReplay::FUpdatePriorityEventData>(),
+				FOnDemandContentInstallReplay::FUpdatePriorityEventData
+				{
+					.InstallEventId = 1,
+					.NewPriority = 3
+				}
+			}
+		});
+		ExpectedReplay.Events.Emplace(FOnDemandContentInstallReplay::FEvent
+		{
+			.TimeOffset = 4,
+			.EventData{
+				TInPlaceType<FOnDemandContentInstallReplay::FCancelInstallEventData>(),
+				FOnDemandContentInstallReplay::FCancelInstallEventData
+				{
+					.InstallEventId = 1
+				}
+			}
+		});
+		ExpectedReplay.Events.Emplace(FOnDemandContentInstallReplay::FEvent
+		{
+			.TimeOffset = 5,
+			.EventData{
+				TInPlaceType<FOnDemandContentInstallReplay::FInstallDestroyedEventData>(),
+				FOnDemandContentInstallReplay::FInstallDestroyedEventData
+				{
+					.InstallEventId = 1
+				}
+			}
+		});
+		ExpectedReplay.Events.Emplace(FOnDemandContentInstallReplay::FEvent
+		{
+			.TimeOffset = 6,
+			.EventData{
+				TInPlaceType<FOnDemandContentInstallReplay::FContentHandleDestroyedEventData>(),
+				FOnDemandContentInstallReplay::FContentHandleDestroyedEventData
+				{
+					.ContentHandleEventId = 0
+				}
+			}
+		});
+
+		FTmpDirectoryScope _(TestBaseDir);
+		const FString Filename = TestBaseDir / TEXT("test.installreplay");
+
+		TResult<int64> SaveResult = FOnDemandContentInstallReplay::Save(ExpectedReplay, Filename);
+		CHECK(SaveResult.HasValue());
+
+		int64 BytesLoaded = 0;
+		TResult<FOnDemandContentInstallReplay> LoadResult = FOnDemandContentInstallReplay::Load(Filename, &BytesLoaded);
+		CHECK(LoadResult.HasValue());
+
+		CHECK(SaveResult.GetValue() == BytesLoaded);
+
+		CHECK(ExpectedReplay == LoadResult.GetValue());
+	}
+}
+
+#endif // WITH_IOSTORE_ONDEMAND_TESTS
diff --git a/Engine/Source/Runtime/Experimental/IoStore/OnDemandCore/Private/OnDemandError.cpp b/Engine/Source/Runtime/Experimental/IoStore/OnDemandCore/Private/OnDemandError.cpp
--- a/Engine/Source/Runtime/Experimental/IoStore/OnDemandCore/Private/OnDemandError.cpp
+++ b/Engine/Source/Runtime/Experimental/IoStore/OnDemandCore/Private/OnDemandError.cpp
@@ -16,6 +16,8 @@ UE_DEFINE_ERROR(CasError, IoStoreOnDemand);
 UE_DEFINE_ERROR(CasJournalError, IoStoreOnDemand);
 UE_DEFINE_ERROR(CasSnapshotError, IoStoreOnDemand);
 UE_DEFINE_ERROR(InstallerCacheFullError, IoStoreOnDemand);
+UE_DEFINE_ERROR(InstallReplayLoadError, IoStoreOnDemand);
+UE_DEFINE_ERROR(InstallReplaySaveError, IoStoreOnDemand);
 
 namespace UE::UnifiedError::IoStore
 {
diff --git a/Engine/Source/Runtime/Experimental/IoStore/OnDemandCore/Public/IO/OnDemandError.h b/Engine/Source/Runtime/Experimental/IoStore/OnDemandCore/Public/IO/OnDemandError.h
--- a/Engine/Source/Runtime/Experimental/IoStore/OnDemandCore/Public/IO/OnDemandError.h
+++ b/Engine/Source/Runtime/Experimental/IoStore/OnDemandCore/Public/IO/OnDemandError.h
@@ -25,6 +25,8 @@ UE_DECLARE_ERROR(UE_API,			CasError,							9,	IoStoreOnDemand, NSLOCTEXT("IoStor
 UE_DECLARE_ERROR(UE_API,			CasJournalError,					10,	IoStoreOnDemand, NSLOCTEXT("IoStoreOnDemand", "CasJournalError", "Cas journal error"));
 UE_DECLARE_ERROR(UE_API,			CasSnapshotError,					11,	IoStoreOnDemand, NSLOCTEXT("IoStoreOnDemand", "CasSnapshotError", "Cas snapshot error"));
 UE_DECLARE_ERROR(UE_API,			InstallerCacheFullError,			12, IoStoreOnDemand, NSLOCTEXT("IoStoreOnDemand", "InstallerCacheFullError", "Cache full error"));
+UE_DECLARE_ERROR(UE_API,			InstallReplayLoadError,				13, IoStoreOnDemand, NSLOCTEXT("IoStoreOnDemand", "InstallReplayLoadError", "Failed to load install replay"));
+UE_DECLARE_ERROR(UE_API,			InstallReplaySaveError,				14, IoStoreOnDemand, NSLOCTEXT("IoStoreOnDemand", "InstallReplaySaveError", "Failed to save install replay"));
 
 ////////////////////////////////////////////////////////////////////////////////
 namespace UE::IoStore
