Potential crash fix when constraint is destroyed but the Graph Edge hasn't been removed first.
The fix is behind a cVar, which is disabled by default

#Synced-CL 49997110

diff --git a/Engine/Source/Runtime/Experimental/Chaos/Private/Chaos/PBDRigidsEvolution.cpp b/Engine/Source/Runtime/Experimental/Chaos/Private/Chaos/PBDRigidsEvolution.cpp
--- a/Engine/Source/Runtime/Experimental/Chaos/Private/Chaos/PBDRigidsEvolution.cpp
+++ b/Engine/Source/Runtime/Experimental/Chaos/Private/Chaos/PBDRigidsEvolution.cpp
@@ -61,6 +61,9 @@ namespace Chaos
 																Pruning the acceleration structure is an optimization technique that speed up the spatial acceleration update which is a normally a single threaded process by using several threads to prune data to avoid some unecessary update. \
 																In single worker physics it would be better to disable it."));
 
+	bool bSafeDisconnectConstraints = false;
+	FAutoConsoleVariableRef CVarSafeDisconnectConstraints(TEXT("p.Chaos.SafeDisconnectConstraints"), bSafeDisconnectConstraints, TEXT("Additional check in DisconnectConstraints to be sure all MidPhase's collision constraints have been removed from the Island Manager."));
+
 	CHAOS_API FBroadPhaseConfig BroadPhaseConfig;
 
 	FAutoConsoleVariableRef CVarBroadphaseIsTree(TEXT("p.BroadphaseType"), BroadPhaseConfig.BroadphaseType, TEXT(""));
diff --git a/Engine/Source/Runtime/Experimental/Chaos/Public/Chaos/PBDRigidsEvolution.h b/Engine/Source/Runtime/Experimental/Chaos/Public/Chaos/PBDRigidsEvolution.h
--- a/Engine/Source/Runtime/Experimental/Chaos/Public/Chaos/PBDRigidsEvolution.h
+++ b/Engine/Source/Runtime/Experimental/Chaos/Public/Chaos/PBDRigidsEvolution.h
@@ -38,6 +38,7 @@ extern int32 ChaosNumContactIterationsOverride;
 namespace Chaos
 {
 extern CHAOS_API int32 ForceNoCollisionIntoSQ;
+extern CHAOS_API bool bSafeDisconnectConstraints;
 
 struct FBroadPhaseConfig
 {
@@ -677,6 +678,19 @@ public:
 		for (FGeometryParticleHandle* ParticleHandle : RemovedParticles)
 		{
 			RemoveConstraintsFromConstraintGraph(ParticleHandle->ParticleConstraints());
+
+			if (bSafeDisconnectConstraints)
+			{
+				for(int32 Index = 0; Index < ParticleHandle->ParticleCollisions().Num(); Index++)
+				{
+					FParticlePairMidPhase* MidPhase = ParticleHandle->ParticleCollisions().GetMidPhase(Index);
+					MidPhase->VisitCollisions([this](FPBDCollisionConstraint& Constraint)
+					{
+						RemoveConstraintsFromConstraintGraph({Constraint.GetConstraintHandle()});
+						return ECollisionVisitorResult::Continue;
+					}, ECollisionVisitorFlags::VisitAllCurrentAndExpired);
+				}
+			}
 		}
 
 		for (FPBDConstraintContainer* Container : ConstraintContainers)
