Refactor FOnDemandContentInstallReplayRecorder to record from internal IAD installer state instead of external API calls.
- Fix replays not using TArray64
- Fix replay recorder ensures after recording ends.
- Use the InstallerRequest as the request id.

#Synced-CL 49969322

diff --git a/Engine/Source/Runtime/Experimental/IoStore/OnDemand/Private/OnDemandContentInstaller.cpp b/Engine/Source/Runtime/Experimental/IoStore/OnDemand/Private/OnDemandContentInstaller.cpp
--- a/Engine/Source/Runtime/Experimental/IoStore/OnDemand/Private/OnDemandContentInstaller.cpp
+++ b/Engine/Source/Runtime/Experimental/IoStore/OnDemand/Private/OnDemandContentInstaller.cpp
@@ -364,25 +364,6 @@ FOnDemandContentInstaller::FRequest::FInstall::FInstall(
 	, OnProgress(MoveTemp(InOnProgress)) 
 	, Request(MakeShared<FOnDemandInternalInstallRequest, ESPMode::ThreadSafe>(InstallerRequest))
 {
-#if !UE_BUILD_SHIPPING
-	// Ideally, this would be in FOnDemandInstallRequest::FOnDemandInstallRequest but that constructor
-	// doesn't have access to all the data because the type erasure we do on the request.
-	if (FOnDemandContentInstallReplayRecorder* Recorder = FOnDemandContentInstallReplayRecorder::Get())
-	{
-		Recorder->RecordInstall(
-			UPTRINT(Request.Get()),
-			Args.ContentHandle.IsValid() ? Args.ContentHandle.Handle->HandleId() : 0,
-			Args.MountId,
-			Args.TagSets,
-			Args.PackageIds,
-			Args.Priority,
-			Args.Options);
-	}
-
-	// FOnDemandInstallRequest::~FOnDemandInstallRequest will call RecordInstallDestroyed()
-	// In general we want to record the public API calls, not the lifetime of internal data structures.
-	// See comments in FOnDemandInstallRequest::FOnDemandInstallRequest
-#endif // !UE_BUILD_SHIPPING
 }
 
 FOnDemandContentInstaller::FOnDemandContentInstaller(FOnDemandIoStore& InIoStore)
@@ -408,6 +389,7 @@ FSharedInternalInstallRequest FOnDemandContentInstaller::EnqueueInstallRequest(
 	}
 
 	new(Request) FRequest(MoveTemp(Args), MoveTemp(OnCompleted), MoveTemp(OnProgress));
+	Request->AsInstall().Request->IoStore = IoStore.AsWeak();
 
 	FOnDemandContentInstallerStats::OnRequestEnqueued();
 
@@ -485,6 +467,13 @@ void FOnDemandContentInstaller::CancelInstallRequest(FSharedInternalInstallReque
 						return;
 					}
 
+#if !UE_BUILD_SHIPPING
+					if (FOnDemandContentInstallReplayRecorder* Recorder = FOnDemandContentInstallReplayRecorder::Get())
+					{
+						Recorder->RecordCancel(InstallRequest->InstallerRequest);
+					}
+#endif // !UE_BUILD_SHIPPING
+
 					UE_LOG(LogIoStoreOnDemand, Log, TEXT("Cancelling install request, ContentHandle=(%s)"),
 						*LexToString(Install.Args.ContentHandle));
 
@@ -499,6 +488,13 @@ void FOnDemandContentInstaller::CancelInstallRequest(FSharedInternalInstallReque
 						return;
 					}
 
+#if !UE_BUILD_SHIPPING
+					if (FOnDemandContentInstallReplayRecorder* Recorder = FOnDemandContentInstallReplayRecorder::Get())
+					{
+						Recorder->RecordCancel(InstallRequest->InstallerRequest);
+					}
+#endif // !UE_BUILD_SHIPPING
+
 					UE_LOG(LogIoStoreOnDemand, Log, TEXT("Cancelling install request, ContentHandle=(%s)"),
 						*LexToString(Install.Args.ContentHandle));
 
@@ -532,6 +528,13 @@ void FOnDemandContentInstaller::UpdateInstallRequestPriority(FSharedInternalInst
 				return;
 			}
 
+#if !UE_BUILD_SHIPPING
+			if (FOnDemandContentInstallReplayRecorder* Recorder = FOnDemandContentInstallReplayRecorder::Get())
+			{
+				Recorder->RecordUpdatePriority(InstallRequest->InstallerRequest, NewPriority);
+			}
+#endif // !UE_BUILD_SHIPPING
+
 			FRequest&			Request	= *reinterpret_cast<FRequest*>(InstallRequest->InstallerRequest);
 			FRequest::FInstall& Install	= Request.AsInstall();
 
@@ -791,9 +794,35 @@ void FOnDemandContentInstaller::ProcessInstallRequest(FRequest& Request)
 	if (ContentHandle->IoStore.IsValid() == false)
 	{
 		// First time this content handle is used
+		// Is this threadsafe? This is read in the ContentHandle destructor which is invoked by the shard ptr
+		// ref controller which uses release/aquire semantics. So I believe it is safe as long as its only read 
+		// in the destructor.
 		ContentHandle->IoStore = IoStore.AsWeak();
+
+#if !UE_BUILD_SHIPPING		
+		if (FOnDemandContentInstallReplayRecorder* Recorder = FOnDemandContentInstallReplayRecorder::Get())
+		{
+			const FSharedString& DebugName = ContentHandle->DebugName;
+			Recorder->RecordContentHandle(ContentHandle->HandleId(), FString::ConstructFromPtrSize(*DebugName, DebugName.Len()));
+		}
+#endif // !UE_BUILD_SHIPPING
 	}
 
+#if !UE_BUILD_SHIPPING	
+	if (FOnDemandContentInstallReplayRecorder* Recorder = FOnDemandContentInstallReplayRecorder::Get())
+	{
+		const FOnDemandInstallArgs& Args = InstallRequest.Args;
+		Recorder->RecordInstall(
+			InstallRequest.Request->InstallerRequest,
+			ContentHandle->HandleId(),
+			Args.MountId,
+			Args.TagSets,
+			Args.PackageIds,
+			Args.Priority,
+			Args.Options);
+	}
+#endif // !UE_BUILD_SHIPPING
+
 	TSet<FSharedOnDemandContainer> ContainersForInstallation;
 	TSet<FPackageId> PackageIdsToInstall;
 	if (FIoStatus Status = IoStore.GetContainersAndPackagesForInstall(
@@ -1604,6 +1633,13 @@ void FOnDemandContentInstaller::CompleteInstallRequest(FRequest& Request)
 	InstallResult.DurationInSeconds	= DurationInSeconds;
 	InstallResult.Error				= Request.ConsumeError();
 
+#if !UE_BUILD_SHIPPING	
+	if (FOnDemandContentInstallReplayRecorder* Recorder = FOnDemandContentInstallReplayRecorder::Get())
+	{
+		Recorder->RecordInstallDestroyed(InstallRequest.Request->InstallerRequest);
+	}
+#endif // !UE_BUILD_SHIPPING
+
 	{
 		TUniqueLock Lock(Mutex);
 
diff --git a/Engine/Source/Runtime/Experimental/IoStore/OnDemand/Private/OnDemandIoStore.cpp b/Engine/Source/Runtime/Experimental/IoStore/OnDemand/Private/OnDemandIoStore.cpp
--- a/Engine/Source/Runtime/Experimental/IoStore/OnDemand/Private/OnDemandIoStore.cpp
+++ b/Engine/Source/Runtime/Experimental/IoStore/OnDemand/Private/OnDemandIoStore.cpp
@@ -631,7 +631,7 @@ FOnDemandInstallRequest FOnDemandIoStore::Install(
 		MoveTemp(OnCompleted),
 		MoveTemp(OnProgress));
 
-	return FOnDemandInstallRequest(AsWeak(), InstallRequest);
+	return FOnDemandInstallRequest(InstallRequest);
 }
 
 void FOnDemandIoStore::Purge(FOnDemandPurgeArgs&& Args, FOnDemandPurgeCompleted&& OnCompleted)
@@ -1766,6 +1766,13 @@ void FOnDemandIoStore::ReleaseContent(FOnDemandInternalContentHandle& ContentHan
 	LLM_SCOPE_BYNAME(TEXT("FileSystem/OnDemandIoStore"));
 	UE_LOG(LogIoStoreOnDemand, Verbose, TEXT("Releasing content handle, %s"), *LexToString(ContentHandle));
 
+#if !UE_BUILD_SHIPPING
+	if (FOnDemandContentInstallReplayRecorder* Recorder = FOnDemandContentInstallReplayRecorder::Get())
+	{
+		Recorder->RecordContentHandleDestroyed(ContentHandle.HandleId());
+	}
+#endif // !UE_BUILD_SHIPPING
+
 	FOnDemandFlushLastAccessArgs FlushLastAccessArgs;
 	FlushLastAccessArgs.bForceLastAccessDirty = false;
 
diff --git a/Engine/Source/Runtime/Experimental/IoStore/OnDemandCore/Internal/IO/IoStoreOnDemandInternals.h b/Engine/Source/Runtime/Experimental/IoStore/OnDemandCore/Internal/IO/IoStoreOnDemandInternals.h
--- a/Engine/Source/Runtime/Experimental/IoStore/OnDemandCore/Internal/IO/IoStoreOnDemandInternals.h
+++ b/Engine/Source/Runtime/Experimental/IoStore/OnDemandCore/Internal/IO/IoStoreOnDemandInternals.h
@@ -50,10 +50,19 @@ public:
 
 	UE_API ~FOnDemandInternalContentHandle();
 
+	FOnDemandInternalContentHandle(const FOnDemandInternalContentHandle& Other) = default;
+	FOnDemandInternalContentHandle(FOnDemandInternalContentHandle&& Other) = default;
+
+	FOnDemandInternalContentHandle& operator=(const FOnDemandInternalContentHandle& Other) = default;
+	FOnDemandInternalContentHandle& operator=(FOnDemandInternalContentHandle&& Other) = default;
+
 	UPTRINT HandleId() const { return UPTRINT(this); }
 
 	FSharedString			DebugName;
-	FWeakOnDemandIoStore	IoStore;
+
+private:
+	friend class FOnDemandContentInstaller;
+	FWeakOnDemandIoStore	IoStore; // This needs additional syncronization if its read outside the destructor
 };
 
 ////////////////////////////////////////////////////////////////////////////////
@@ -65,6 +74,7 @@ public:
 	explicit FOnDemandInternalInstallRequest(UPTRINT InInstallerRequest)
 		: InstallerRequest(InInstallerRequest) { }
 
+	FWeakOnDemandIoStore	IoStore;
 	UPTRINT					InstallerRequest = 0;
 	std::atomic<EStatus>	Status{EStatus::Pending};
 };
@@ -195,7 +205,7 @@ struct FOnDemandContentInstallReplay
 	static TResult<FOnDemandContentInstallReplay>	Load(const FString& Filename, int64* OutFileSize = nullptr);
 	static TResult<int64>							Save(const FOnDemandContentInstallReplay& Replay, const FString& Filename);
 
-	TArray<FEvent> Events;
+	TArray64<FEvent> Events;
 };
 
 class FOnDemandContentInstallReplayRecorder
@@ -245,7 +255,7 @@ public:
 private:
 	void StartRecordingInternal();
 
-	void Reset();
+	void Reset(bool bInValid = false);
 
 	int64 GetTimeOffset() const;
 
@@ -265,18 +275,10 @@ private:
 	}
 
 private:
-
 	FMutex Mutex;
-
-	struct FContentHandleRecord
-	{
-		int64 EventIndex = -1;
-		uint64 RefCount = 0;
-	};
-
-	TMap<UPTRINT, FContentHandleRecord> ContentHandleEventToIndex;
+	TMap<UPTRINT, int64> ContentHandleEventToIndex;
 	TMap<UPTRINT, int64> InstallEventToIndex;
-	TPagedArray<FOnDemandContentInstallReplay::FEvent, 4096> Events;
+	TPagedArray<FOnDemandContentInstallReplay::FEvent, 16384, FDefaultAllocator64> Events;
 	FDateTime StartTime;
 	bool bValid = true;
 };
diff --git a/Engine/Source/Runtime/Experimental/IoStore/OnDemandCore/Private/IoStoreOnDemandInternals.cpp b/Engine/Source/Runtime/Experimental/IoStore/OnDemandCore/Private/IoStoreOnDemandInternals.cpp
--- a/Engine/Source/Runtime/Experimental/IoStore/OnDemandCore/Private/IoStoreOnDemandInternals.cpp
+++ b/Engine/Source/Runtime/Experimental/IoStore/OnDemandCore/Private/IoStoreOnDemandInternals.cpp
@@ -36,6 +36,7 @@ FOnDemandWeakContentHandle FOnDemandWeakContentHandle::FromUnsafeHandle(UPTRINT
 	};
 }
 
+///////////////////////////////////////////////////////////////////////////////
 FOnDemandInternalContentHandle::FOnDemandInternalContentHandle()
 	: FOnDemandInternalContentHandle(FSharedString(TEXT("NoName")))
 {
@@ -46,7 +47,6 @@ FOnDemandInternalContentHandle::FOnDemandInternalContentHandle(FSharedString InD
 {
 }
 
-///////////////////////////////////////////////////////////////////////////////
 FOnDemandInternalContentHandle::~FOnDemandInternalContentHandle()
 {
 	if (TSharedPtr<IOnDemandIoStore, ESPMode::ThreadSafe> Pinned = IoStore.Pin(); Pinned.IsValid())
@@ -56,99 +56,10 @@ FOnDemandInternalContentHandle::~FOnDemandInternalContentHandle()
 }
 
 ///////////////////////////////////////////////////////////////////////////////
-FOnDemandContentHandle::FOnDemandContentHandle()
-{
-}
-
-FOnDemandContentHandle::~FOnDemandContentHandle()
-{
-#if !UE_BUILD_SHIPPING
-	if (IsValid())
-	{
-		if (FOnDemandContentInstallReplayRecorder* Recorder = FOnDemandContentInstallReplayRecorder::Get())
-		{
-			Recorder->RecordContentHandleDestroyed(Handle->HandleId());
-		}
-	}
-#endif // !UE_BUILD_SHIPPING
-}
-
-FOnDemandContentHandle::FOnDemandContentHandle(const FOnDemandContentHandle& Other)
-	: Handle(Other.Handle)
-{
-#if !UE_BUILD_SHIPPING
-	if (IsValid())
-	{
-		if (FOnDemandContentInstallReplayRecorder* Recorder = FOnDemandContentInstallReplayRecorder::Get())
-		{
-			Recorder->RecordContentHandle(Handle->HandleId(), FString()); // Only need to record debug name on first creation
-		}
-	}
-#endif // !UE_BUILD_SHIPPING
-}
-
-FOnDemandContentHandle& FOnDemandContentHandle::operator=(const FOnDemandContentHandle& Other)
-{
-	if (this == &Other)
-	{
-		return *this;
-	}
-
-#if !UE_BUILD_SHIPPING
-	FOnDemandContentInstallReplayRecorder* Recorder = FOnDemandContentInstallReplayRecorder::Get();
-	if (IsValid() && Recorder)
-	{
-		Recorder->RecordContentHandleDestroyed(Handle->HandleId());
-	}
-#endif // !UE_BUILD_SHIPPING
-
-	Handle = Other.Handle;
-
-#if !UE_BUILD_SHIPPING
-	if (IsValid() && Recorder)
-	{
-		Recorder->RecordContentHandle(Handle->HandleId(), FString()); // Only need to record debug name on first creation
-	}
-#endif // !UE_BUILD_SHIPPING
-
-	return *this;
-}
-
-FOnDemandContentHandle& FOnDemandContentHandle::operator=(FOnDemandContentHandle&& Other)
-{
-	if (this == &Other)
-	{
-		return *this;
-	}
-
-#if !UE_BUILD_SHIPPING
-	if (IsValid())
-	{
-		if (FOnDemandContentInstallReplayRecorder* Recorder = FOnDemandContentInstallReplayRecorder::Get())
-		{
-			Recorder->RecordContentHandleDestroyed(Handle->HandleId());
-		}
-	}
-#endif // !UE_BUILD_SHIPPING
-
-	Handle = MoveTemp(Other.Handle);
-
-	return *this;
-}
-
 FOnDemandContentHandle FOnDemandContentHandle::Create()
 {
 	FOnDemandContentHandle NewHandle;
 	NewHandle.Handle = MakeShared<FOnDemandInternalContentHandle, ESPMode::ThreadSafe>();
-
-#if !UE_BUILD_SHIPPING
-	FSharedString DebugName = NewHandle.Handle->DebugName;
-	if (FOnDemandContentInstallReplayRecorder* Recorder = FOnDemandContentInstallReplayRecorder::Get())
-	{
-		Recorder->RecordContentHandle(NewHandle.Handle->HandleId(), FString::ConstructFromPtrSize(*DebugName, DebugName.Len()));
-	}
-#endif // !UE_BUILD_SHIPPING
-
 	return NewHandle;
 }
 
@@ -156,14 +67,6 @@ FOnDemandContentHandle FOnDemandContentHandle::Create(FSharedString DebugName)
 {
 	FOnDemandContentHandle NewHandle;
 	NewHandle.Handle = MakeShared<FOnDemandInternalContentHandle, ESPMode::ThreadSafe>(DebugName);
-
-#if !UE_BUILD_SHIPPING
-	if (FOnDemandContentInstallReplayRecorder* Recorder = FOnDemandContentInstallReplayRecorder::Get())
-	{
-		Recorder->RecordContentHandle(NewHandle.Handle->HandleId(), FString::ConstructFromPtrSize(*DebugName, DebugName.Len()));
-	}
-#endif // !UE_BUILD_SHIPPING
-
 	return NewHandle;
 }
 
@@ -178,45 +81,6 @@ FString LexToString(const FOnDemandContentHandle& Handle)
 }
 
 ///////////////////////////////////////////////////////////////////////////////
-FOnDemandInstallRequest::FOnDemandInstallRequest()
-{
-}
-
-FOnDemandInstallRequest::FOnDemandInstallRequest(FWeakOnDemandIoStore InIoStore, FSharedInternalInstallRequest InternalRequest)
-	: IoStore(InIoStore)
-	, Request(InternalRequest)
-{
-	// Unfortunately, request creation is recorded by FOnDemandContentInstaller::FRequest::FInstall::FInstall()
-	// because some of the information needed for the replay event is type erased and we don't have access to it here.
-	// Luckily, this is a move only type, so there are no other places that RecordInstall() needs to be called.
-}
-
-FOnDemandInstallRequest::~FOnDemandInstallRequest()
-{
-#if !UE_BUILD_SHIPPING
-	if (Request)
-	{
-		if (FOnDemandContentInstallReplayRecorder* Recorder = FOnDemandContentInstallReplayRecorder::Get())
-		{
-			Recorder->RecordInstallDestroyed(UPTRINT(Request.Get()));
-		}
-	}
-#endif // !UE_BUILD_SHIPPING
-}
-
-FOnDemandInstallRequest::FOnDemandInstallRequest(FOnDemandInstallRequest&& Other)
-	: IoStore(MoveTemp(Other.IoStore))
-	, Request(MoveTemp(Other.Request))
-{
-}
-
-FOnDemandInstallRequest& FOnDemandInstallRequest::operator=(FOnDemandInstallRequest&& Other)
-{
-	Request = MoveTemp(Other.Request);
-	IoStore = MoveTemp(Other.IoStore);
-	return *this;
-}
-
 FOnDemandInstallRequest::EStatus FOnDemandInstallRequest::GetStatus() const
 {
 	return Request.IsValid() ? Request->Status.load(std::memory_order_relaxed) : EStatus::None; 
@@ -226,16 +90,9 @@ void FOnDemandInstallRequest::Cancel()
 {
 	if (IsPending())
 	{
-		if (TSharedPtr<IOnDemandIoStore, ESPMode::ThreadSafe> Pinned = IoStore.Pin(); Pinned.IsValid())
+		if (TSharedPtr<IOnDemandIoStore, ESPMode::ThreadSafe> Pinned = Request->IoStore.Pin(); Pinned.IsValid())
 		{
 			Pinned->CancelInstallRequest(Request);
-
-#if !UE_BUILD_SHIPPING
-			if (FOnDemandContentInstallReplayRecorder* Recorder = FOnDemandContentInstallReplayRecorder::Get())
-			{
-				Recorder->RecordCancel(UPTRINT(Request.Get()));
-			}
-#endif // !UE_BUILD_SHIPPING
 		}
 	}
 }
@@ -244,16 +101,9 @@ void FOnDemandInstallRequest::UpdatePriority(int32 NewPriority)
 {
 	if (IsPending())
 	{
-		if (TSharedPtr<IOnDemandIoStore, ESPMode::ThreadSafe> Pinned = IoStore.Pin(); Pinned.IsValid())
+		if (TSharedPtr<IOnDemandIoStore, ESPMode::ThreadSafe> Pinned = Request->IoStore.Pin(); Pinned.IsValid())
 		{
 			Pinned->UpdateInstallRequestPriority(Request, NewPriority);
-
-#if !UE_BUILD_SHIPPING
-			if (FOnDemandContentInstallReplayRecorder* Recorder = FOnDemandContentInstallReplayRecorder::Get())
-			{
-				Recorder->RecordUpdatePriority(UPTRINT(Request.Get()), NewPriority);
-			}
-#endif // !UE_BUILD_SHIPPING
 		}
 	}
 }
@@ -506,7 +356,8 @@ void FOnDemandContentInstallReplayRecorder::StartRecordingInternal()
 {
 	TUniqueLock Lock(Mutex);
 
-	Reset();
+	constexpr bool bRecorderValid = true;
+	Reset(bRecorderValid);
 	StartTime = FDateTime::UtcNow();
 }
 
@@ -529,47 +380,41 @@ TValueOrError<FOnDemandContentInstallReplay, void> FOnDemandContentInstallReplay
 
 void FOnDemandContentInstallReplayRecorder::RecordContentHandle(UPTRINT HandleId, FString&& DebugName)
 {
+	TUniqueLock Lock(Mutex);
+
 	if (ensure(HandleId) == false)
 	{
-		// still valid but this shouldn't happen.
-		return;
+		bValid = false;
 	}
 
-	TUniqueLock Lock(Mutex);
+	if (!IsValid())
+	{
+		return;
+	}
 
 	const int64 EventIndex = AddEvent(
 		FOnDemandContentInstallReplay::FContentHandleEventData{ MoveTemp(DebugName) });
 
-	FContentHandleRecord& Record = ContentHandleEventToIndex.FindOrAdd(HandleId, FContentHandleRecord{ EventIndex, 0 });
-	++Record.RefCount;
+	ContentHandleEventToIndex.Add(HandleId, EventIndex);
 }
 
 void FOnDemandContentInstallReplayRecorder::RecordContentHandleDestroyed(UPTRINT HandleId)
 {
 	TUniqueLock Lock(Mutex);
 
-	FContentHandleRecord* Record = ContentHandleEventToIndex.Find(HandleId);
-	if (ensure(Record) == false)
+	if (!IsValid())
 	{
-		bValid = false;
 		return;
 	}
 
-	if (ensure(Record->RefCount != 0) == false)
+	int64 ContentHandleEventIndex = -1;
+	if (ensure(ContentHandleEventToIndex.RemoveAndCopyValue(HandleId, ContentHandleEventIndex)) == false)
 	{
 		bValid = false;
 		return;
 	}
 
-	--Record->RefCount;
-
-	if (Record->RefCount == 0)
-	{
-		int64 ContentHandleEventIndex = Record->EventIndex;
-		ContentHandleEventToIndex.Remove(HandleId);
-
-		AddEvent(FOnDemandContentInstallReplay::FContentHandleDestroyedEventData{ ContentHandleEventIndex });
-	}
+	AddEvent(FOnDemandContentInstallReplay::FContentHandleDestroyedEventData{ ContentHandleEventIndex });
 }
 
 void FOnDemandContentInstallReplayRecorder::RecordInstall(
@@ -601,30 +446,24 @@ void FOnDemandContentInstallReplayRecorder::RecordInstall(
 	int32 Priority,
 	EOnDemandInstallOptions Options)
 {
+	TUniqueLock Lock(Mutex);
+
 	if (ensure(RequestId) == false)
 	{
-		// still valid but this shouldn't happen.
-		return;
+		bValid = false;
 	}
 
-	TUniqueLock Lock(Mutex);
-
-	int64 ContentHandleEventIndex = -1;
-	if (HandleId)
+	if (!IsValid())
 	{
-		// If the content handle is valid, it should exist in the lookup table
-		FContentHandleRecord* Record = ContentHandleEventToIndex.Find(HandleId);
-		if (ensure(Record) == false)
-		{
-			bValid = false;
-			return;
-		}
-
-		ContentHandleEventIndex = Record->EventIndex;
+		return;
 	}
 
-	// If the content handle is invalid, the install request will fail at runtime but the game
-	// still made the request, so we will record it.
+	const int64 ContentHandleEventIndex = ContentHandleEventToIndex.FindRef(HandleId, -1);
+	if (ensure(ContentHandleEventIndex >= 0) == false)
+	{
+		bValid = false;
+		return;
+	}
 
 	const int64 EventIndex = AddEvent(
 		FOnDemandContentInstallReplay::FInstallEventData
@@ -645,6 +484,11 @@ void FOnDemandContentInstallReplayRecorder::RecordCancel(UPTRINT RequestId)
 {
 	TUniqueLock Lock(Mutex);
 
+	if (!IsValid())
+	{
+		return;
+	}
+
 	const int64 InstallEventIndex = InstallEventToIndex.FindRef(RequestId, -1);
 	if (ensure(InstallEventIndex >= 0) == false)
 	{
@@ -659,6 +503,11 @@ void FOnDemandContentInstallReplayRecorder::RecordUpdatePriority(UPTRINT Request
 {
 	TUniqueLock Lock(Mutex);
 
+	if (!IsValid())
+	{
+		return;
+	}
+
 	const int64 InstallEventIndex = InstallEventToIndex.FindRef(RequestId, -1);
 	if (ensure(InstallEventIndex >= 0) == false)
 	{
@@ -677,6 +526,11 @@ void FOnDemandContentInstallReplayRecorder::RecordInstallDestroyed(UPTRINT Reque
 {
 	TUniqueLock Lock(Mutex);
 
+	if (!IsValid())
+	{
+		return;
+	}
+
 	int64 InstallEventIndex = -1;
 	if (ensure(InstallEventToIndex.RemoveAndCopyValue(RequestId, InstallEventIndex)) == false)
 	{
@@ -687,13 +541,13 @@ void FOnDemandContentInstallReplayRecorder::RecordInstallDestroyed(UPTRINT Reque
 	AddEvent(FOnDemandContentInstallReplay::FInstallDestroyedEventData{ InstallEventIndex });
 }
 
-void FOnDemandContentInstallReplayRecorder::Reset()
+void FOnDemandContentInstallReplayRecorder::Reset(bool bInValid /*= false*/)
 {
 	ContentHandleEventToIndex.Empty();
 	InstallEventToIndex.Empty();
 	Events.Empty();
 	StartTime = FDateTime();
-	bValid = true;
+	bValid = bInValid;
 }
 
 int64 FOnDemandContentInstallReplayRecorder::GetTimeOffset() const
diff --git a/Engine/Source/Runtime/Experimental/IoStore/OnDemandCore/Public/IO/IoStoreOnDemand.h b/Engine/Source/Runtime/Experimental/IoStore/OnDemandCore/Public/IO/IoStoreOnDemand.h
--- a/Engine/Source/Runtime/Experimental/IoStore/OnDemandCore/Public/IO/IoStoreOnDemand.h
+++ b/Engine/Source/Runtime/Experimental/IoStore/OnDemandCore/Public/IO/IoStoreOnDemand.h
@@ -88,15 +88,15 @@ class FOnDemandContentHandle
 {
 public:
 	/** Creates a new invalid content handle. */ 
-	UE_API FOnDemandContentHandle();
+	FOnDemandContentHandle() = default;
 	/** Destroy the handle and release any referenced content. */ 
-	UE_API ~FOnDemandContentHandle();
+	~FOnDemandContentHandle() = default;
 
-	UE_API FOnDemandContentHandle(const FOnDemandContentHandle& Other);
+	FOnDemandContentHandle(const FOnDemandContentHandle& Other) = default;
 	FOnDemandContentHandle(FOnDemandContentHandle&& Other) = default;
 
-	UE_API FOnDemandContentHandle& operator=(const FOnDemandContentHandle& Other);
-	UE_API FOnDemandContentHandle& operator=(FOnDemandContentHandle&& Other);
+	FOnDemandContentHandle& operator=(const FOnDemandContentHandle& Other) = default;
+	FOnDemandContentHandle& operator=(FOnDemandContentHandle&& Other) = default;
 
 	/** Destroy the handle and release any referenced content. */ 
 	void Reset() { Handle.Reset(); }
@@ -338,13 +338,13 @@ class FOnDemandInstallRequest final
 	friend class FOnDemandIoStore;
 public:
 	/** Creates an invalid install request. */
-	UE_API FOnDemandInstallRequest();
+	FOnDemandInstallRequest() = default;
 	/** Move constructor. */;
-	UE_API FOnDemandInstallRequest(FOnDemandInstallRequest&&);
+	FOnDemandInstallRequest(FOnDemandInstallRequest&&) = default;
 	/** Move assignable. */
-	UE_API FOnDemandInstallRequest& operator=(FOnDemandInstallRequest&&);
+	FOnDemandInstallRequest& operator=(FOnDemandInstallRequest&&) = default;
 	/** Destructor. */
-	UE_API ~FOnDemandInstallRequest();
+	~FOnDemandInstallRequest() = default;
 	/** Returns the current request status. */
 	UE_API virtual EStatus GetStatus() const override;
 	/** Cancel the install request. */
@@ -353,11 +353,12 @@ public:
 	UE_API void UpdatePriority(int32 NewPriority);
 
 private:
-	UE_API FOnDemandInstallRequest(FWeakOnDemandIoStore IoStore, FSharedInternalInstallRequest InternalRequest);
+	FOnDemandInstallRequest(FSharedInternalInstallRequest InternalRequest)
+		: Request(InternalRequest)
+	{}
 	FOnDemandInstallRequest(const FOnDemandInstallRequest&) = delete;
 	FOnDemandInstallRequest& operator=(const FOnDemandInstallRequest&) = delete;
 
-	FWeakOnDemandIoStore IoStore;	
 	FSharedInternalInstallRequest Request;
 };
 
