Dataflow : add new texture 2D terminal that can properly create atexture from an asset path

#Synced-CL 50024027

#Synced-CL 50024027

diff --git a/Engine/Source/Runtime/Experimental/Dataflow/Engine/Private/Dataflow/DataflowTextureAssetNodes.cpp b/Engine/Source/Runtime/Experimental/Dataflow/Engine/Private/Dataflow/DataflowTextureAssetNodes.cpp
--- a/Engine/Source/Runtime/Experimental/Dataflow/Engine/Private/Dataflow/DataflowTextureAssetNodes.cpp
+++ b/Engine/Source/Runtime/Experimental/Dataflow/Engine/Private/Dataflow/DataflowTextureAssetNodes.cpp
@@ -1,9 +1,13 @@
 // Copyright Epic Games, Inc. All Rights Reserved.
 
 #include "Dataflow/DataflowTextureAssetNodes.h"
+
 #include "Dataflow/DataflowNodeFactory.h"
 #include "Engine/Texture2D.h"
 
+#include "AssetRegistry/AssetRegistryModule.h"
+#include "Misc/PackageName.h"
+
 #include UE_INLINE_GENERATED_CPP_BY_NAME(DataflowTextureAssetNodes)
 
 
@@ -14,6 +18,7 @@ namespace UE::Dataflow
 	void RegisterTextureAssetNodes()
 	{
 		DATAFLOW_NODE_REGISTER_CREATION_FACTORY(FDataflowTextureTerminalNode);
+		DATAFLOW_NODE_REGISTER_CREATION_FACTORY(FDataflowTexture2DTerminalNode);
 		DATAFLOW_NODE_REGISTER_CREATION_FACTORY(FDataflowGetTextureAssetNode);
 		DATAFLOW_NODE_REGISTER_CREATION_FACTORY(FDataflowGetTextureAssetFromPathNode);
 		DATAFLOW_NODE_REGISTER_CREATION_FACTORY(FDataflowTextureToImageNode);
@@ -65,7 +70,6 @@ void FDataflowTextureTerminalNode::Evaluate(UE::Dataflow::FContext& Context) con
 	SetValue(Context, OutTexture, &TextureAsset);
 }
 
-
 void FDataflowTextureTerminalNode::SetAssetValue(TObjectPtr<UObject> Asset, UE::Dataflow::FContext& Context) const
 {
 	UTexture2D* AssetToSet = Cast<UTexture2D>(Asset.Get());
@@ -87,6 +91,89 @@ void FDataflowTextureTerminalNode::SetAssetValue(TObjectPtr<UObject> Asset, UE::
 
 ////////////////////////////////////////////////////////////////////////////////////////////
 
+FDataflowTexture2DTerminalNode::FDataflowTexture2DTerminalNode(const UE::Dataflow::FNodeParameters& InParam, FGuid InGuid)
+	: FDataflowTerminalNode(InParam, InGuid)
+{
+	RegisterInputConnection(&Image);
+	RegisterInputConnection(&AssetPath);
+	RegisterOutputConnection(&TextureAsset);
+}
+
+void FDataflowTexture2DTerminalNode::Evaluate(UE::Dataflow::FContext& Context) const
+{
+	// Asset is created and stored by the SetAssetValue
+	// Evaluate only return an existing one or nullptr
+	const FString& InAssetPath = GetValue(Context, &AssetPath);
+	TObjectPtr<UTexture2D> OutTexture = ::LoadObject<UTexture2D>(nullptr, InAssetPath);
+	SetValue(Context, OutTexture, &TextureAsset);
+}
+
+void FDataflowTexture2DTerminalNode::SetAssetValue(TObjectPtr<UObject> Asset, UE::Dataflow::FContext& Context) const
+{
+	TObjectPtr<UTexture2D> NullTexturePtr = nullptr;
+
+#if WITH_EDITOR
+	const FString& InAssetPath = GetValue(Context, &AssetPath);
+	if (!FPackageName::IsValidObjectPath(InAssetPath))
+	{
+		Context.Error(FString::Format(TEXT("Asset Path input is not a valid path : {0}"), { InAssetPath }), this);
+		SetValue(Context, NullTexturePtr, &TextureAsset);
+		return;
+	}
+	const FString PackageName = FPackageName::ObjectPathToPackageName(InAssetPath);
+
+	UPackage* Package = Cast<UPackage>(FindPackage(nullptr, *PackageName));
+	if (Package == nullptr)
+	{
+		Package = CreatePackage(*PackageName);
+	}
+	if (Package == nullptr)
+	{
+		Context.Error(FString::Format(TEXT("Failed to find or create package {0}"), { PackageName }), this);
+		SetValue(Context, NullTexturePtr, &TextureAsset);
+		return;
+	}
+
+	const FName AssetName = FName(FPackageName::GetLongPackageAssetName(PackageName));
+	UObject* ExistingObject = StaticFindObjectFastInternal( /*Class=*/ NULL, Package, AssetName, EFindObjectFlags::ExactClass);
+	if (ExistingObject && !ExistingObject->GetClass()->IsChildOf<UTexture2D>())
+	{
+		Context.Error(FString::Format(TEXT("Asset {0} already exists but is not a Texture2D compatible type"), { InAssetPath }), this);
+		SetValue(Context, NullTexturePtr, &TextureAsset);
+		return;
+	}
+	// It is an error if we are trying to replace an object of a different class
+	TObjectPtr<UTexture2D> OutTexture2D = Cast<UTexture2D>(ExistingObject);
+	if (!OutTexture2D)
+	{
+		OutTexture2D = NewObject<UTexture2D>(Package, UTexture2D::StaticClass(), AssetName, RF_Public | RF_Standalone | RF_Transactional);
+		if (!OutTexture2D)
+		{
+			Context.Error(FString::Format(TEXT("Failed to create texture asset {0}"), { InAssetPath }), this);
+			SetValue(Context, NullTexturePtr, &TextureAsset);
+			return;
+		}
+	}
+
+	// make sure the asset is set properly
+	OutTexture2D->MarkPackageDirty();
+	FAssetRegistryModule::AssetCreated(OutTexture2D.Get());
+
+	// update with image
+	const FDataflowImage& InImage = GetValue(Context, &Image);
+	DataflowTextureAssetNodes::Private::UpdateTexture2DFromImage(*OutTexture2D, InImage.GetImage());
+
+	// finally set the texture to the output
+	SetValue(Context, OutTexture2D, &TextureAsset);
+
+#else
+	Context.Error(TEXT("Creation of texture asset only supported in Editor"), this);
+	SetValue(Context, NullTexturePtr, &TextureAsset);
+#endif
+}
+
+////////////////////////////////////////////////////////////////////////////////////////////
+
 FDataflowGetTextureAssetNode::FDataflowGetTextureAssetNode(const UE::Dataflow::FNodeParameters& InParam, FGuid InGuid)
 	: FDataflowNode(InParam, InGuid)
 {
diff --git a/Engine/Source/Runtime/Experimental/Dataflow/Engine/Public/Dataflow/DataflowTextureAssetNodes.h b/Engine/Source/Runtime/Experimental/Dataflow/Engine/Public/Dataflow/DataflowTextureAssetNodes.h
--- a/Engine/Source/Runtime/Experimental/Dataflow/Engine/Public/Dataflow/DataflowTextureAssetNodes.h
+++ b/Engine/Source/Runtime/Experimental/Dataflow/Engine/Public/Dataflow/DataflowTextureAssetNodes.h
@@ -32,6 +32,36 @@ public:
 	FDataflowTextureTerminalNode(const UE::Dataflow::FNodeParameters& InParam, FGuid InGuid = FGuid::NewGuid());
 };
 
+/*
+* Terminal node to a save a dependent 2D texture
+* The terminal node will create an asset if it does not exists yet 
+*/
+USTRUCT(meta = (DataflowTerminal))
+struct FDataflowTexture2DTerminalNode : public FDataflowTerminalNode
+{
+	GENERATED_USTRUCT_BODY()
+	DATAFLOW_NODE_DEFINE_INTERNAL(FDataflowTexture2DTerminalNode, "Texture2DTerminal", "Terminal", "Path Asset Image")
+
+private:
+	/** Image data to store in the texture */
+	UPROPERTY(meta = (DataflowInput))
+	FDataflowImage Image;
+
+	/** Path to the texture asset to create or update */
+	UPROPERTY(EditAnywhere, Category = Asset, meta = (DataflowInput))
+	FString AssetPath;
+
+	/** Texture asset matching AssetPath and updated with the image data */
+	UPROPERTY(meta = (DataflowOutput))
+	TObjectPtr<UTexture2D> TextureAsset = nullptr;
+
+	virtual void Evaluate(UE::Dataflow::FContext& Context) const override;
+	virtual void SetAssetValue(TObjectPtr<UObject> Asset, UE::Dataflow::FContext& Context) const override;
+
+public:
+	FDataflowTexture2DTerminalNode(const UE::Dataflow::FNodeParameters& InParam, FGuid InGuid = FGuid::NewGuid());
+};
+
 /*
 * Get a texture asset
 */
