Changed from using SeekFromEnd to Seek when reading on-demand container TOC file(s)

#rnx

#Synced-CL 49940788

diff --git a/Engine/Source/Runtime/Experimental/IoStore/OnDemandCore/Private/Serialization/OnDemandContainerToc.cpp b/Engine/Source/Runtime/Experimental/IoStore/OnDemandCore/Private/Serialization/OnDemandContainerToc.cpp
--- a/Engine/Source/Runtime/Experimental/IoStore/OnDemandCore/Private/Serialization/OnDemandContainerToc.cpp
+++ b/Engine/Source/Runtime/Experimental/IoStore/OnDemandCore/Private/Serialization/OnDemandContainerToc.cpp
@@ -488,12 +488,6 @@ TIoStatusOr<FOnDemandTocReader> FOnDemandTocReader::Read(const FString& Filename
 {
 	IPlatformFile& Ipf = FPlatformFileManager::Get().GetPlatformFile();
 
-	const int64 FileSize = Ipf.FileSize(*Filename);
-	if (FileSize < int64(sizeof(FOnDemandTocHeader)))
-	{
-		return FIoStatus(EIoErrorCode::FileOpenFailed);
-	}
-
 	FFileOpenResult OpenResult = Ipf.OpenRead(*Filename, IPlatformFile::EOpenReadFlags::None);
 	if (OpenResult.HasError())
 	{
@@ -525,11 +519,11 @@ TIoStatusOr<FOnDemandTocReader> FOnDemandTocReader::Read(const FString& Filename
 	}
 
 	const int64 TotalHeaderSize =
-		sizeof(FOnDemandTocHeader) +
-		(sizeof(UTF8CHAR) * Header.StringTableLen) +
-		(sizeof(FOnDemandContainerEntry) * Header.ContainerCount);
+		int64(sizeof(FOnDemandTocHeader)) +
+		int64((sizeof(UTF8CHAR) * Header.StringTableLen)) +
+		int64((sizeof(FOnDemandContainerEntry) * Header.ContainerCount));
 
-	if (TotalHeaderSize >= FileSize)
+	if ((TotalHeaderSize < int64(sizeof(FOnDemandTocHeader))) || (TotalHeaderSize >= FileHandle->Size()))
 	{
 		return FIoStatusBuilder(EIoErrorCode::CorruptToc)
 			<< TEXT("TOC header size greater than file size in '") << Filename << TEXT("'");
@@ -546,7 +540,7 @@ TIoStatusOr<FOnDemandTocReader> FOnDemandTocReader::Read(const FString& Filename
 			<< TEXT("Failed to read container TOC header from '") << Filename << TEXT("'");
 	}
 
-	if (FileHandle->SeekFromEnd(-int64(sizeof(FOnDemandTocFooter))) == false)
+	if (FileHandle->Seek(FileHandle->Size() - int64(sizeof(FOnDemandTocFooter))) == false)
 	{
 		return FIoStatusBuilder(EIoErrorCode::FileSeekFailed)
 			<< TEXT("Failed to seek to container TOC footer in '") << Filename << TEXT("'");
