[DataFlow]

Removes delegate registration from FDataflowInstance constructor since it is not safe to manipulate delegates during a broadcast, which is likely when objects are constructed as part of loading on multiple threads. Delegate registration happens during object construction when not loading and otherwise expects the loading object to invoke a new PostLoad method on the FDataflowInstance USTRUCT.

#rnx

#Synced-CL 50155972

diff --git a/Engine/Plugins/ChaosClothAsset/Source/ChaosClothAssetEngine/Private/ChaosClothAsset/ClothAssetBase.cpp b/Engine/Plugins/ChaosClothAsset/Source/ChaosClothAssetEngine/Private/ChaosClothAsset/ClothAssetBase.cpp
--- a/Engine/Plugins/ChaosClothAsset/Source/ChaosClothAssetEngine/Private/ChaosClothAsset/ClothAssetBase.cpp
+++ b/Engine/Plugins/ChaosClothAsset/Source/ChaosClothAssetEngine/Private/ChaosClothAsset/ClothAssetBase.cpp
@@ -367,6 +367,8 @@ void UChaosClothAssetBase::PostLoad()
 {
 	Super::PostLoad();
 
+	GetDataflowInstance().PostLoad();
+
 #if WITH_EDITORONLY_DATA
 	bHasDataflowAsset = (GetDataflowInstance().GetDataflowAsset() != nullptr);
 
diff --git a/Engine/Plugins/Experimental/ChaosFlesh/Source/ChaosFleshEngine/Private/ChaosFlesh/FleshAsset.cpp b/Engine/Plugins/Experimental/ChaosFlesh/Source/ChaosFleshEngine/Private/ChaosFlesh/FleshAsset.cpp
--- a/Engine/Plugins/Experimental/ChaosFlesh/Source/ChaosFleshEngine/Private/ChaosFlesh/FleshAsset.cpp
+++ b/Engine/Plugins/Experimental/ChaosFlesh/Source/ChaosFleshEngine/Private/ChaosFlesh/FleshAsset.cpp
@@ -619,6 +619,8 @@ void UFleshAsset::PostLoad()
 {
 	Super::PostLoad();
 
+	DataflowInstance.PostLoad();
+
 	// migrate deprecated data if necessary 
 	MigrateDeprecatedDataflowData();
 }
diff --git a/Engine/Plugins/Experimental/ChaosRigidAsset/Source/ChaosRigidAssetEngine/Private/DataflowAttachment.cpp b/Engine/Plugins/Experimental/ChaosRigidAsset/Source/ChaosRigidAssetEngine/Private/DataflowAttachment.cpp
--- a/Engine/Plugins/Experimental/ChaosRigidAsset/Source/ChaosRigidAssetEngine/Private/DataflowAttachment.cpp
+++ b/Engine/Plugins/Experimental/ChaosRigidAsset/Source/ChaosRigidAssetEngine/Private/DataflowAttachment.cpp
@@ -34,6 +34,13 @@ FDataflowInstance& UDataflowAttachment::GetDataflowInstance()
 	return Instance;
 }
 
+void UDataflowAttachment::PostLoad()
+{
+	Super::PostLoad();
+
+	Instance.PostLoad();
+}
+
 TObjectPtr<UDataflowBaseContent> UDataflowAttachment::CreateDataflowContent()
 {
 	TObjectPtr<UDataflowBaseContent> BaseContent = NewObject<UPhysicsAssetDataflowContent>();
diff --git a/Engine/Plugins/Experimental/ChaosRigidAsset/Source/ChaosRigidAssetEngine/Public/DataflowAttachment.h b/Engine/Plugins/Experimental/ChaosRigidAsset/Source/ChaosRigidAssetEngine/Public/DataflowAttachment.h
--- a/Engine/Plugins/Experimental/ChaosRigidAsset/Source/ChaosRigidAssetEngine/Public/DataflowAttachment.h
+++ b/Engine/Plugins/Experimental/ChaosRigidAsset/Source/ChaosRigidAssetEngine/Public/DataflowAttachment.h
@@ -33,6 +33,9 @@ public:
 	FDataflowInstance& GetDataflowInstance() override;
 	////////////////////////////////////////////////////////////////////////////////////////////////////////////////
 
+	// UObject
+	virtual void PostLoad() override;
+	////////////////////////////////////////////////////////////////////////////////////////////////////////////////
 private:
 
 	// Dataflow instance for the attachment when editing.
diff --git a/Engine/Plugins/Runtime/HairStrands/Source/HairStrandsCore/Private/GroomAsset.cpp b/Engine/Plugins/Runtime/HairStrands/Source/HairStrandsCore/Private/GroomAsset.cpp
--- a/Engine/Plugins/Runtime/HairStrands/Source/HairStrandsCore/Private/GroomAsset.cpp
+++ b/Engine/Plugins/Runtime/HairStrands/Source/HairStrandsCore/Private/GroomAsset.cpp
@@ -1331,6 +1331,8 @@ void UGroomAsset::PostLoad()
 
 	Super::PostLoad();
 
+	DataflowSettings.PostLoad();
+
 	// Early out if hair strands are disabled
 	const bool bIsHairEnabled = IsGroomEnabled();
 	if (!bIsHairEnabled)
diff --git a/Engine/Source/Runtime/Experimental/Dataflow/Engine/Private/Dataflow/DataflowInstance.cpp b/Engine/Source/Runtime/Experimental/Dataflow/Engine/Private/Dataflow/DataflowInstance.cpp
--- a/Engine/Source/Runtime/Experimental/Dataflow/Engine/Private/Dataflow/DataflowInstance.cpp
+++ b/Engine/Source/Runtime/Experimental/Dataflow/Engine/Private/Dataflow/DataflowInstance.cpp
@@ -490,10 +490,35 @@ FDataflowInstance::FDataflowInstance(UObject* InOwner, UDataflow* InDataflowAsse
 	, Owner(InOwner)
 {
 #if WITH_EDITOR
+	// If we are in the process of loading, we can't be certain that we aren't already inside one of these delegate's broadcasts
+	if (!IsInAsyncLoadingThread())
+	{
+		BindEditorDelegates();
+	}
+#endif
+}
+
+FDataflowInstance::~FDataflowInstance()
+{
+#if WITH_EDITOR
+	UnbindEditorDelegates();
+#endif
+}
+
+void FDataflowInstance::PostLoad()
+{
+#if WITH_EDITOR
+	BindEditorDelegates();
+#endif
+}
+
+#if WITH_EDITOR
+void FDataflowInstance::BindEditorDelegates()
+{
 	// let listen to property change on the owner so we can listen to our own changes ( because we are a struct )
 	if (Owner)
 	{
-		FDataflowAssetDelegates::OnVariablesChanged.AddWeakLambda(Owner.Get(), 
+		FDataflowAssetDelegates::OnVariablesChanged.AddWeakLambda(Owner.Get(),
 			[this](const UDataflow* InDataflowAsset, FName InVariableName)
 			{
 				if (InDataflowAsset == DataflowAsset)
@@ -503,22 +528,18 @@ FDataflowInstance::FDataflowInstance(UObject* InOwner, UDataflow* InDataflowAsse
 			});
 		OnOwnerPropertyChangedHandle = FCoreUObjectDelegates::OnObjectPropertyChanged.AddRaw(this, &FDataflowInstance::OnOwnerPostEditChangeProperty);
 	}
-#endif
 }
 
-FDataflowInstance::~FDataflowInstance()
+void FDataflowInstance::UnbindEditorDelegates()
 {
-#if WITH_EDITOR
 	if (Owner)
 	{
 		FDataflowAssetDelegates::OnVariablesChanged.RemoveAll(Owner.Get());
 	}
 	// make sure we remove this outside of the Owner check in case of the Owner changed value since construction
 	FCoreUObjectDelegates::OnObjectPropertyChanged.Remove(OnOwnerPropertyChangedHandle);
-#endif
 }
 
-#if WITH_EDITOR
 void FDataflowInstance::OnOwnerPostEditChangeProperty(UObject* InObject, struct FPropertyChangedEvent& InPropertyChangedEvent)
 {
 	if (IsValid(Owner) && IsValid(InObject) && Owner == InObject)
diff --git a/Engine/Source/Runtime/Experimental/Dataflow/Engine/Public/Dataflow/DataflowInstance.h b/Engine/Source/Runtime/Experimental/Dataflow/Engine/Public/Dataflow/DataflowInstance.h
--- a/Engine/Source/Runtime/Experimental/Dataflow/Engine/Public/Dataflow/DataflowInstance.h
+++ b/Engine/Source/Runtime/Experimental/Dataflow/Engine/Public/Dataflow/DataflowInstance.h
@@ -259,6 +259,9 @@ public:
 	*/
 	DATAFLOWENGINE_API bool UpdateOwnerAsset(bool bUpdateDependentAssets = false) const;
 
+	/** Call as part of the PostLoad of any UObject containing a FDataflowInstance member */
+	DATAFLOWENGINE_API void PostLoad();
+
 #if WITH_EDITOR
 	DATAFLOWENGINE_API TSharedPtr<FStructOnScope> MakeStructOnScope() const;
 #endif 
@@ -286,6 +289,8 @@ private:
 	friend class UE::Dataflow::Private::FVariablesOverridesDetails;
 
 #if WITH_EDITOR
+	void BindEditorDelegates();
+	void UnbindEditorDelegates();
 	void OnOwnerPostEditChangeProperty(UObject* InObject, struct FPropertyChangedEvent& InPropertyChangedEvent);
 
 	FDelegateHandle OnOwnerPropertyChangedHandle;
diff --git a/Engine/Source/Runtime/Experimental/GeometryCollectionEngine/Private/GeometryCollection/GeometryCollectionObject.cpp b/Engine/Source/Runtime/Experimental/GeometryCollectionEngine/Private/GeometryCollection/GeometryCollectionObject.cpp
--- a/Engine/Source/Runtime/Experimental/GeometryCollectionEngine/Private/GeometryCollection/GeometryCollectionObject.cpp
+++ b/Engine/Source/Runtime/Experimental/GeometryCollectionEngine/Private/GeometryCollection/GeometryCollectionObject.cpp
@@ -2119,6 +2119,8 @@ void UGeometryCollection::PostLoad()
 {
 	Super::PostLoad();
 
+	DataflowInstance.PostLoad();
+
 #if WITH_EDITOR
 	constexpr bool bInitResources = true;
 	EnsureRenderDataIsCooked(bInitResources);
