Chaos: Use transactionally safe lock inside TRigidParticles

#Synced-CL 50118602

diff --git a/Engine/Source/Runtime/Experimental/Chaos/Public/Chaos/RigidParticles.h b/Engine/Source/Runtime/Experimental/Chaos/Public/Chaos/RigidParticles.h
--- a/Engine/Source/Runtime/Experimental/Chaos/Public/Chaos/RigidParticles.h
+++ b/Engine/Source/Runtime/Experimental/Chaos/Public/Chaos/RigidParticles.h
@@ -14,6 +14,7 @@
 #include "Chaos/RigidParticleControlFlags.h"
 #include "HAL/LowLevelMemTracker.h"
 #include "Chaos/Evolution/IterationSettings.h"
+#include "Misc/TransactionallySafeRWLock.h"
 #include "UObject/FortniteMainBranchObjectVersion.h"
 
 namespace Chaos
@@ -270,7 +271,7 @@ public:
 		MSleepData.Empty();
 		SleepDataLock.WriteUnlock();
 	}
-	FORCEINLINE FRWLock& GetSleepDataLock() { return SleepDataLock; }
+	FORCEINLINE FTransactionallySafeRWLock& GetSleepDataLock() { return SleepDataLock; }
 
 	FORCEINLINE const EObjectStateType ObjectState(const int32 Index) const { return CoreData[Index].ObjectState; }
 	FORCEINLINE EObjectStateType& ObjectState(const int32 Index) { return CoreData[Index].ObjectState; }
@@ -431,7 +432,7 @@ private:
 	TArrayCollectionArray<Private::FIterationSettings> MParticleIterationCounts;
 
 	TArray<TSleepData<T, d>> MSleepData;
-	FRWLock SleepDataLock;
+	FTransactionallySafeRWLock SleepDataLock;
 };
 
 
