Support for resolving and downloading I/O store shader map chunks on-demand.

#rnx

#Synced-CL 50197459

diff --git a/Engine/Source/Runtime/Core/Internal/Misc/CoreDelegatesInternal.h b/Engine/Source/Runtime/Core/Internal/Misc/CoreDelegatesInternal.h
--- a/Engine/Source/Runtime/Core/Internal/Misc/CoreDelegatesInternal.h
+++ b/Engine/Source/Runtime/Core/Internal/Misc/CoreDelegatesInternal.h
@@ -41,4 +41,16 @@ struct FCoreInternalDelegates
 	static CORE_API FMountOperationPak& GetOnPakMountOperation();
 
 	static CORE_API FCurrentlyMountedPaksDelegate& GetCurrentlyMountedPaksDelegate();
+
+	/** Callback triggered when a shader map hash has been resolved to an I/O store chunk ID. */
+	using FIoChunkIdResolvedFunc = TFunctionRef<void(const class FIoChunkId&)>;
+
+	/**
+	 * Callback to resolve I/O store chunk IDs from shader code hashes.
+	 * Used when installing content on-demand.
+	 *
+	 * @param	Hashes		Array of hashes that should be resolved.
+	 * @param	Callback	Callback triggered for reach resolved chunk ID.
+	 */
+	static CORE_API TDelegate<void(TArrayView<const FSHAHash>, FIoChunkIdResolvedFunc)> ResolvePackageShaderMaps;
 };
diff --git a/Engine/Source/Runtime/Core/Private/Misc/CoreDelegates.cpp b/Engine/Source/Runtime/Core/Private/Misc/CoreDelegates.cpp
--- a/Engine/Source/Runtime/Core/Private/Misc/CoreDelegates.cpp
+++ b/Engine/Source/Runtime/Core/Private/Misc/CoreDelegates.cpp
@@ -79,6 +79,8 @@ FCurrentlyMountedPaksDelegate& FCoreInternalDelegates::GetCurrentlyMountedPaksDe
 	return Delegate;
 }
 
+TDelegate<void(TArrayView<const FSHAHash>, FCoreInternalDelegates::FIoChunkIdResolvedFunc)> FCoreInternalDelegates::ResolvePackageShaderMaps;
+
 //} // namespace UE::Core
 
 TMulticastDelegate<void(bool, int32, int32)> FCoreDelegates::OnUserLoginChangedEvent;
@@ -396,4 +398,4 @@ FSimpleMulticastDelegate FCoreDelegates::OnChildEndFramePostFork;
 FCoreDelegates::FOnDisallowedExecCommandCalled FCoreDelegates::OnDisallowedExecCommandCalled;
 FCoreDelegates::FOnNamedCommandParsed FCoreDelegates::OnNamedCommandParsed;
 
-FSimpleMulticastDelegate FCoreDelegates::OnLowLevelMemTrackerTick;
\ No newline at end of file
+FSimpleMulticastDelegate FCoreDelegates::OnLowLevelMemTrackerTick;
diff --git a/Engine/Source/Runtime/Experimental/IoStore/OnDemand/Private/OnDemandContentInstaller.cpp b/Engine/Source/Runtime/Experimental/IoStore/OnDemand/Private/OnDemandContentInstaller.cpp
--- a/Engine/Source/Runtime/Experimental/IoStore/OnDemand/Private/OnDemandContentInstaller.cpp
+++ b/Engine/Source/Runtime/Experimental/IoStore/OnDemand/Private/OnDemandContentInstaller.cpp
@@ -14,6 +14,7 @@
 #include "IO/IoStoreOnDemand.h"
 #include "IO/OnDemandError.h"
 #include "Misc/App.h"
+#include "Misc/CoreDelegatesInternal.h"
 #include "Misc/Timespan.h"
 #include "OnDemandContentInstallReplay.h"
 #include "OnDemandInstallCache.h"
@@ -188,82 +189,68 @@ void VisitPackageDependencies(
 		}
 	}
 }
-
 ////////////////////////////////////////////////////////////////////////////////
-void ResolvePackageDependencies(
-	const TSet<FPackageId>& PackageIds,
-	bool bIncludeSoftReferences,
-	TSet<FPackageId>& OutResolved,
-	TSet<FPackageId>& OutMissing,
-	FPackageStoreReadScope* InReadScope /*= nullptr*/)
+static bool IsPackageStoreEntryMissing(EPackageStoreEntryStatus EntryStatus)
 {
-	TRACE_CPUPROFILER_EVENT_SCOPE(FOnDemandContentInstaller::ResolvePackageDependencies);
-
-	auto Visitor = [&OutResolved, &OutMissing](
-		FPackageId PackageId,
-		EPackageStoreEntryStatus EntryStatus,
-		const FPackageStoreEntry& PackageStoreEntry,
-		const TSharedPtr<IPackageStoreBackend>& Backend) -> bool
-	{
 #if !UE_BUILD_SHIPPING
-		if (UE::IoStore::CVars::IoStoreChunkMissingErrorChance > 0.0f)
-		{
-			if (FMath::RandRange(0.0f, 100.0f) < UE::IoStore::CVars::IoStoreChunkMissingErrorChance)
-			{
-				EntryStatus = EPackageStoreEntryStatus::Missing;
-			}
-		}
+	const float ErrorChance = UE::IoStore::CVars::IoStoreChunkMissingErrorChance;
+	if (ErrorChance > 0.0f && (FMath::RandRange(0.0f, 100.0f) < ErrorChance))
+	{
+		EntryStatus = EPackageStoreEntryStatus::Missing;
+	}
 #endif
-		if (EntryStatus != EPackageStoreEntryStatus::Missing)
-		{
-			OutResolved.Add(PackageId);
-		}
-		else
-		{
-			OutMissing.Add(PackageId);
-		}
-
-		return true;
-	};
-
-	VisitPackageDependencies(PackageIds, bIncludeSoftReferences, Visitor, InReadScope);
+	return EntryStatus == EPackageStoreEntryStatus::Missing;
 }
 
 ////////////////////////////////////////////////////////////////////////////////
 void ResolveChunksToInstall(
+	const TSharedPtr<IOnDemandPackageStoreBackend>& OnDemandPackageStore,
 	const TSet<FSharedOnDemandContainer>& Containers,
 	const TSet<FPackageId>& PackageIds,
 	bool bIncludeSoftReferences,
 	bool bIncludeOptionalBulkData,
 	TArray<FResolvedContainerChunks>& OutResolvedContainerChunks,
-	TSet<FPackageId>& OutMissing)
+	TSet<FIoChunkId>& OutMissing)
 {
 	TRACE_CPUPROFILER_EVENT_SCOPE(FOnDemandContentInstaller::ResolveChunksToInstall);
-
-	// For now we always download these required chunks
-	// TODO: only add these if we ref packages from that container
+	
 	for (const FSharedOnDemandContainer& Container : Containers)
 	{
 		FResolvedContainerChunks& ResolvedChunks = OutResolvedContainerChunks.AddDefaulted_GetRef();
 		ResolvedChunks.Container = Container;
+	}
 
-		for (int32 EntryIndex = 0; const FIoChunkId& ChunkId : Container->ChunkIds)
-		{
-			switch(ChunkId.GetChunkType())
+	TSet<FPackageId> ResolvedPackageIds;
+	TSet<FIoChunkId> ShaderMapChunkIds;
+	VisitPackageDependencies(
+		PackageIds,
+		bIncludeSoftReferences,
+		[&ResolvedPackageIds, &ShaderMapChunkIds, &OutMissing, &OnDemandPackageStore](
+			FPackageId PackageId,
+			EPackageStoreEntryStatus EntryStatus,
+			const FPackageStoreEntry& PackageStoreEntry,
+			const TSharedPtr<IPackageStoreBackend>& Backend)
 			{
-				case EIoChunkType::ExternalFile:
-				case EIoChunkType::ShaderCodeLibrary:
-				case EIoChunkType::ShaderCode:
+				if (IsPackageStoreEntryMissing(EntryStatus))
 				{
-					ResolvedChunks.EntryIndices.Emplace(EntryIndex); 
-					ResolvedChunks.TotalSize += Container->ChunkEntries[EntryIndex].GetDiskSize();
+					OutMissing.Add(CreatePackageDataChunkId(PackageId));
+				}
+				else if (Backend.Get() == OnDemandPackageStore.Get())
+				{
+					ResolvedPackageIds.Add(PackageId);
+					if (PackageStoreEntry.ShaderMapHashes.IsEmpty() == false)
+					{
+						FCoreInternalDelegates::ResolvePackageShaderMaps.ExecuteIfBound(
+							PackageStoreEntry.ShaderMapHashes,
+							[&ShaderMapChunkIds](const FIoChunkId& ChunkId)
+							{
+								ShaderMapChunkIds.Add(ChunkId);
+							});
+					}
 				}
-				default:
-					break;
-			}
-			++EntryIndex;
-		}
-	}
+
+				return true;
+			});
 
 	auto FindChunkEntry = [&OutResolvedContainerChunks](const FIoChunkId& ChunkId, int32& OutIndex) -> FResolvedContainerChunks*
 	{
@@ -277,28 +264,38 @@ void ResolveChunksToInstall(
 		return nullptr; 
 	};
 
-	TSet<FPackageId> ResolvedPackageIds;
-	ResolvePackageDependencies(
-		PackageIds,
-		bIncludeSoftReferences,
-		ResolvedPackageIds,
-		OutMissing);
+	// Resolve shader map chunks
+	for (const FIoChunkId& ChunkId : ShaderMapChunkIds)
+	{
+		int32 EntryIndex = INDEX_NONE;
+		if (FResolvedContainerChunks* ResolvedChunks = FindChunkEntry(ChunkId, EntryIndex))
+		{
+			check(EntryIndex != INDEX_NONE);
+			FOnDemandContainer& Container = *ResolvedChunks->Container;
+			ResolvedChunks->EntryIndices.Emplace(EntryIndex); 
+			ResolvedChunks->TotalSize += Container.ChunkEntries[EntryIndex].GetDiskSize();
+		}
+		else
+		{
+			OutMissing.Add(ChunkId);
+		}
+	}
 
 	// Resolve all chunk entries from the resolved package ID's
 	for (const FPackageId& PackageId : ResolvedPackageIds)
 	{
 		const FIoChunkId PackageChunkId				= CreatePackageDataChunkId(PackageId);
-		int32 EntryIndex							= INDEX_NONE; 
+		int32 EntryIndex							= INDEX_NONE;
 		FResolvedContainerChunks* ResolvedChunks	= FindChunkEntry(PackageChunkId, EntryIndex);
 
-		if (ResolvedChunks == nullptr) 
+		if (ResolvedChunks == nullptr)
 		{
-			// The chunk resides in a base game container
+			UE_LOG(LogIoStoreOnDemand, Warning, TEXT("Failed to resolve I/O store chunk for package ID '%s'"), *LexToString(PackageId));
 			continue;
 		}
 
 		check(EntryIndex != INDEX_NONE);
-		FOnDemandContainer& Container	= *ResolvedChunks->Container; 
+		FOnDemandContainer& Container = *ResolvedChunks->Container;
 
 		ResolvedChunks->EntryIndices.Emplace(EntryIndex); 
 		ResolvedChunks->TotalSize += Container.ChunkEntries[EntryIndex].GetDiskSize();
@@ -866,10 +863,11 @@ void FOnDemandContentInstaller::ProcessInstallRequest(FRequest& Request)
 		return CompleteInstallRequest(Request);
 	}
 
-	TSet<FPackageId> Missing;
+	TSet<FIoChunkId> Missing;
 	const bool bIncludeSoftReferences	= EnumHasAnyFlags(InstallRequest.Args.Options, EOnDemandInstallOptions::InstallSoftReferences);
 	const bool bIncludeOptionalBulkData = EnumHasAnyFlags(InstallRequest.Args.Options, EOnDemandInstallOptions::InstallOptionalBulkData);
 	Private::ResolveChunksToInstall(
+		IoStore.PackageStoreBackend,
 		ContainersForInstallation,
 		PackageIdsToInstall,
 		bIncludeSoftReferences,
@@ -882,15 +880,14 @@ void FOnDemandContentInstaller::ProcessInstallRequest(FRequest& Request)
 	{
 		FIoDispatcher& IoDispatcher = FIoDispatcher::Get();
 		TArray<FIoChunkId> MissingChunkIds;
-		for (const FPackageId& PackageId : Missing)
+		for (const FIoChunkId& ChunkId : Missing)
 		{
-			const FIoChunkId ChunkId = CreatePackageDataChunkId(PackageId);
 			if (IoDispatcher.DoesChunkExist(ChunkId) == false)
 			{
 				UE_CLOG(MissingChunkIds.Num() == 0, LogIoStoreOnDemand, Warning, TEXT("Failed to resolve the following chunk(s) for content handle '%s':"),
 					*LexToString(InstallRequest.Args.ContentHandle));
 
-				UE_LOG(LogIoStoreOnDemand, Warning, TEXT("\tChunkId='%s'"), *LexToString(ChunkId));
+				UE_LOG(LogIoStoreOnDemand, Warning, TEXT("\tChunkId='%s' (%s)"), *LexToString(ChunkId), *LexToString(ChunkId.GetChunkType()));
 				MissingChunkIds.Add(ChunkId);
 			}
 		}
diff --git a/Engine/Source/Runtime/Experimental/IoStore/OnDemand/Private/OnDemandContentInstaller.h b/Engine/Source/Runtime/Experimental/IoStore/OnDemand/Private/OnDemandContentInstaller.h
--- a/Engine/Source/Runtime/Experimental/IoStore/OnDemand/Private/OnDemandContentInstaller.h
+++ b/Engine/Source/Runtime/Experimental/IoStore/OnDemand/Private/OnDemandContentInstaller.h
@@ -57,22 +57,15 @@ void VisitPackageDependencies(
 	FPackageDependencyVisitor Visitor, 
 	FPackageStoreReadScope* InReadScope = nullptr);
 
-////////////////////////////////////////////////////////////////////////////////
-void ResolvePackageDependencies(
-	const TSet<FPackageId>& PackageIds,
-	bool bIncludeSoftReferences,
-	TSet<FPackageId>& OutResolved,
-	TSet<FPackageId>& OutMissing,
-	FPackageStoreReadScope* ReadScope = nullptr);
-
 ////////////////////////////////////////////////////////////////////////////////
 void ResolveChunksToInstall(
+	const TSharedPtr<IOnDemandPackageStoreBackend>& OnDemandPackageStore,
 	const TSet<FSharedOnDemandContainer>& Containers,
 	const TSet<FPackageId>& PackageIds,
 	bool bIncludeSoftReferences,
 	bool bIncludeOptionalBulkData,
 	TArray<FResolvedContainerChunks>& OutResolvedContainerChunks,
-	TSet<FPackageId>& OutMissing);
+	TSet<FIoChunkId>& OutMissing);
 
 } // namespace Private
 
diff --git a/Engine/Source/Runtime/Experimental/IoStore/OnDemand/Private/OnDemandIoStore.cpp b/Engine/Source/Runtime/Experimental/IoStore/OnDemand/Private/OnDemandIoStore.cpp
--- a/Engine/Source/Runtime/Experimental/IoStore/OnDemand/Private/OnDemandIoStore.cpp
+++ b/Engine/Source/Runtime/Experimental/IoStore/OnDemand/Private/OnDemandIoStore.cpp
@@ -724,9 +724,10 @@ TIoStatusOr<FOnDemandInstallSizeResult> FOnDemandIoStore::GetInstallSize(const F
 	const bool bIncludeSoftReferences	= EnumHasAnyFlags(Args.Options, EOnDemandGetInstallSizeOptions::IncludeSoftReferences);
 	const bool bIncludeOptionalBulkData = EnumHasAnyFlags(Args.Options, EOnDemandGetInstallSizeOptions::IncludeOptionalBulkData);
 	TArray<FResolvedContainerChunks> ResolvedChunks;
-	TSet<FPackageId> Missing;
+	TSet<FIoChunkId> Missing;
 
 	ResolveChunksToInstall(
+		PackageStoreBackend,
 		ContainersForInstallation,
 		PackageIdsToInstall,
 		bIncludeSoftReferences,
@@ -812,9 +813,10 @@ FIoStatus FOnDemandIoStore::GetInstallSizesByMountId(const FOnDemandGetInstallSi
 	const bool bIncludeSoftReferences	= EnumHasAnyFlags(Args.Options, EOnDemandGetInstallSizeOptions::IncludeSoftReferences);
 	const bool bIncludeOptionalBulkData = EnumHasAnyFlags(Args.Options, EOnDemandGetInstallSizeOptions::IncludeOptionalBulkData);
 	TArray<FResolvedContainerChunks> ResolvedChunks;
-	TSet<FPackageId> Missing;
+	TSet<FIoChunkId> Missing;
 
 	ResolveChunksToInstall(
+		PackageStoreBackend,
 		ContainersForInstallation,
 		PackageIdsToInstall,
 		bIncludeSoftReferences,
