Constraints & ControlRig: get the component's global transform using the evaluators if needed and possible
since GetGlobalTransform calculates the component transform via the animation evaluation, SetGlobalTransform must also use the same schema
without this, the attachment is evaluated in a second step, which doubles the child transform on the first evaluation

#rnx

#Synced-CL 50195398

diff --git a/Engine/Plugins/Animation/ControlRig/Source/ControlRig/Private/Constraints/ControlRigTransformableHandle.cpp b/Engine/Plugins/Animation/ControlRig/Source/ControlRig/Private/Constraints/ControlRigTransformableHandle.cpp
--- a/Engine/Plugins/Animation/ControlRig/Source/ControlRig/Private/Constraints/ControlRigTransformableHandle.cpp
+++ b/Engine/Plugins/Animation/ControlRig/Source/ControlRig/Private/Constraints/ControlRigTransformableHandle.cpp
@@ -209,8 +209,9 @@ void UTransformableControlHandle::SetGlobalTransform(const FTransform& InGlobal)
 		return;
 	}
 	
-	const FRigElementKey& ControlKey = ControlElement->GetKey();
-	const FTransform& ComponentTransform = BoundComponent->GetComponentTransform();
+	const FTransform ComponentTransform = TransformableHandleUtils::SkipTicking() ?
+		TransformableHandleUtils::GetGlobalTransform(BoundComponent, NAME_None) : 
+		BoundComponent->GetComponentTransform();
 
 	static const FRigControlModifiedContext Context(EControlRigSetKey::Never);
 	static constexpr bool bNotify = false, bSetupUndo = false, bPrintPython = false, bFixEulerFlips = false;
@@ -219,6 +220,7 @@ void UTransformableControlHandle::SetGlobalTransform(const FTransform& InGlobal)
 		UE::Private::FControlPoseChangedNotifier Notifier(ControlRig->GetHierarchy(), ControlElement, BoundComponent);
 	
 		//use this function so we don't set the preferred angles
+		const FRigElementKey& ControlKey = ControlElement->GetKey();
 		ControlRig->SetControlGlobalTransform(ControlKey.Name, InGlobal.GetRelativeTransform(ComponentTransform),
 			bNotify, Context, bSetupUndo, bPrintPython, bFixEulerFlips);
 	}
diff --git a/Engine/Source/Runtime/Experimental/Animation/Constraints/Private/Transform/AnimationEvaluation.cpp b/Engine/Source/Runtime/Experimental/Animation/Constraints/Private/Transform/AnimationEvaluation.cpp
--- a/Engine/Source/Runtime/Experimental/Animation/Constraints/Private/Transform/AnimationEvaluation.cpp
+++ b/Engine/Source/Runtime/Experimental/Animation/Constraints/Private/Transform/AnimationEvaluation.cpp
@@ -536,6 +536,17 @@ const FAnimationEvaluator& FAnimationEvaluationCache::GetEvaluator(USkeletalMesh
 	return NewEvaluator;
 }
 
+const FAnimationEvaluator& FAnimationEvaluationCache::FindEvaluator(const USkeletalMeshComponent* InSkeletalMeshComponent) const
+{
+	if (const FAnimationEvaluator* Found = InSkeletalMeshComponent ? PerSkeletalMeshEvaluator.Find(InSkeletalMeshComponent) : nullptr)
+	{
+		return *Found;
+	}
+	
+	static const FAnimationEvaluator Invalid(nullptr);
+	return Invalid;
+}
+	
 void FAnimationEvaluationCache::RegisterNotifications()
 {
 	if (!ConstraintsNotificationHandle.IsValid())
diff --git a/Engine/Source/Runtime/Experimental/Animation/Constraints/Private/Transform/TransformableHandle.cpp b/Engine/Source/Runtime/Experimental/Animation/Constraints/Private/Transform/TransformableHandle.cpp
--- a/Engine/Source/Runtime/Experimental/Animation/Constraints/Private/Transform/TransformableHandle.cpp
+++ b/Engine/Source/Runtime/Experimental/Animation/Constraints/Private/Transform/TransformableHandle.cpp
@@ -188,7 +188,7 @@ FTransform UTransformableComponentHandle::GetGlobalTransform() const
 			return AnimationEvaluator.GetGlobalTransform(SocketName);
 		}
 		
-		return TransformableHandleUtils::GetGlobalTransform(SceneComponent, SocketName);
+		return TransformableHandleUtils::EvaluateGlobalTransform(SceneComponent, SocketName);
 	}
 	
 	return Component.IsValid() ? Component->GetSocketTransform(SocketName) : FTransform::Identity;
diff --git a/Engine/Source/Runtime/Experimental/Animation/Constraints/Private/Transform/TransformableHandleUtils.cpp b/Engine/Source/Runtime/Experimental/Animation/Constraints/Private/Transform/TransformableHandleUtils.cpp
--- a/Engine/Source/Runtime/Experimental/Animation/Constraints/Private/Transform/TransformableHandleUtils.cpp
+++ b/Engine/Source/Runtime/Experimental/Animation/Constraints/Private/Transform/TransformableHandleUtils.cpp
@@ -115,6 +115,11 @@ void MarkComponentForEvaluation(const USceneComponent* InSceneComponent, const b
 }
 
 FTransform GetGlobalTransform(USceneComponent* InSceneComponent, const FName InSocketName)
+{
+	return EvaluateGlobalTransform(InSceneComponent, InSocketName);
+}
+	
+FTransform EvaluateGlobalTransform(USceneComponent* InSceneComponent, const FName InSocketName)
 {
 	if (!InSceneComponent)
 	{
@@ -183,6 +188,93 @@ FTransform GetGlobalTransform(USceneComponent* InSceneComponent, const FName InS
 	return Global;
 }
 
+FTransform GetGlobalTransform(const USceneComponent* InSceneComponent, const FName InSocketName)
+{
+	using namespace UE::Anim;
+	
+	if (!InSceneComponent)
+	{
+		return FTransform::Identity;
+	}
+
+	const USkeletalMeshComponent* SkeletalMeshComponent = Cast<USkeletalMeshComponent>(InSceneComponent);	
+	
+	constexpr bool bValidForTransforms = true;
+	FAnimationEvaluationCache& EvaluationCache = FAnimationEvaluationCache::Get();
+	const FAnimationEvaluator& Evaluator = EvaluationCache.FindEvaluator(SkeletalMeshComponent);
+	if (SkeletalMeshComponent && !Evaluator.IsValid(bValidForTransforms))
+	{
+		return InSceneComponent->GetSocketTransform(InSocketName);
+	}
+	
+	const bool bShouldUseEvaluators = [&EvaluationCache, InSceneComponent]()
+	{
+		bool bUseEvaluators = false;
+		
+		const USceneComponent* Parent = InSceneComponent->GetAttachParent();
+		FName AttachSocket = InSceneComponent->GetAttachSocketName();
+		
+		while (Parent)
+		{
+			const USkeletalMeshComponent* SkeletalMeshComponent = Cast<USkeletalMeshComponent>(Parent);
+			if (SkeletalMeshComponent && AttachSocket != NAME_None)
+			{
+				constexpr bool bValidForTransforms = true;
+				const FAnimationEvaluator& Evaluator = EvaluationCache.FindEvaluator(SkeletalMeshComponent);
+				bUseEvaluators |= Evaluator.IsValid(bValidForTransforms);
+				if (!bUseEvaluators)
+				{
+					return false;
+				}
+			}
+			
+			AttachSocket = Parent->GetAttachSocketName();
+			Parent = Parent->GetAttachParent();
+		}
+			
+		return bUseEvaluators;
+	}();
+
+	if (!bShouldUseEvaluators)
+	{
+		return InSceneComponent->GetSocketTransform(InSocketName);
+	}
+
+	FTransform Global = SkeletalMeshComponent && Evaluator.IsValid(bValidForTransforms) ?
+		Evaluator.GetRelativeTransform(InSocketName) :
+		InSceneComponent->GetSocketTransform(InSocketName, RTS_Component) * InSceneComponent->GetRelativeTransform();
+	
+	const USceneComponent* Parent = InSceneComponent->GetAttachParent();
+	FName AttachSocket = InSceneComponent->GetAttachSocketName();
+	while (Parent)
+	{
+		const USkeletalMeshComponent* ParentSkeletalMeshComponent = Cast<USkeletalMeshComponent>(Parent);
+
+		bool bFromEvaluator = false;
+		if (ParentSkeletalMeshComponent && AttachSocket != NAME_None)
+		{
+			const FAnimationEvaluator& ParentEvaluator = EvaluationCache.FindEvaluator(ParentSkeletalMeshComponent);
+			if (!ParentEvaluator.IsValid(bValidForTransforms))
+			{
+				return InSceneComponent->GetSocketTransform(InSocketName);
+			}
+			Global = Global * ParentEvaluator.GetRelativeTransform(AttachSocket);
+			bFromEvaluator = true;
+		}
+
+		if (!bFromEvaluator)
+		{
+			FTransform Socket = Parent->GetSocketTransform(AttachSocket, RTS_Component) * Parent->GetRelativeTransform();
+			Global = Global * Socket;
+		}
+
+		AttachSocket = Parent->GetAttachSocketName();
+		Parent = Parent->GetAttachParent();
+	}
+
+	return Global;
+}
+
 const UE::Anim::FAnimationEvaluator& EvaluateComponent(USceneComponent* InSceneComponent)
 {
 	if (USkeletalMeshComponent* SkeletalMeshComponent = Cast<USkeletalMeshComponent>(InSceneComponent))
@@ -194,7 +286,7 @@ const UE::Anim::FAnimationEvaluator& EvaluateComponent(USceneComponent* InSceneC
 	return Invalid;
 }
 
-CONSTRAINTS_API const UE::Anim::FAnimationEvaluator& EvaluateComponent(USceneComponent* InSceneComponent, const UE::Anim::FAnimationEvaluationTask& InTask)
+const UE::Anim::FAnimationEvaluator& EvaluateComponent(USceneComponent* InSceneComponent, const UE::Anim::FAnimationEvaluationTask& InTask)
 {
 	if (USkeletalMeshComponent* SkeletalMeshComponent = Cast<USkeletalMeshComponent>(InSceneComponent))
 	{
diff --git a/Engine/Source/Runtime/Experimental/Animation/Constraints/Public/Transform/AnimationEvaluation.h b/Engine/Source/Runtime/Experimental/Animation/Constraints/Public/Transform/AnimationEvaluation.h
--- a/Engine/Source/Runtime/Experimental/Animation/Constraints/Public/Transform/AnimationEvaluation.h
+++ b/Engine/Source/Runtime/Experimental/Animation/Constraints/Public/Transform/AnimationEvaluation.h
@@ -107,6 +107,12 @@ public:
 	 * Note that calling this function will actually evaluate the animation and the post-evaluation tasks if the evaluator has been previously dirtied.
 	 */
 	const FAnimationEvaluator& GetEvaluator(USkeletalMeshComponent* InSkeletalMeshComponent, const FAnimationEvaluationTask& InTask);
+	
+	/**
+	 * Returns the evaluator related to InSkeletalMeshComponent if any or an invalid one otherwise.
+	 * Note that calling this function won't update the evaluator so make sure to use FAnimationEvaluator::IsValid after calling that function.  
+	 */
+	const FAnimationEvaluator& FindEvaluator(const USkeletalMeshComponent* InSkeletalMeshComponent) const;
 
 private:
 	FAnimationEvaluationCache() = default;
diff --git a/Engine/Source/Runtime/Experimental/Animation/Constraints/Public/Transform/TransformableHandleUtils.h b/Engine/Source/Runtime/Experimental/Animation/Constraints/Public/Transform/TransformableHandleUtils.h
--- a/Engine/Source/Runtime/Experimental/Animation/Constraints/Public/Transform/TransformableHandleUtils.h
+++ b/Engine/Source/Runtime/Experimental/Animation/Constraints/Public/Transform/TransformableHandleUtils.h
@@ -27,11 +27,22 @@ namespace TransformableHandleUtils
 	 */
 	CONSTRAINTS_API void MarkComponentForEvaluation(const USceneComponent* InSceneComponent, const bool bRecursive = false);
 
+	UE_DEPRECATED(5.8, "Use TransformableHandleUtils::EvaluateGlobalTransform instead")
+	CONSTRAINTS_API FTransform GetGlobalTransform(USceneComponent* InSceneComponent, const FName InSocketName);
+	
 	/**
-	 * Returns the global transform of the scene component (composed with a socket component transform if InSocketName is not none).
-	 * If any of the component's parent is a skeletal mesh and has a valid evaluator, then it will be used to compose the final transform.  
+	 * Evaluates the global transform of the scene component (composed with a socket component transform if InSocketName is not none).
+	 * If any of the component's parent is a skeletal mesh and has a valid evaluator, then it will be used to compose the final transform.
+	 * Note that this function will update invalid evaluators if needed.
 	 */
-	CONSTRAINTS_API FTransform GetGlobalTransform(USceneComponent* InSceneComponent, const FName InSocketName);
+	CONSTRAINTS_API FTransform EvaluateGlobalTransform(USceneComponent* InSceneComponent, const FName InSocketName);
+	
+	/**
+	* Returns the global transform of the scene component (composed with a socket component transform if InSocketName is not none).
+	* If the component or one of its parents is a skeletal mesh component and has an invalid evaluator, this function will return the equivalent
+	* of USceneComponent::GetSocketTransform. Otherwise, the evaluators will be used to compose the final transform.
+	*/
+	CONSTRAINTS_API FTransform GetGlobalTransform(const USceneComponent* InSceneComponent, const FName InSocketName);
 	
 	/** Returns an updated version of InSceneComponent's animation evaluator. */
 	CONSTRAINTS_API const UE::Anim::FAnimationEvaluator& EvaluateComponent(USceneComponent* InSceneComponent);
