Geometry collection : fix velocities not being properly set on the GT particles when syncing from the game thread
- also update the character movement to properly pass the boneName to the function that get the linear velocity

#Synced-CL 50239131

diff --git a/Engine/Source/Runtime/Engine/Private/Components/CharacterMovementComponent.cpp b/Engine/Source/Runtime/Engine/Private/Components/CharacterMovementComponent.cpp
--- a/Engine/Source/Runtime/Engine/Private/Components/CharacterMovementComponent.cpp
+++ b/Engine/Source/Runtime/Engine/Private/Components/CharacterMovementComponent.cpp
@@ -7763,7 +7763,7 @@ void UCharacterMovementComponent::ApplyImpactPhysicsForces(const FHitResult& Imp
 
 				float PushForceModificator = 1.0f;
 
-				const FVector ComponentVelocity = ImpactComponent->GetPhysicsLinearVelocity();
+				const FVector ComponentVelocity = ImpactComponent->GetPhysicsLinearVelocity(Impact.BoneName);
 				const FVector VirtualVelocity = ImpactAcceleration.IsZero() ? ImpactVelocity : ImpactAcceleration.GetSafeNormal() * GetMaxSpeed();
 
 				float Dot = 0.0f;
diff --git a/Engine/Source/Runtime/Experimental/Chaos/Private/PhysicsProxy/GeometryCollectionPhysicsProxy.cpp b/Engine/Source/Runtime/Experimental/Chaos/Private/PhysicsProxy/GeometryCollectionPhysicsProxy.cpp
--- a/Engine/Source/Runtime/Experimental/Chaos/Private/PhysicsProxy/GeometryCollectionPhysicsProxy.cpp
+++ b/Engine/Source/Runtime/Experimental/Chaos/Private/PhysicsProxy/GeometryCollectionPhysicsProxy.cpp
@@ -127,6 +127,18 @@ FAutoConsoleVariableRef CVarGeometryCollectionRotationUpdateTolerance(
 	GeometryCollectionRotationUpdateTolerance,
 	TEXT("Tolerance to detect if particle rotation has changed has changed when syncing PT to GT"));
 
+float GeometryCollectionLinearVelocityUpdateTolerance = UE_KINDA_SMALL_NUMBER;
+FAutoConsoleVariableRef CVarGeometryCollectionLinearVelocityUpdateTolerance(
+	TEXT("p.GeometryCollection.LinearVelocityUpdateTolerance"),
+	GeometryCollectionLinearVelocityUpdateTolerance,
+	TEXT("Tolerance to detect if particle linear velocity has changed has changed when syncing PT to GT"));
+
+float GeometryCollectionAngularVelocityUpdateTolerance = UE_KINDA_SMALL_NUMBER;
+FAutoConsoleVariableRef CVarGeometryCollectionAngularVelocityUpdateTolerance(
+	TEXT("p.GeometryCollection.AngularVelocityUpdateTolerance"),
+	GeometryCollectionAngularVelocityUpdateTolerance,
+	TEXT("Tolerance to detect if particle angular velocity has changed has changed when syncing PT to GT"));
+
 bool bGeometryCollectionUseRootBrokenFlag = true;
 FAutoConsoleVariableRef CVarGeometryCollectionUseRootBrokenFlag(
 	TEXT("p.GeometryCollection.UseRootBrokenFlag"),
@@ -4318,6 +4330,27 @@ static inline bool UpdateGTParticleXR(Chaos::FPBDRigidParticle& GTParticle, cons
 	return bChanged;
 }
 
+// Update game thread particle V and W properties and retun true if they changed from previous values
+static inline bool UpdateGTParticleVW(Chaos::FPBDRigidParticle& GTParticle, const Chaos::FVec3& NewV, const Chaos::FVec3& NewW)
+{
+	GC_PHYSICSPROXY_CHECK_FOR_NAN(NewX);
+
+	const Chaos::FVec3 OldV = GTParticle.GetV();
+	const bool bNeedUpdateV = (!NewV.Equals(OldV, GeometryCollectionLinearVelocityUpdateTolerance));
+	if (bNeedUpdateV)
+	{
+		GTParticle.SetV(NewV, false);
+	}
+
+	const Chaos::FVec3 OldW = GTParticle.GetW();
+	const bool bNeedUpdateW = (!NewW.Equals(OldW, GeometryCollectionAngularVelocityUpdateTolerance));
+	if (bNeedUpdateW)
+	{
+		GTParticle.SetW(NewW, false);
+	}
+	return bNeedUpdateV || bNeedUpdateW;
+}
+
 // Update a value and return true if the value was different from the old one 
 template <typename TOld, typename TNew>
 static inline bool UpdateValue(TOld& ValueInOut,const TNew& NewValue)
@@ -4554,9 +4587,12 @@ bool FGeometryCollectionPhysicsProxy::PullNonInterpolatableDataFromSinglePhysics
 				const bool XRModified = UpdateGTParticleXR(GTParticle, PositionData.ParticleX, PositionData.ParticleR, Parameters.bForceUpdateActiveTransforms);
 				bIsCollectionDirty |= XRModified;
 
+				const FGeometryCollectionResults::FVelocityData& VelocityData = CurrentResults.GetVelocities(EntryIndex);
+				const bool VWModified = UpdateGTParticleVW(GTParticle, VelocityData.ParticleV, VelocityData.ParticleW);
+				bIsCollectionDirty |= VWModified;
+
 				if (LinearVelocities && AngularVelocities)
 				{
-					const FGeometryCollectionResults::FVelocityData& VelocityData = CurrentResults.GetVelocities(EntryIndex);
 					bIsCollectionDirty |= UpdateValue((*LinearVelocities)[TransformGroupIndex], VelocityData.ParticleV);
 					bIsCollectionDirty |= UpdateValue((*AngularVelocities)[TransformGroupIndex], VelocityData.ParticleW);
 				}
@@ -4768,11 +4804,14 @@ bool FGeometryCollectionPhysicsProxy::PullFromPhysicsState(const Chaos::FDirtyGe
 					const bool XRModified = UpdateGTParticleXR(GTParticle, NewX, NewR, Parameters.bForceUpdateActiveTransforms);
 					bIsCollectionDirty |= XRModified;
 
+					Chaos::FVec3f NewV;
+					Chaos::FVec3f NewW;
+					ResultInterpolator.GetVelocities(NewV, NewW);
+					const bool VWModified = UpdateGTParticleVW(GTParticle, NewV, NewW);
+					bIsCollectionDirty |= VWModified;
+
 					if (LinearVelocities && AngularVelocities)
 					{
-						Chaos::FVec3f NewV;
-						Chaos::FVec3f NewW;
-						ResultInterpolator.GetVelocities(NewV, NewW);
 						bIsCollectionDirty |= UpdateValue((*LinearVelocities)[TransformGroupIndex], NewV);
 						bIsCollectionDirty |= UpdateValue((*AngularVelocities)[TransformGroupIndex], NewW);
 					}
