Added on-demand install cache on-screen stats.

#rnx

#Synced-CL 50152864

diff --git a/Engine/Source/Runtime/Experimental/IoStore/OnDemand/Private/OnDemandInstallCache.cpp b/Engine/Source/Runtime/Experimental/IoStore/OnDemand/Private/OnDemandInstallCache.cpp
--- a/Engine/Source/Runtime/Experimental/IoStore/OnDemand/Private/OnDemandInstallCache.cpp
+++ b/Engine/Source/Runtime/Experimental/IoStore/OnDemand/Private/OnDemandInstallCache.cpp
@@ -2534,7 +2534,6 @@ bool FOnDemandInstallCache::Resolve(FIoRequestImpl* Request)
 		const FString Filename = Cas.GetBlockFilename(CasLoc.BlockId);
 		UE_LOG(LogIoStoreOnDemand, Error, TEXT("Failed to open CAS block '%s' for async reading, reason '%s'"), *Filename, *FileOpenResult.GetError().GetMessage());
 		TUniquePtr<FChunkRequest> Detached = FChunkRequest::Detach(*Request);
-		FOnDemandInstallCacheStats::OnReadCompleted(EIoErrorCode::FileOpenFailed);
 		CompleteRequest(Request, EIoErrorCode::FileOpenFailed);
 		return true;
 	}
@@ -2655,6 +2654,7 @@ TResult<bool> FOnDemandInstallCache::PutChunk(FIoBuffer&& Chunk, const FIoHash&
 
 	TUniqueLock Lock(PendingChunks.SharedMutex);
 
+	FOnDemandInstallCacheStats::OnPutChunk(Chunk.GetSize());
 	PendingChunks.Append(MoveTemp(Chunk), ChunkHash);
 	return MakeValue(true);
 }
@@ -4003,11 +4003,13 @@ FResult FOnDemandInstallCache::FlushPendingChunksImpl(const FPendingChunks& Chun
 
 void FOnDemandInstallCache::CompleteRequest(FIoRequestImpl* Request, EIoErrorCode Status)
 {
+	uint64 EncodedSize = 0;
 	if (Status == EIoErrorCode::Ok && !Request->IsCancelled())
 	{
 		FChunkRequest& ChunkRequest = FChunkRequest::GetRef(*Request);
 		const FOnDemandChunkInfo& ChunkInfo = ChunkRequest.ChunkInfo;
 		FIoBuffer EncodedChunk = MoveTemp(ChunkRequest.EncodedChunk);
+		EncodedSize = EncodedChunk.GetSize();
 
 		if (EncodedChunk.GetSize() > 0)
 		{
@@ -4046,7 +4048,7 @@ void FOnDemandInstallCache::CompleteRequest(FIoRequestImpl* Request, EIoErrorCod
 	{
 		UE::TUniqueLock Lock(Mutex);
 		CompletedRequests.AddTail(Request);
-		FOnDemandInstallCacheStats::OnReadCompleted(Status);
+		FOnDemandInstallCacheStats::OnReadCompleted(Status, EncodedSize);
 	}
 
 	BackendContext->WakeUpDispatcherThreadDelegate.Execute();
diff --git a/Engine/Source/Runtime/Experimental/IoStore/OnDemand/Private/Statistics.cpp b/Engine/Source/Runtime/Experimental/IoStore/OnDemand/Private/Statistics.cpp
--- a/Engine/Source/Runtime/Experimental/IoStore/OnDemand/Private/Statistics.cpp
+++ b/Engine/Source/Runtime/Experimental/IoStore/OnDemand/Private/Statistics.cpp
@@ -601,6 +601,31 @@ struct FInstallCacheAnalytics
 FMutex FInstallCacheAnalytics::Mutex;
 static FInstallCacheAnalytics InstallCacheAnalytics;
 
+////////////////////////////////////////////////////////////////////////////////
+struct FInstallCacheCounters
+{
+	FInstallCacheCounters()
+		: MaxCacheSize(TEXT("Iad/MaxCacheSize"), TraceCounterDisplayHint_Memory)
+		, CacheSize(TEXT("Iad/CacheSize"), TraceCounterDisplayHint_Memory)
+		, FlushedBytes(TEXT("Iad/FlushedBytes"), TraceCounterDisplayHint_Memory)
+		, FlushCount(TEXT("Iad/FlushCount"), TraceCounterDisplayHint_None)
+		, InMemoryBytes(TEXT("Iad/InMemoryBytes"), TraceCounterDisplayHint_Memory)
+		, InMemoryCount(TEXT("Iad/InMemoryCount"), TraceCounterDisplayHint_None)
+		, ReadBytes(TEXT("Iad/ReadBytes"), TraceCounterDisplayHint_Memory)
+		, ReadCount(TEXT("Iad/ReadCount"), TraceCounterDisplayHint_None)
+	{ }
+
+	FCounterAtomicInt MaxCacheSize;
+	FCounterAtomicInt CacheSize;
+	FCounterAtomicInt FlushedBytes;
+	FCounterAtomicInt FlushCount;
+	FCounterAtomicInt InMemoryBytes;
+	FCounterAtomicInt InMemoryCount;
+	FCounterAtomicInt ReadBytes;
+	FCounterAtomicInt ReadCount;
+};
+FInstallCacheCounters InstallCacheCounters;
+
 ////////////////////////////////////////////////////////////////////////////////
 static FOnDemandIoBackendStats* GStatistics = nullptr;
 
@@ -962,10 +987,10 @@ void FOnDemandIoBackendStats::PrintOnScreenStatistics(FCoreDelegates::FSeverityM
 	}
 #endif //#if UE_TRACK_CDN_HIT_STATUS
 
-	// Print IAS Cache
+	// Print Streaming Cache
 	{
 		TStringBuilder<256> CachingText;
-		CachingText << TEXT("IAS Cache: ");
+		CachingText << TEXT("Streaming Cache: ");
 		if (BackendStatus.IsCacheEnabled())
 		{
 			const float CacheUsagePercent = GCacheMaxBytes > 0 ? 100.f * (float(GCacheCachedBytes.Get()) / float(GCacheMaxBytes)) : 0.f;
@@ -992,17 +1017,17 @@ void FOnDemandIoBackendStats::PrintOnScreenStatistics(FCoreDelegates::FSeverityM
 		}
 		else
 		{
-			CachingText << TEXTVIEW("Caching Disabled");
+			CachingText << TEXTVIEW("Streaming Cache: Disabled");
 		}
 
 		OutMessages.Add(FCoreDelegates::EOnScreenMessageSeverity::Info, FText::FromStringView(CachingText.ToView()));
 	}
 
-	// Print IAS errors
+	// Print Streaming Cache errors
 	if (GHttpDecodeErrorCount[(uint8)EHttpRequestType::Streaming].Get() > 0 || GCacheDecodeErrorCount.Get() > 0 || GHttpErrorCount[(uint8)EHttpRequestType::Streaming].Get() > 0)
 	{
 		TStringBuilder<256> Builder;
-		Builder.Appendf(TEXT("IAS Errors: Cache Decode: %d | Http Decode: %d | Http: %d"),
+		Builder.Appendf(TEXT("Streaming Cache Errors: Cache Decode: %d | Http Decode: %d | Http: %d"),
 			GCacheDecodeErrorCount.Get(),
 			GHttpDecodeErrorCount[(uint8)EHttpRequestType::Streaming].Get(),
 			GHttpErrorCount[(uint8)EHttpRequestType::Streaming].Get()
@@ -1011,6 +1036,18 @@ void FOnDemandIoBackendStats::PrintOnScreenStatistics(FCoreDelegates::FSeverityM
 		OutMessages.Add(FCoreDelegates::EOnScreenMessageSeverity::Error, FText::FromStringView(Builder.ToView()));
 	}
 
+	// Print Install Cache
+	{
+		TStringBuilder<256> Sb;
+		Sb << TEXT("Install Cache: ");
+		Sb << TEXT(" Memory: ") << FText::AsMemory(InstallCacheCounters.InMemoryBytes.Get()).ToString();
+		Sb << TEXT(" (") << InstallCacheCounters.InMemoryCount.Get() << TEXT(")");
+		Sb << TEXT(" | Write: ") << FText::AsMemory(InstallCacheCounters.FlushedBytes.Get()).ToString();
+		Sb << TEXT(" (") << InstallCacheCounters.FlushCount.Get() << TEXT(")");
+		Sb << TEXT(" | Read: ") << FText::AsMemory(InstallCacheCounters.ReadBytes.Get()).ToString();
+		OutMessages.Add(FCoreDelegates::EOnScreenMessageSeverity::Info, FText::FromStringView(Sb.ToView()));
+	}
+
 #endif // UE_ENABLE_ONSCREEN_STATISTICS
 }
 
@@ -1683,13 +1720,25 @@ void FOnDemandInstallCacheStats::OnStartupError(const FResult& Result)
 	}
 }
 
+void FOnDemandInstallCacheStats::OnPutChunk(uint64 ByteCount)
+{
+	InstallCacheCounters.InMemoryBytes.Add(ByteCount);
+	InstallCacheCounters.InMemoryCount.Increment();
+}
+
 void FOnDemandInstallCacheStats::OnFlush(const FResult& Result, int64 ByteCount, bool bFromDefrag)
 {
+	InstallCacheCounters.InMemoryBytes.Set(0);
+	InstallCacheCounters.InMemoryCount.Set(0);
+
 	if (Result.HasError() && UE::UnifiedError::IsCancellationError(Result.GetError()))
 	{
 		return;
 	}
 
+	InstallCacheCounters.FlushedBytes.Add(ByteCount);
+	InstallCacheCounters.FlushCount.Increment();
+
 	UE::TUniqueLock Lock(InstallCacheAnalytics);
 
 	if (bFromDefrag)
@@ -1862,6 +1911,9 @@ void FOnDemandInstallCacheStats::OnCacheUsage(
 	uint64 FragmentedSize,
 	int64 OldestBlockAccess)
 {
+	InstallCacheCounters.MaxCacheSize.Set(MaxCacheSize);
+	InstallCacheCounters.CacheSize.Set(CacheSize);
+
 	UE::TUniqueLock Lock(InstallCacheAnalytics);
 	InstallCacheAnalytics.MaxCacheSize				= MaxCacheSize; 
 	InstallCacheAnalytics.MaxCacheUsageSize			= FMath::Max(InstallCacheAnalytics.MaxCacheUsageSize, CacheSize);
@@ -1871,13 +1923,16 @@ void FOnDemandInstallCacheStats::OnCacheUsage(
 	InstallCacheAnalytics.OldestBlockAccess			= FMath::Min(InstallCacheAnalytics.OldestBlockAccess, OldestBlockAccess);
 }
 
-void FOnDemandInstallCacheStats::OnReadCompleted(EIoErrorCode ErrorCode)
+void FOnDemandInstallCacheStats::OnReadCompleted(EIoErrorCode ErrorCode, uint64 EncodedSize)
 {
 	if (ErrorCode == EIoErrorCode::Cancelled)
 	{
 		return;
 	}
 
+	InstallCacheCounters.ReadBytes.Add(EncodedSize);
+	InstallCacheCounters.ReadCount.Increment();
+
 	UE::TUniqueLock Lock(InstallCacheAnalytics);
 	if (ErrorCode != EIoErrorCode::Ok)
 	{
diff --git a/Engine/Source/Runtime/Experimental/IoStore/OnDemand/Private/Statistics.h b/Engine/Source/Runtime/Experimental/IoStore/OnDemand/Private/Statistics.h
--- a/Engine/Source/Runtime/Experimental/IoStore/OnDemand/Private/Statistics.h
+++ b/Engine/Source/Runtime/Experimental/IoStore/OnDemand/Private/Statistics.h
@@ -166,6 +166,7 @@ class FOnDemandInstallCacheStats
 {
 public:
 	static void OnStartupError(const FResult& Result) IAS_STATISTICS_IMPL()
+	static void OnPutChunk(uint64 ChunkSize) IAS_STATISTICS_IMPL()
 	static void OnFlush(const FResult& Result, int64 ByteCount, bool bFromDefrag) IAS_STATISTICS_IMPL()
 	static void OnJournalCommit(EIoErrorCode ErrorCode, int64 ByteCount) IAS_STATISTICS_IMPL()
 	static void OnJournalCommit(const FResult& Result, int64 ByteCount) IAS_STATISTICS_IMPL()
@@ -190,7 +191,7 @@ public:
 		uint64 ReferencedSize,
 		uint64 FragmentedSize,
 		int64 OldestBlockAccess) IAS_STATISTICS_IMPL()
-	static void OnReadCompleted(EIoErrorCode ErrorCode) IAS_STATISTICS_IMPL()
+	static void OnReadCompleted(EIoErrorCode ErrorCode, uint64 EncodedSize) IAS_STATISTICS_IMPL()
 };
 
 #undef IAS_STATISTICS_IMPL
