Fixed a bug in SweepCapsuleVsSphere where the wrong normal was used.

#Synced-CL 50255566

diff --git a/Engine/Source/Runtime/Experimental/Chaos/Private/Chaos/GeometryQueries.cpp b/Engine/Source/Runtime/Experimental/Chaos/Private/Chaos/GeometryQueries.cpp
--- a/Engine/Source/Runtime/Experimental/Chaos/Private/Chaos/GeometryQueries.cpp
+++ b/Engine/Source/Runtime/Experimental/Chaos/Private/Chaos/GeometryQueries.cpp
@@ -154,8 +154,8 @@ namespace Chaos
 
 			OutNormal *= -1;
 			// Reconstruct the position on the sphere in local space (we transform everything below)
-			OutPosition = WorldSweepDir * SphereRadius;
-			TransformSweepResultsToWorld(bResult, OutTime, bComputeMTD, SweptCapsule, SphereLocation, WorldSweepDir, OutPosition, OutNormal, OutFaceIndex, OutPosition, OutNormal, OutFaceNormal);
+			OutPosition = OutNormal * SphereRadius;
+			TransformSweepResultsToWorld(bResult, OutTime, bComputeMTD, TestSphere, SphereLocation, WorldSweepDir, OutPosition, OutNormal, OutFaceIndex, OutPosition, OutNormal, OutFaceNormal);
 		}
 		return bResult;
 	}
