Dataflow : temporary fix for the bug that prevent traces from working on partial clusters when using remove on break/sleep

#Synced-CL 50239131

diff --git a/Engine/Source/Runtime/Experimental/Chaos/Private/PBDRigidsSolver.cpp b/Engine/Source/Runtime/Experimental/Chaos/Private/PBDRigidsSolver.cpp
--- a/Engine/Source/Runtime/Experimental/Chaos/Private/PBDRigidsSolver.cpp
+++ b/Engine/Source/Runtime/Experimental/Chaos/Private/PBDRigidsSolver.cpp
@@ -428,9 +428,17 @@ namespace Chaos
 		bool bRemoveParticleFromMovingKinematicsOnDisable = true;
 		FAutoConsoleVariableRef CVarChaosRemoveParticleFromMovingKinematicsOnDisable(TEXT("p.Chaos.RemoveParticleFromMovingKinematicsOnDisable"), bRemoveParticleFromMovingKinematicsOnDisable, TEXT(""));
 
-		bool bSingleThreadPullData = false;
+		bool bSingleThreadPullData = true;
 		FAutoConsoleVariableRef CVarChaosSingleThreadPushData(TEXT("p.Chaos.SingleThreadPullData"), bSingleThreadPullData, TEXT("Run pull Data in Single Thread."));
 	}
+
+	namespace Private
+	{
+		bool IsUsingSingleThreadPullData()
+		{
+			return CVars::bSingleThreadPullData;
+		}
+	}
 }
 
 namespace PhysicsReplicationCVars
diff --git a/Engine/Source/Runtime/Experimental/Chaos/Private/PhysicsProxy/GeometryCollectionPhysicsProxy.cpp b/Engine/Source/Runtime/Experimental/Chaos/Private/PhysicsProxy/GeometryCollectionPhysicsProxy.cpp
--- a/Engine/Source/Runtime/Experimental/Chaos/Private/PhysicsProxy/GeometryCollectionPhysicsProxy.cpp
+++ b/Engine/Source/Runtime/Experimental/Chaos/Private/PhysicsProxy/GeometryCollectionPhysicsProxy.cpp
@@ -4045,6 +4045,15 @@ static void UpdateParticleHandleTransform(Chaos::FPBDRigidsSolver& CurrentSolver
 	Handle.SetX(NewTransform.GetTranslation());
 	Handle.SetR(NewTransform.GetRotation());
 	Handle.UpdateWorldSpaceState(NewTransform, Chaos::FVec3{0});
+
+	// Temporary fix for the bug where player looses collision with partial clusters after using remove on sleep/break 
+	// Issue was introduced at CL 45952481
+	// Current temporary fix is to re-enable Single Thread Pull Data on the solver and when doing so we can safely update the dirty particle from this proxy
+	// Proper fix will require marking the dirty particle in a thread safe way via a potentially new API and having the actual dirtying happening as apost process to pull data when we use multi threaded pull data
+	if (Chaos::Private::IsUsingSingleThreadPullData())
+	{
+		CurrentSolver.GetEvolution()->DirtyParticle(Handle);
+	}
 }
 
 static void UpdateParticleHandleTransformIfNeeded(Chaos::FPBDRigidsSolver& CurrentSolver, Chaos::FPBDRigidClusteredParticleHandle& Handle, const Chaos::FRigidTransform3& NewTransform)
diff --git a/Engine/Source/Runtime/Experimental/Chaos/Public/PBDRigidsSolver.h b/Engine/Source/Runtime/Experimental/Chaos/Public/PBDRigidsSolver.h
--- a/Engine/Source/Runtime/Experimental/Chaos/Public/PBDRigidsSolver.h
+++ b/Engine/Source/Runtime/Experimental/Chaos/Public/PBDRigidsSolver.h
@@ -50,6 +50,8 @@ namespace ChaosTest
 namespace Chaos::Private
 {
 	class FTaskDispatcherSolver;
+
+	UE_INTERNAL CHAOS_API bool IsUsingSingleThreadPullData();
 }
 /**
 *
