Reset IAD cache and journal on certain platforms

diff --git a/Engine/Source/Runtime/Experimental/IoStore/OnDemand/Private/OnDemandInstallCache.cpp b/Engine/Source/Runtime/Experimental/IoStore/OnDemand/Private/OnDemandInstallCache.cpp
--- a/Engine/Source/Runtime/Experimental/IoStore/OnDemand/Private/OnDemandInstallCache.cpp
+++ b/Engine/Source/Runtime/Experimental/IoStore/OnDemand/Private/OnDemandInstallCache.cpp
@@ -45,6 +45,20 @@
 #define UE_ONDEMANDINSTALLCACHE_EXCLUSIVE_WRITE (0)
 #endif
 
+// Can't load these journal versions
+#if UE_ONDEMANDINSTALLCACHE_EXCLUSIVE_WRITE
+#define UE_ONDEMANDINSTALLCACHE_BAD_JOURNAL_VERSIONS EVersion::Invalid, EVersion::Initial
+#define UE_ONDEMANDINSTALLCACHE_BAD_SNAPSHOT_VERSIONS EVersion::Invalid, EVersion::Initial, EVersion::AddModificationTime
+#endif
+
+#ifndef UE_ONDEMANDINSTALLCACHE_BAD_JOURNAL_VERSIONS
+#define UE_ONDEMANDINSTALLCACHE_BAD_JOURNAL_VERSIONS EVersion::Invalid
+#endif
+
+#ifndef UE_ONDEMANDINSTALLCACHE_BAD_SNAPSHOT_VERSIONS
+#define UE_ONDEMANDINSTALLCACHE_BAD_SNAPSHOT_VERSIONS EVersion::Invalid
+#endif
+
 #ifndef UE_ONDEMANDINSTALLCACHE_USE_MODTIME
 #define UE_ONDEMANDINSTALLCACHE_USE_MODTIME (1)
 #endif
@@ -201,6 +215,8 @@ struct FCasJournal
 		Invalid = 0,
 		Initial,
 
+		FixSwitchBuild,
+
 		LatestPlusOne,
 		Latest = LatestPlusOne - 1
 	};
@@ -347,6 +363,16 @@ bool FCasJournal::FHeader::IsValid() const
 		return false;
 	}
 
+	// Versions that can no longer be loaded. The cache must be dropped instead.
+	const EVersion BadVersions[] = { UE_ONDEMANDINSTALLCACHE_BAD_JOURNAL_VERSIONS };
+	for (EVersion BadVersion : BadVersions)
+	{
+		if (Version == BadVersion)
+		{
+			return false;
+		}
+	}
+
 	return true;
 }
 
@@ -1485,7 +1511,9 @@ struct FCasSnapshot
 	{
 		Invalid = 0,
 		Initial,
+
 		AddModificationTime,
+		FixSwitchBuild,
 
 		LatestPlusOne,
 		Latest = LatestPlusOne - 1
@@ -1564,6 +1592,16 @@ bool FCasSnapshot::FHeader::IsValid() const
 		return false;
 	}
 
+	// Versions that can no longer be loaded. The cache must be dropped instead.
+	const EVersion BadVersions[] = { UE_ONDEMANDINSTALLCACHE_BAD_SNAPSHOT_VERSIONS };
+	for (EVersion BadVersion : BadVersions)
+	{
+		if (Version == BadVersion)
+		{
+			return false;
+		}
+	}
+
 	return true;
 }
 
@@ -1745,6 +1783,11 @@ TResult<FCasSnapshot> FCasSnapshot::Load(const FString& SnapshotFile, int64* Out
 
 	IFileManager& Ifm = IFileManager::Get();
 
+	if (Ifm.FileExists(*SnapshotFile) == false)
+	{
+		return MakeSnapshotError<FCasSnapshot>(ECasErrorCode::LoadSnapshotFailed, EIoErrorCode::NotFound, FString::Printf(TEXT("Failed to find snapshot '%s'"), *SnapshotFile));
+	}
+
 	TUniquePtr<FArchive> Ar(Ifm.CreateFileReader(*SnapshotFile));
 	if (Ar.IsValid() == false)
 	{
@@ -2156,9 +2199,9 @@ FOnDemandInstallCache::FOnDemandInstallCache(const FOnDemandInstallCacheConfig&
 		else if (const FCasErrorContext* Ctx = SnapshotResult.GetError().GetErrorContext<FCasErrorContext>())
 		{
 			// It's ok for the snapshot to not exist, other errors are significant
-			if (Ctx->IoErrorCode == EIoErrorCode::NotFound)
+			if (Ctx->IoErrorCode != EIoErrorCode::NotFound)
 			{
-				InitResult = MakeValue();
+				InitResult = MakeError(SnapshotResult.GetError());
 			}
 		}
 	}
