[NetPhysics] Resize RewindData when physics tick-rate changes.

#Synced-CL 50293424

diff --git a/Engine/Source/Runtime/Engine/Private/PhysicsEngine/NetworkPhysicsComponent.cpp b/Engine/Source/Runtime/Engine/Private/PhysicsEngine/NetworkPhysicsComponent.cpp
--- a/Engine/Source/Runtime/Engine/Private/PhysicsEngine/NetworkPhysicsComponent.cpp
+++ b/Engine/Source/Runtime/Engine/Private/PhysicsEngine/NetworkPhysicsComponent.cpp
@@ -1845,6 +1845,11 @@ void FAsyncNetworkPhysicsComponent::OnInitialize_Internal()
 				DelegateOnPostProcessInputs_Internal = SolverCallback->PostProcessInputsInternal.AddRaw(this, &FAsyncNetworkPhysicsComponent::OnPostProcessInputs_Internal);
 				DelegateOnAddResimulationRequest_Internal = SolverCallback->AddResimulationRequestInternal.AddRaw(this, &FAsyncNetworkPhysicsComponent::OnAddResimulationRequest_Internal);
 			}
+
+			if (Chaos::FRewindData* RewindData = BaseSolver->GetRewindData())
+			{
+				DelegateOnRewindDataResize_Internal = RewindData->RewindDataResize.AddRaw(this, &FAsyncNetworkPhysicsComponent::OnRewindDataResize_Internal);
+			}
 		}
 		else
 		{
@@ -1869,6 +1874,12 @@ void FAsyncNetworkPhysicsComponent::OnUninitialize_Internal()
 			SolverCallback->AddResimulationRequestInternal.Remove(DelegateOnAddResimulationRequest_Internal);
 			DelegateOnAddResimulationRequest_Internal.Reset();
 		}
+
+		if (Chaos::FRewindData* RewindData = BaseSolver->GetRewindData())
+		{
+			RewindData->RewindDataResize.Remove(DelegateOnRewindDataResize_Internal);
+			DelegateOnRewindDataResize_Internal.Reset();
+		}
 	}
 
 	UnregisterDataHistoryFromRewindData();
@@ -2859,6 +2870,18 @@ void FAsyncNetworkPhysicsComponent::TriggerResimulation(int32 ResimFrame)
 	}
 }
 
+void FAsyncNetworkPhysicsComponent::OnRewindDataResize_Internal(int32 InNumFrames)
+{
+	if (InputHistory)
+	{
+		InputHistory->ResizeDataHistory(InNumFrames);
+	}
+	if (StateHistory)
+	{
+		StateHistory->ResizeDataHistory(InNumFrames);
+	}
+}
+
 const float FAsyncNetworkPhysicsComponent::GetCurrentInputDecay() const
 {
 	if (!InputHistory)
@@ -3225,6 +3248,11 @@ const int32 FAsyncNetworkPhysicsComponent::SetupRewindData()
 		if (Chaos::FRewindData* RewindData = RigidSolver->GetRewindData())
 		{
 			NumFrames = RewindData->Capacity();
+
+			if (!DelegateOnRewindDataResize_Internal.IsValid())
+			{
+				DelegateOnRewindDataResize_Internal = RewindData->RewindDataResize.AddRaw(this, &FAsyncNetworkPhysicsComponent::OnRewindDataResize_Internal);
+			}
 		}
 	}
 
diff --git a/Engine/Source/Runtime/Engine/Public/Physics/NetworkPhysicsComponent.h b/Engine/Source/Runtime/Engine/Public/Physics/NetworkPhysicsComponent.h
--- a/Engine/Source/Runtime/Engine/Public/Physics/NetworkPhysicsComponent.h
+++ b/Engine/Source/Runtime/Engine/Public/Physics/NetworkPhysicsComponent.h
@@ -2017,6 +2017,9 @@ private:
 	/** Trigger a resimulation on frame */
 	void TriggerResimulation(int32 ResimFrame);
 
+	/** Rewind data has resized, also resize the Input and State histories */
+	void OnRewindDataResize_Internal(int32 InNumFrames);
+
 	/** Returns the current amount of input decay during resimulation as a magnitude from 0.0 to 1.0. Returns 0 if not currently resimulating. */
 	UE_DEPRECATED(5.8, "Use function that doesn't take a parameter, the decay curve will be based on the current frame number compared with the latest input frame number available")
 	const float GetCurrentInputDecay(const FNetworkPhysicsPayload* PhysicsData) { return GetCurrentInputDecay(); }
@@ -2114,6 +2117,7 @@ private:
 	// Cache compare incoming inputs with locally predicted inputs via FNetworkPhysicsData::CompareData to trigger a resim if they desync
 	bool bCompareInputToTriggerRewind;
 
+	FDelegateHandle DelegateOnRewindDataResize_Internal;
 	FDelegateHandle DelegateOnPreProcessInputs_Internal;
 	FDelegateHandle DelegateOnPostProcessInputs_Internal;
 	FDelegateHandle DelegateOnAddResimulationRequest_Internal;
diff --git a/Engine/Source/Runtime/Experimental/Chaos/Private/PBDRigidsSolver.cpp b/Engine/Source/Runtime/Experimental/Chaos/Private/PBDRigidsSolver.cpp
--- a/Engine/Source/Runtime/Experimental/Chaos/Private/PBDRigidsSolver.cpp
+++ b/Engine/Source/Runtime/Experimental/Chaos/Private/PBDRigidsSolver.cpp
@@ -1239,8 +1239,7 @@ namespace Chaos
 
 	void FPBDRigidsSolver::EnableRewindCapture()
 	{
-		int32 NumFrames = FMath::Max<int32>(1, FMath::CeilToInt32((0.001f * FPBDRigidsSolver::GetPhysicsHistoryTimeLength()) / GetAsyncDeltaTime()));
-		EnableRewindCapture(NumFrames);
+		EnableRewindCapture(0);
 	}
 
 	void FPBDRigidsSolver::Reset()
@@ -1260,10 +1259,12 @@ namespace Chaos
 		//todo: do we need this?
 		//MarshallingManager.Reset();
 
-		if (GetRewindData())
+		MRewindData = nullptr;
+		if (MRewindCallback)
 		{
-			EnableRewindCapture(GetRewindData()->Capacity(), GetUseCollisionResimCache());
+			MRewindCallback->RewindData = nullptr;
 		}
+		MarshallingManager.SetHistoryLength_Internal(0);
 
 		FEventDefaults::RegisterSystemEvents(*GetEventManager());
 	}
diff --git a/Engine/Source/Runtime/Experimental/Chaos/Private/RewindData.cpp b/Engine/Source/Runtime/Experimental/Chaos/Private/RewindData.cpp
--- a/Engine/Source/Runtime/Experimental/Chaos/Private/RewindData.cpp
+++ b/Engine/Source/Runtime/Experimental/Chaos/Private/RewindData.cpp
@@ -194,8 +194,17 @@ bool FJointStateBase::IsInSync(const FPBDJointConstraintHandle& Handle, const FF
 bool bCVarRewindDataOptimization = true;
 FAutoConsoleVariableRef CVarRewindDataOptimization(TEXT("p.Resim.RewindDataOptimization"), bCVarRewindDataOptimization, TEXT("Default value for RewinData optimization, note that this can be overridden at runtime by API calls. Effect: Only alter the minimum required properties during a resim for particles not marked for FullResim and only cache data during the PrePushData phase and lower memory allocation for the history cache to 1/3 of non-optimized flow."));
 
+static int32 CalculateNumFrames(FPBDRigidsSolver* InSolver, int32 NumFrames)
+{
+	if (ensure(InSolver) && NumFrames <= 0 && InSolver->GetAsyncDeltaTime() > UE_SMALL_NUMBER)
+	{
+		NumFrames = FMath::Max<int32>(1, FMath::CeilToInt32((0.001f * FPBDRigidsSolver::GetPhysicsHistoryTimeLength()) / InSolver->GetAsyncDeltaTime()));
+	}
+	return NumFrames;
+}
+
 FRewindData::FRewindData(FPBDRigidsSolver* InSolver, int32 NumFrames, bool InRewindDataOptimization, int32 InCurrentFrame)
-	: Managers(NumFrames + 1)	//give 1 extra for saving at head
+	: Managers(CalculateNumFrames(InSolver, NumFrames) + 1)
 	, Solver(InSolver)
 	, CurFrame(InCurrentFrame)
 	, LatestFrame(InCurrentFrame)
@@ -204,13 +213,20 @@ FRewindData::FRewindData(FPBDRigidsSolver* InSolver, int32 NumFrames, bool InRew
 	, bNeedsSave(false)
 	, bRewindDataOptimization(InRewindDataOptimization)
 	, LatestTargetFrame(0)
+	, AllowResizingFromTickTrate(false)
 {
+	if (NumFrames <= 0 && Solver)
+	{
+		AllowResizingFromTickTrate = true;
+		InitializedDeltaTime = Solver->GetAsyncDeltaTime();
+	}
+
 	ResimClampedFrames = Managers.Capacity();
 	RegisterEvolutionCallbacks();
 }
 
 FRewindData::FRewindData(FPBDRigidsSolver* InSolver, int32 NumFrames, int32 InCurrentFrame)
-	: Managers(NumFrames + 1)	//give 1 extra for saving at head
+	: Managers(CalculateNumFrames(InSolver, NumFrames) + 1)
 	, Solver(InSolver)
 	, CurFrame(InCurrentFrame)
 	, LatestFrame(InCurrentFrame)
@@ -219,8 +235,56 @@ FRewindData::FRewindData(FPBDRigidsSolver* InSolver, int32 NumFrames, int32 InCu
 	, bNeedsSave(false)
 	, bRewindDataOptimization(bCVarRewindDataOptimization)
 	, LatestTargetFrame(0)
+	, AllowResizingFromTickTrate(false)
+{
+	if (NumFrames <= 0 && Solver)
+	{
+		AllowResizingFromTickTrate = true;
+		InitializedDeltaTime = Solver->GetAsyncDeltaTime();
+	}
+
+	ResimClampedFrames = Managers.Capacity();
+	RegisterEvolutionCallbacks();
+}
+
+void FRewindData::Init(FPBDRigidsSolver* InSolver, int32 NumFrames, bool InRewindDataOptimization, int32 InCurrentFrame)
+{
+	Solver = InSolver;
+	CurFrame = InCurrentFrame;
+	LatestFrame = InCurrentFrame;
+	bRewindDataOptimization = InRewindDataOptimization;
+	LatestTargetFrame = 0;
+	AllowResizingFromTickTrate = false;
+
+	if (NumFrames <= 0 && Solver)
+	{
+		AllowResizingFromTickTrate = true;
+		InitializedDeltaTime = Solver->GetAsyncDeltaTime();
+	}
+
+	Managers = TCircularBuffer<FFrameManagerInfo>(CalculateNumFrames(InSolver, NumFrames) + 1);
+	ResimClampedFrames = Managers.Capacity();
+
+	RegisterEvolutionCallbacks();
+}
+
+void FRewindData::Init(FPBDRigidsSolver* InSolver, int32 NumFrames, int32 InCurrentFrame)
 {
+	Solver = InSolver;
+	CurFrame = InCurrentFrame;
+	LatestFrame = InCurrentFrame;
+	LatestTargetFrame = 0;
+	AllowResizingFromTickTrate = false;
+
+	if (NumFrames <= 0 && Solver)
+	{
+		AllowResizingFromTickTrate = true;
+		InitializedDeltaTime = Solver->GetAsyncDeltaTime();
+	}
+
+	Managers = TCircularBuffer<FFrameManagerInfo>(CalculateNumFrames(InSolver, NumFrames) + 1);
 	ResimClampedFrames = Managers.Capacity();
+
 	RegisterEvolutionCallbacks();
 }
 
@@ -1370,11 +1434,35 @@ bool FRewindData::GetUseCollisionResimCache() const
 	return Solver ? Solver->GetUseCollisionResimCache() : false;
 }
 
+int32 ResizeHistoryFromTickRate = 1;
+FAutoConsoleVariableRef CVarResimResizeHistoryFromTickRate(TEXT("p.Resim.ResizeHistoryFromTickRate"), ResizeHistoryFromTickRate, TEXT("When set to 1, RewindData will resize the history if needed when physics tick-rate changes."));
+
+void FRewindData::ResizeRewindDataIfNeeded()
+{
+	if (!AllowResizingFromTickTrate || ResizeHistoryFromTickRate != 1)
+	{
+		return;
+	}
+
+	// Check if tick-rate has changed since RewindData was initialized and if so, reinitialize
+	if (Solver && !IsResim() && FMath::IsNearlyEqual(InitializedDeltaTime, Solver->GetAsyncDeltaTime()) == false)
+	{
+		Solver->EnableRewindCapture();
+		BlockResim();
+		if (RewindDataResize.IsBound())
+		{
+			RewindDataResize.Broadcast(Managers.Capacity());
+		}
+	}
+}
+
 CHAOS_API int32 SkipDesyncTest = 0;
 FAutoConsoleVariableRef CVarSkipDesyncTest(TEXT("p.SkipDesyncTest"), SkipDesyncTest, TEXT("Skips hard desync test, this means all particles will assume to be clean except spawning at different times. This is useful for a perf lower bound, not actually correct"));
 
 void FRewindData::AdvanceFrameImp(IResimCacheBase* ResimCache)
 {
+	ResizeRewindDataIfNeeded();
+
 	FramesSaved = FMath::Min(FramesSaved + 1, static_cast<int32>(Managers.Capacity()));
 
 	const bool bHasResimCache = ResimCache != nullptr;
diff --git a/Engine/Source/Runtime/Experimental/Chaos/Public/RewindData.h b/Engine/Source/Runtime/Experimental/Chaos/Public/RewindData.h
--- a/Engine/Source/Runtime/Experimental/Chaos/Public/RewindData.h
+++ b/Engine/Source/Runtime/Experimental/Chaos/Public/RewindData.h
@@ -1719,46 +1719,31 @@ extern CHAOS_API int32 SkipDesyncTest;
 
 class FPBDRigidsSolver;
 
+DECLARE_MULTICAST_DELEGATE_OneParam(FOnRewindDataResizeInternal, const int32);
+
 class FRewindData
 {
 public:
 	FRewindData(FPBDRigidsSolver* InSolver, int32 NumFrames, bool InRewindDataOptimization, int32 InCurrentFrame);
 	FRewindData(FPBDRigidsSolver* InSolver, int32 NumFrames, int32 InCurrentFrame);
 
-	void Init(FPBDRigidsSolver* InSolver, int32 NumFrames, bool InRewindDataOptimization, int32 InCurrentFrame)
-	{
-		Solver = InSolver;
-		CurFrame = InCurrentFrame;
-		LatestFrame = InCurrentFrame;
-		bRewindDataOptimization = InRewindDataOptimization;
-		LatestTargetFrame = 0;
-		Managers = TCircularBuffer<FFrameManagerInfo>(NumFrames + 1);
-		ResimClampedFrames = Managers.Capacity();
-
-		RegisterEvolutionCallbacks();
-	}
-
-	void Init(FPBDRigidsSolver* InSolver, int32 NumFrames, int32 InCurrentFrame)
-	{
-		Solver = InSolver;
-		CurFrame = InCurrentFrame;
-		LatestFrame = InCurrentFrame;
-		LatestTargetFrame = 0;
-		Managers = TCircularBuffer<FFrameManagerInfo>(NumFrames + 1);
-		ResimClampedFrames = Managers.Capacity();
-
-		RegisterEvolutionCallbacks();
-	}
+	void Init(FPBDRigidsSolver* InSolver, int32 NumFrames, bool InRewindDataOptimization, int32 InCurrentFrame);
+	void Init(FPBDRigidsSolver* InSolver, int32 NumFrames, int32 InCurrentFrame);
 
 private:
 	void RegisterEvolutionCallbacks();
 
 public: 
+	FOnRewindDataResizeInternal RewindDataResize;
+
 	int32 Capacity() const { return Managers.Capacity(); }
 	int32 CurrentFrame() const { return CurFrame; }
 	int32 GetLatestFrame() const { return LatestFrame; }
 	int32 GetFramesSaved() const { return FramesSaved; }
 
+	/** Check if we need to resize the rewind data */
+	void ResizeRewindDataIfNeeded();
+
 	FReal GetDeltaTimeForFrame(int32 Frame) const
 	{
 		ensure(Managers[Frame].FrameCreatedFor == Frame);
@@ -2046,7 +2031,7 @@ private:
 		//A manager can have data for two frames at once, the important part is just knowing which frame it was created on so we know whether the physics data can rely on it
 		//Consider the case where nothing is dirty from GT and then an object moves from the simulation, in that case it needs a manager to record the data into
 		int32 FrameCreatedFor = INDEX_NONE;
-		FReal DeltaTime;
+		FReal DeltaTime = 0.0f;
 	};
 
 	template <typename THistoryType, typename TObj>
@@ -2259,6 +2244,8 @@ private:
 	int32 ResimCooldownUntilFrame = INDEX_NONE;
 	int32 ResimClampedFrames = INDEX_NONE;
 	int32 LatestTargetFrame;
+	bool AllowResizingFromTickTrate = false;
+	Chaos::FReal InitializedDeltaTime = 0.0f;
 
 	// Used to block rewinding past a physics change we currently don't handle
 	int32 BlockResimFrame = INDEX_NONE;
