Constraints: unregister interaction cache on deletion

#rnx
#Synced-CL 50298309

diff --git a/Engine/Source/Runtime/Experimental/Animation/Constraints/Private/Transform/TransformConstraintUtil.cpp b/Engine/Source/Runtime/Experimental/Animation/Constraints/Private/Transform/TransformConstraintUtil.cpp
--- a/Engine/Source/Runtime/Experimental/Animation/Constraints/Private/Transform/TransformConstraintUtil.cpp
+++ b/Engine/Source/Runtime/Experimental/Animation/Constraints/Private/Transform/TransformConstraintUtil.cpp
@@ -260,6 +260,12 @@ bool FConstraintsInteractionCache::HasAnyDependencyInternal(UObject* InChild, UO
 		
 	return false;
 }
+
+FConstraintsInteractionCache::~FConstraintsInteractionCache()
+{
+	Reset();
+	UnregisterNotifications();
+}
 	
 void FConstraintsInteractionCache::Reset()
 {
diff --git a/Engine/Source/Runtime/Experimental/Animation/Constraints/Public/Transform/TransformConstraintUtil.h b/Engine/Source/Runtime/Experimental/Animation/Constraints/Public/Transform/TransformConstraintUtil.h
--- a/Engine/Source/Runtime/Experimental/Animation/Constraints/Public/Transform/TransformConstraintUtil.h
+++ b/Engine/Source/Runtime/Experimental/Animation/Constraints/Public/Transform/TransformConstraintUtil.h
@@ -30,6 +30,8 @@ namespace UE::TransformConstraintUtil
 struct FConstraintsInteractionCache
 {
 public:
+	UE_API ~FConstraintsInteractionCache();
+	
 	UE_API void Reset();
 
 	UE_API const TArray< TWeakObjectPtr<UTickableConstraint> >& Get(const UObject* InObject, const FName& InAttachmentName);
