IAD cache - move bad version defines into build script.

#Synced-CL 50152522

diff --git a/Engine/Source/Runtime/Experimental/IoStore/OnDemand/Private/OnDemandInstallCache.cpp b/Engine/Source/Runtime/Experimental/IoStore/OnDemand/Private/OnDemandInstallCache.cpp
--- a/Engine/Source/Runtime/Experimental/IoStore/OnDemand/Private/OnDemandInstallCache.cpp
+++ b/Engine/Source/Runtime/Experimental/IoStore/OnDemand/Private/OnDemandInstallCache.cpp
@@ -45,12 +45,6 @@
 #define UE_ONDEMANDINSTALLCACHE_EXCLUSIVE_WRITE (0)
 #endif
 
-// Can't load these journal versions
-#if UE_ONDEMANDINSTALLCACHE_EXCLUSIVE_WRITE
-#define UE_ONDEMANDINSTALLCACHE_BAD_JOURNAL_VERSIONS EVersion::Invalid, EVersion::Initial
-#define UE_ONDEMANDINSTALLCACHE_BAD_SNAPSHOT_VERSIONS EVersion::Invalid, EVersion::Initial, EVersion::AddModificationTime
-#endif
-
 #ifndef UE_ONDEMANDINSTALLCACHE_BAD_JOURNAL_VERSIONS
 #define UE_ONDEMANDINSTALLCACHE_BAD_JOURNAL_VERSIONS EVersion::Invalid
 #endif
