Dataflow compiler : add option to trim dead ends
- Add IsTerminal virtual method to Dataflow node

#Synced-CL 50239131

diff --git a/Engine/Source/Runtime/Experimental/Dataflow/Core/Public/Dataflow/DataflowNode.h b/Engine/Source/Runtime/Experimental/Dataflow/Core/Public/Dataflow/DataflowNode.h
--- a/Engine/Source/Runtime/Experimental/Dataflow/Core/Public/Dataflow/DataflowNode.h
+++ b/Engine/Source/Runtime/Experimental/Dataflow/Core/Public/Dataflow/DataflowNode.h
@@ -743,6 +743,17 @@ PRAGMA_ENABLE_DEPRECATION_WARNINGS
 
 	bool HasValidConnections() const { return bHasValidConnections; }
 
+	/** 
+	* Define is a node is a terminal type of node 
+	* terminal nodes can be pulled from to generate content (see FDataflowTerminalNode ) 
+	* or write results to a simulation for example ( see FDataflowSimulationNode )
+	* The dataflow compiler also use this infromation to decide what part of the graph is a dead end and can be optimized out 
+	*/
+	virtual bool IsTerminal() const
+	{
+		return false;
+	}
+
 	virtual bool IsA(FName InType) const 
 	{ 
 		return InType.ToString().Equals(StaticType().ToString());
diff --git a/Engine/Source/Runtime/Experimental/Dataflow/Core/Public/Dataflow/DataflowTerminalNode.h b/Engine/Source/Runtime/Experimental/Dataflow/Core/Public/Dataflow/DataflowTerminalNode.h
--- a/Engine/Source/Runtime/Experimental/Dataflow/Core/Public/Dataflow/DataflowTerminalNode.h
+++ b/Engine/Source/Runtime/Experimental/Dataflow/Core/Public/Dataflow/DataflowTerminalNode.h
@@ -45,6 +45,11 @@ struct FDataflowTerminalNode : public FDataflowNode
 			|| Super::IsA(InType); 
 	}
 
+	virtual bool IsTerminal() const override
+	{
+		return true;
+	}
+
 	/** Return the terminal asset */
 	virtual TObjectPtr<UObject> GetTerminalAsset() const {return nullptr;}
 
diff --git a/Engine/Source/Runtime/Experimental/Dataflow/Engine/Private/Dataflow/DataflowCompiler.cpp b/Engine/Source/Runtime/Experimental/Dataflow/Engine/Private/Dataflow/DataflowCompiler.cpp
--- a/Engine/Source/Runtime/Experimental/Dataflow/Engine/Private/Dataflow/DataflowCompiler.cpp
+++ b/Engine/Source/Runtime/Experimental/Dataflow/Engine/Private/Dataflow/DataflowCompiler.cpp
@@ -33,6 +33,23 @@ namespace UE::Dataflow
 				return Nodes.IsEmpty();
 			}
 
+			bool HasTerminalNode() const
+			{
+				for (const FConstNodeRef& Node : Nodes)
+				{
+					if (Node->IsTerminal())
+					{
+						return true;
+					}
+				}
+				return false;
+			}
+
+			bool IsDeadEnd() const
+			{
+				return (DownstreamNodes.IsEmpty() && !HasTerminalNode());
+			}
+
 			static void AddDependency(FNodeSequenceRef UpstreamSequence, FNodeSequenceRef DownstreamSequence)
 			{
 				if (ensure(!UpstreamSequence->IsEmpty() && !DownstreamSequence->IsEmpty()))
@@ -45,6 +62,15 @@ namespace UE::Dataflow
 				}
 			}
 
+			static void CutFromUpstream(FNodeSequenceRef SequenceToCut)
+			{
+				for (FNodeSequenceRef UpstreamSequence : SequenceToCut->GetUpstreamDependencies())
+				{
+					UpstreamSequence->RemoveDownstreamDependency(SequenceToCut);
+				}
+				SequenceToCut->ClearUpstreamDependencies();
+			}
+
 			int32 GetExecutionIndex() const
 			{
 				return ExecutionIndex;
@@ -90,6 +116,31 @@ namespace UE::Dataflow
 				UpstreamDependencies.AddUnique(UpstreamSequence);
 			}
 
+			void RemoveDownstreamDependency(FNodeSequenceRef DownStreamSequence) 
+			{ 
+				DownstreamDependencies.Remove(DownStreamSequence);
+				const TArray<FConstNodeRef>& DownStreamSequenceNodes = DownStreamSequence->GetNodes();
+				if (DownStreamSequenceNodes.Num() > 0)
+				{
+					DownstreamNodes.Remove(DownStreamSequenceNodes[0]);
+				}
+			}
+
+			const TArray<FNodeSequenceRef>& GetUpstreamDependencies()
+			{
+				return UpstreamDependencies;
+			}
+
+			const TArray<FNodeSequenceRef>& GetDownstreamDependencies()
+			{
+				return DownstreamDependencies;
+			}
+
+			void ClearUpstreamDependencies()
+			{
+				UpstreamDependencies.Reset();
+			}
+
 		private:
 			int32 TaskIndex = INDEX_NONE;
 			int32 ExecutionIndex = INDEX_NONE;
@@ -259,7 +310,36 @@ namespace UE::Dataflow
 			}
 		}
 
-		// sort the sequences
+		// if needed, let's remove dead end sequences recursively 
+		// a dead end is a sequence that does not have downstream dependecies and does not end with a terminal or output node
+		if (Params.bRemoveDeadEnds)
+		{
+			TArray<TPair<FConstNodeRef, Private::FNodeSequenceRef>> SequencesToProcess = NodeSequencesByStartNode.Array();
+			while (SequencesToProcess.Num() > 0)
+			{
+				TPair<FConstNodeRef, Private::FNodeSequenceRef> SequenceToProcess = SequencesToProcess.Pop();
+				if (SequenceToProcess.Value->IsDeadEnd())
+				{
+					// add the previous sequences as they may become dead end themselves after we cut the current one
+					for (Private::FNodeSequenceRef UpstreamSequence : SequenceToProcess.Value->GetUpstreamDependencies())
+					{
+						TArray<FConstNodeRef> UpstreamNodes = UpstreamSequence->GetNodes();
+						if (UpstreamNodes.Num())
+						{
+							TPair<FConstNodeRef, Private::FNodeSequenceRef> SequenceToAdd{ UpstreamNodes[0], UpstreamSequence };
+							if (!SequencesToProcess.Contains(SequenceToAdd))
+							{
+								SequencesToProcess.Push(SequenceToAdd);
+							}
+						}
+					}
+					Private::FNodeSequence::CutFromUpstream(SequenceToProcess.Value);
+					NodeSequencesByStartNode.Remove(SequenceToProcess.Key);
+				}
+			}
+		}
+
+		// Sort the sequences
 		TArray<Private::FNodeSequenceRef> SortedSequencesByDependencies;
 		NodeSequencesByStartNode.GenerateValueArray(SortedSequencesByDependencies);
 
diff --git a/Engine/Source/Runtime/Experimental/Dataflow/Engine/Private/Dataflow/DataflowObject.cpp b/Engine/Source/Runtime/Experimental/Dataflow/Engine/Private/Dataflow/DataflowObject.cpp
--- a/Engine/Source/Runtime/Experimental/Dataflow/Engine/Private/Dataflow/DataflowObject.cpp
+++ b/Engine/Source/Runtime/Experimental/Dataflow/Engine/Private/Dataflow/DataflowObject.cpp
@@ -478,6 +478,7 @@ TSharedPtr<const UE::Dataflow::FCompiledGraph> UDataflow::CompileGraphIfNeeded()
 		if (UE::Dataflow::FCompiler::Compile(Params, Results))
 		{
 			CompiledGraphs = Results.GetCompiledGraphs();
+			UE_LOG(LogTemp, Warning, TEXT("Dataflow Compiled : time = %.3f ms (%d compiled graphs)"), (float)Results.GetCompilationTimeMs(), CompiledGraphs.Num());
 		}
 		else
 		{
diff --git a/Engine/Source/Runtime/Experimental/Dataflow/Engine/Public/Dataflow/DataflowCompiler.h b/Engine/Source/Runtime/Experimental/Dataflow/Engine/Public/Dataflow/DataflowCompiler.h
--- a/Engine/Source/Runtime/Experimental/Dataflow/Engine/Public/Dataflow/DataflowCompiler.h
+++ b/Engine/Source/Runtime/Experimental/Dataflow/Engine/Public/Dataflow/DataflowCompiler.h
@@ -27,7 +27,7 @@ namespace UE::Dataflow
 	struct FCompilerParameters
 	{
 	public:
-		bool OptimizeRerouteNodes = true;
+		bool bRemoveDeadEnds = true;
 		TWeakObjectPtr<const UDataflow> DataflowObjectToCompile;
 	};
 
@@ -61,6 +61,8 @@ namespace UE::Dataflow
 
 		const TMap<const UObject*, TSharedRef<FCompiledGraph>> GetCompiledGraphs() const { return CompiledGraphs; }
 
+		double GetCompilationTimeMs() const { return CompilationTimeMs; }
+
 	private:
 		friend struct FCompiler;
 
diff --git a/Engine/Source/Runtime/Experimental/Dataflow/Engine/Public/Dataflow/DataflowSubGraphNodes.h b/Engine/Source/Runtime/Experimental/Dataflow/Engine/Public/Dataflow/DataflowSubGraphNodes.h
--- a/Engine/Source/Runtime/Experimental/Dataflow/Engine/Public/Dataflow/DataflowSubGraphNodes.h
+++ b/Engine/Source/Runtime/Experimental/Dataflow/Engine/Public/Dataflow/DataflowSubGraphNodes.h
@@ -93,6 +93,12 @@ public:
 	DATAFLOWENGINE_API void AddConnectionsTo(TArrayView<FDataflowConnection*> Connections);
 
 private:
+	// Subgraph output nodes are considered terminal
+	virtual bool IsTerminal() const override
+	{
+		return true;
+	}
+
 	virtual void Evaluate(UE::Dataflow::FContext& Context, const FDataflowOutput* Out) const override;
 	virtual void PostSerialize(const FArchive& Ar) override;
 	virtual void OnPropertyChanged(UE::Dataflow::FContext& Context, const FPropertyChangedEvent& InPropertyChangedEvent) override;
diff --git a/Engine/Source/Runtime/Experimental/Dataflow/Simulation/Public/Dataflow/DataflowSimulationNodes.h b/Engine/Source/Runtime/Experimental/Dataflow/Simulation/Public/Dataflow/DataflowSimulationNodes.h
--- a/Engine/Source/Runtime/Experimental/Dataflow/Simulation/Public/Dataflow/DataflowSimulationNodes.h
+++ b/Engine/Source/Runtime/Experimental/Dataflow/Simulation/Public/Dataflow/DataflowSimulationNodes.h
@@ -51,6 +51,11 @@ struct FDataflowSimulationNode : public FDataflowNode
 			|| Super::IsA(InType); 
 	}
 
+	virtual bool IsTerminal() const override
+	{
+		return true;
+	}
+
 	/** Evaluate simulation dispatch */
 	virtual void Evaluate(UE::Dataflow::FContext& Context, const FDataflowOutput* Output) const override
 	{
