[Physics] Implement callbacks for InjectInputs_External, ProcessInputs_External and ProcessInputs_Internal into native ISimCallback instead of only via IRewindCallback. By default they are not enabled and if using IRewindCallback the callbacks will be handled via that as long as they are not enabled via ISimCallback.

To be deprecated from IRewindCallback when logic has been converted to using native flow.

#Synced-CL 50394790

diff --git a/Engine/Source/Runtime/Engine/Private/PhysicsEngine/NetworkPhysicsComponent.cpp b/Engine/Source/Runtime/Engine/Private/PhysicsEngine/NetworkPhysicsComponent.cpp
--- a/Engine/Source/Runtime/Engine/Private/PhysicsEngine/NetworkPhysicsComponent.cpp
+++ b/Engine/Source/Runtime/Engine/Private/PhysicsEngine/NetworkPhysicsComponent.cpp
@@ -366,14 +366,42 @@ bool FNetworkPhysicsRewindDataDeltaSourceStateProxy::NetSerialize(FArchive& Ar,
 
 // --------------------------- Network Physics Callback ---------------------------
 
-// Before PreSimulate_Internal
-void FNetworkPhysicsCallback::ProcessInputs_Internal(int32 PhysicsStep, const TArray<Chaos::FSimCallbackInputAndObject>& SimCallbacks)
+void FNetworkPhysicsCallback::InjectInputs_External(int32 PhysicsStep, int32 NumSteps)
+{
+	InjectInputsExternal.Broadcast(PhysicsStep, NumSteps);
+}
+
+void FNetworkPhysicsCallback::ProcessInputs_External(int32 PhysicsStep, const TArray<Chaos::FSimCallbackInputAndObject>& SimCallbackInputs)
+{
+	for (const Chaos::FSimCallbackInputAndObject& SimCallbackObject : SimCallbackInputs)
+	{
+		if (SimCallbackObject.CallbackObject
+			&& SimCallbackObject.CallbackObject->HasOption(Chaos::ESimCallbackOptions::Rewind) == true
+			&& SimCallbackObject.CallbackObject->HasOption(Chaos::ESimCallbackOptions::ProcessInputsExternal) == false) // Only call from here if not already listening to the native callback
+		{
+			SimCallbackObject.CallbackObject->ProcessInputs_External(PhysicsStep);
+		}
+	}
+}
+
+void FNetworkPhysicsCallback::PreProcessInputs_Internal(int32 PhysicsStep)
 {
 	PreProcessInputsInternal.Broadcast(PhysicsStep);
+}
+
+void FNetworkPhysicsCallback::ProcessInputs_Internal(int32 PhysicsStep, const TArray<Chaos::FSimCallbackInputAndObject>& SimCallbacks)
+{
 	for (Chaos::ISimCallbackObject* SimCallbackObject : RewindableCallbackObjects)
 	{
-		SimCallbackObject->ProcessInputs_Internal(PhysicsStep);
+		if (SimCallbackObject->HasOption(Chaos::ESimCallbackOptions::ProcessInputsInternal) == false) // Only call from here if not already listening to the native callback
+		{
+			SimCallbackObject->ProcessInputs_Internal(PhysicsStep);
+		}
 	}
+}
+
+void FNetworkPhysicsCallback::PostProcessInputs_Internal(int32 PhysicsStep)
+{
 	PostProcessInputsInternal.Broadcast(PhysicsStep);
 }
 
@@ -443,23 +471,6 @@ int32 FNetworkPhysicsCallback::TriggerRewindIfNeeded_Internal(int32 LatestStepCo
 	return ResimFrame;
 }
 
-void FNetworkPhysicsCallback::InjectInputs_External(int32 PhysicsStep, int32 NumSteps)
-{
-	InjectInputsExternal.Broadcast(PhysicsStep, NumSteps);
-}
-
-void FNetworkPhysicsCallback::ProcessInputs_External(int32 PhysicsStep, const TArray<Chaos::FSimCallbackInputAndObject>& SimCallbackInputs)
-{
-	for(const Chaos::FSimCallbackInputAndObject& SimCallbackObject : SimCallbackInputs)
-	{
-		if (SimCallbackObject.CallbackObject && SimCallbackObject.CallbackObject->HasOption(Chaos::ESimCallbackOptions::Rewind))
-		{
-			SimCallbackObject.CallbackObject->ProcessInputs_External(PhysicsStep);
-		}
-	}
-}
-
-
 // --------------------------- Network Physics System ---------------------------
 
 UNetworkPhysicsSystem::UNetworkPhysicsSystem()
diff --git a/Engine/Source/Runtime/Engine/Public/Physics/NetworkPhysicsComponent.h b/Engine/Source/Runtime/Engine/Public/Physics/NetworkPhysicsComponent.h
--- a/Engine/Source/Runtime/Engine/Public/Physics/NetworkPhysicsComponent.h
+++ b/Engine/Source/Runtime/Engine/Public/Physics/NetworkPhysicsComponent.h
@@ -706,7 +706,9 @@ struct FNetworkPhysicsCallback : public Chaos::IRewindCallback
 	// Rewind API
 	virtual void InjectInputs_External(int32 PhysicsStep, int32 NumSteps) override;
 	virtual void ProcessInputs_External(int32 PhysicsStep, const TArray<Chaos::FSimCallbackInputAndObject>& SimCallbackInputs);
+	virtual void PreProcessInputs_Internal(int32 PhysicsStep) override;
 	virtual void ProcessInputs_Internal(int32 PhysicsStep, const TArray<Chaos::FSimCallbackInputAndObject>& SimCallbackInputs) override;
+	virtual void PostProcessInputs_Internal(int32 PhysicsStep) override;
 	virtual void PreResimStep_Internal(int32 PhysicsStep, bool bFirst) override;
 	virtual void PostResimStep_Internal(int32 PhysicsStep) override;
 	virtual int32 TriggerRewindIfNeeded_Internal(int32 LatestStepCompleted) override;
diff --git a/Engine/Source/Runtime/Experimental/Chaos/Private/Chaos/Framework/PhysicsSolverBase.cpp b/Engine/Source/Runtime/Experimental/Chaos/Private/Chaos/Framework/PhysicsSolverBase.cpp
--- a/Engine/Source/Runtime/Experimental/Chaos/Private/Chaos/Framework/PhysicsSolverBase.cpp
+++ b/Engine/Source/Runtime/Experimental/Chaos/Private/Chaos/Framework/PhysicsSolverBase.cpp
@@ -454,8 +454,30 @@ namespace Chaos
 			{
 				if (MRewindCallback && !IsShuttingDown() && !bGameThreadFrozen)
 				{
-					QUICK_SCOPE_CYCLE_COUNTER(Chaos_RewindCallback_ProcessInputs_Internal);
-					MRewindCallback->ProcessInputs_Internal(PushData->InternalStep, PushData->SimCallbackInputs);
+					QUICK_SCOPE_CYCLE_COUNTER(Chaos_RewindCallback_PreProcessInputs_Internal);
+					MRewindCallback->PreProcessInputs_Internal(PushData->InternalStep);
+				}
+
+				for (ISimCallbackObject* Callback : ProcessInputsWatchers)
+				{
+					if (Callback->HasOption(ESimCallbackOptions::RunOnFrozenGameThread) == bGameThreadFrozen)
+					{
+						FScopedTraceSolverCallback TraceCallback(Callback);
+
+						Callback->ProcessInputs_Internal(PushData->InternalStep);
+					}
+				}
+
+				if (MRewindCallback && !IsShuttingDown() && !bGameThreadFrozen)
+				{
+					{
+						QUICK_SCOPE_CYCLE_COUNTER(Chaos_RewindCallback_ProcessInputs_Internal);
+						MRewindCallback->ProcessInputs_Internal(PushData->InternalStep, PushData->SimCallbackInputs);
+					}
+					{
+						QUICK_SCOPE_CYCLE_COUNTER(Chaos_RewindCallback_PostProcessInputs_Internal);
+						MRewindCallback->PostProcessInputs_Internal(PushData->InternalStep);
+					}
 				}
 			}
 
@@ -628,9 +650,22 @@ namespace Chaos
 
 		while(FPushPhysicsData* PushData = MarshallingManager.StepInternalTime_External())
 		{
-			if(MRewindCallback && !bIsShuttingDown)
+			if (!bIsShuttingDown)
 			{
-				MRewindCallback->ProcessInputs_External(PushData->InternalStep, PushData->SimCallbackInputs);
+				for (FSimCallbackInputAndObject& Callback : PushData->SimCallbackInputs)
+				{
+					if (Callback.CallbackObject && Callback.CallbackObject->HasOption(ESimCallbackOptions::ProcessInputsExternal))
+					{
+						FScopedTraceSolverCallback TraceCallback(Callback.CallbackObject);
+
+						Callback.CallbackObject->ProcessInputs_External(PushData->InternalStep);
+					}
+				}
+
+				if (MRewindCallback)
+				{
+					MRewindCallback->ProcessInputs_External(PushData->InternalStep, PushData->SimCallbackInputs);
+				}
 			}
 
 			if(ThreadingMode == EThreadingModeTemp::SingleThread)
diff --git a/Engine/Source/Runtime/Experimental/Chaos/Private/PBDRigidsSolver.cpp b/Engine/Source/Runtime/Experimental/Chaos/Private/PBDRigidsSolver.cpp
--- a/Engine/Source/Runtime/Experimental/Chaos/Private/PBDRigidsSolver.cpp
+++ b/Engine/Source/Runtime/Experimental/Chaos/Private/PBDRigidsSolver.cpp
@@ -1627,9 +1627,18 @@ namespace Chaos
 
 		GetEvolution()->GetBroadPhase().GetIgnoreCollisionManager().PushProducerStorageData_External(MarshallingManager.GetExternalTimestamp_External());
 
-		if (MRewindCallback && !IsShuttingDown())
+		if (!IsShuttingDown())
 		{
-			MRewindCallback->InjectInputs_External(MarshallingManager.GetInternalStep_External(), NumSteps);
+			for (ISimCallbackObject* Callback : InjectInputsWatchers_External)
+			{
+				FScopedTraceSolverCallback TraceCallback(Callback);
+				Callback->InjectInputs_External(MarshallingManager.GetInternalStep_External(), NumSteps);
+			}
+
+			if (MRewindCallback)
+			{
+				MRewindCallback->InjectInputs_External(MarshallingManager.GetInternalStep_External(), NumSteps);
+			}
 		}
 
 		MarshallingManager.Step_External(DeltaTime, NumSteps, GetSolverSubstep_External());
@@ -1949,6 +1958,11 @@ namespace Chaos
 					}
 				}
 
+				if (SimCallbackObject->HasOption(ESimCallbackOptions::ProcessInputsInternal))
+				{
+					ProcessInputsWatchers.Add(SimCallbackObject);
+				}
+
 				if (SimCallbackObject->HasOption(ESimCallbackOptions::PhysicsObjectUnregister))
 				{
 					PhysicsObjectUnregistrationWatchers.Add(SimCallbackObject);
@@ -2043,6 +2057,14 @@ namespace Chaos
 				}
 			}
 
+			for (int32 Idx = ProcessInputsWatchers.Num() - 1; Idx >= 0; --Idx)
+			{
+				if (ProcessInputsWatchers[Idx]->bPendingDelete)
+				{
+					ProcessInputsWatchers.RemoveAtSwap(Idx, EAllowShrinking::No);
+				}
+			}
+
 			for (int32 Idx = PhysicsObjectUnregistrationWatchers.Num() - 1; Idx >= 0; --Idx)
 			{
 				if (PhysicsObjectUnregistrationWatchers[Idx]->bPendingDelete)
diff --git a/Engine/Source/Runtime/Experimental/Chaos/Public/Chaos/Framework/PhysicsSolverBase.h b/Engine/Source/Runtime/Experimental/Chaos/Public/Chaos/Framework/PhysicsSolverBase.h
--- a/Engine/Source/Runtime/Experimental/Chaos/Public/Chaos/Framework/PhysicsSolverBase.h
+++ b/Engine/Source/Runtime/Experimental/Chaos/Public/Chaos/Framework/PhysicsSolverBase.h
@@ -452,6 +452,15 @@ namespace Chaos
 		void UnregisterAndFreeSimCallbackObject_External(ISimCallbackObject* SimCallbackObject)
 		{
 			MarshallingManager.UnregisterSimCallbackObject_External(SimCallbackObject);
+
+			// Unregister external callback
+			for (int32 Idx = InjectInputsWatchers_External.Num() - 1; Idx >= 0; --Idx)
+			{
+				if (InjectInputsWatchers_External[Idx]->bPendingDelete_External)
+				{
+					InjectInputsWatchers_External.RemoveAtSwap(Idx, EAllowShrinking::No);
+				}
+			}
 		}
 
 		template <typename Lambda>
@@ -924,6 +933,8 @@ namespace Chaos
 		TArray<ISimCallbackObject*> ContactModifiers;
 		TArray<ISimCallbackObject*> RegistrationWatchers;
 		TArray<ISimCallbackObject*> UnregistrationWatchers;
+		TArray<ISimCallbackObject*> InjectInputsWatchers_External;
+		TArray<ISimCallbackObject*> ProcessInputsWatchers;
 		TArray<ISimCallbackObject*> PhysicsObjectUnregistrationWatchers;
 
 		TUniquePtr<FRewindData> MRewindData;
@@ -948,6 +959,13 @@ namespace Chaos
 		{
 			ensure(SimCallbackObject->Solver == nullptr);	//double register?
 			SimCallbackObject->SetSolver_External(this);
+
+			// Register external callback
+			if (SimCallbackObject->HasOption(ESimCallbackOptions::InjectInputsExternal))
+			{
+				InjectInputsWatchers_External.Add(SimCallbackObject);
+			}
+
 			MarshallingManager.RegisterSimCallbackObject_External(SimCallbackObject);
 		}
 
diff --git a/Engine/Source/Runtime/Experimental/Chaos/Public/Chaos/SimCallbackObject.h b/Engine/Source/Runtime/Experimental/Chaos/Public/Chaos/SimCallbackObject.h
--- a/Engine/Source/Runtime/Experimental/Chaos/Public/Chaos/SimCallbackObject.h
+++ b/Engine/Source/Runtime/Experimental/Chaos/Public/Chaos/SimCallbackObject.h
@@ -29,7 +29,7 @@ namespace Utilities
 	CHAOS_API FReal GetSolverPhysicsResultsTime(const FPhysicsSolverBase*);
 }
 
-enum class ESimCallbackOptions : uint16
+enum class ESimCallbackOptions : uint32
 {
 	None					= 0,
 	Presimulate				= 1 << 0,
@@ -46,6 +46,9 @@ enum class ESimCallbackOptions : uint16
 	PostIntegrate			= 1 << 12,
 	PreSolve				= 1 << 13,
 	PostSolve				= 1 << 14,
+	InjectInputsExternal	= 1 << 15,
+	ProcessInputsExternal	= 1 << 16,
+	ProcessInputsInternal	= 1 << 17,
 };
 ENUM_CLASS_FLAGS(ESimCallbackOptions)
 
@@ -80,11 +83,26 @@ public:
 
 	virtual bool IsFAsyncObjectManagerCallback() const { return false;}
 
+	virtual void InjectInputs_External(int32 PhysicsStep, int32 NumSteps)
+	{
+		OnInjectInputs_External(PhysicsStep, NumSteps);
+	}
+
+	virtual void ProcessInputs_External(int32 PhysicsStep)
+	{
+		OnProcessInputs_External(PhysicsStep);
+	}
+
 	void PostInitialize_Internal()
 	{
 		OnPostInitialize_Internal();
 	}
 
+	virtual void ProcessInputs_Internal(int32 PhysicsStep)
+	{
+		OnProcessInputs_Internal(PhysicsStep);
+	}
+
 	void PreSimulate_Internal()
 	{
 		OnPreSimulate_Internal();
@@ -184,17 +202,6 @@ public:
 	{ 
 		return Solver; 
 	}
-
-	UE_DEPRECATED(5.5, "This callback was never called. Instead look at either ISimCallbackObject::ProcessInputs_External or FNetworkPhysicsCallback::InjectInputsExternal.")
-	virtual void InjectInputs_External(int32 PhysicsStep, int32 NumSteps) {}
-
-	/** Called before any sim callbacks are triggered but after async inputs has marshaled over
-	* Register with ESimCallbackOptions::Rewind to get callback */
-	virtual void ProcessInputs_Internal(int32 PhysicsStep) {}
-
-	/** Called before async inputs are marshaled over to the physics thread, only called if an async input is produced to get marshaled 
-	* Register with ESimCallbackOptions::Rewind to get callback */
-	virtual void ProcessInputs_External(int32 PhysicsStep) {}
 	
 	virtual int32 TriggerRewindIfNeeded_Internal(int32 LastCompletedStep)
 	{
@@ -259,11 +266,32 @@ private:
 	 */
 	virtual FSimCallbackInput* AllocateInputData_External() = 0;
 
+	/** Called before physics task is dispatched from the game thread, last call to produce async inputs in
+	* Register with ESimCallbackOptions::InjectInputsExternal to get callback */
+	virtual void OnInjectInputs_External(int32 PhysicsStep, int32 NumSteps)
+	{
+		check(false);
+	}
+
+	/** Called before async inputs are marshaled over to the physics thread, last call to alter already produced async inputs in
+	* Register with ESimCallbackOptions::ProcessInputsExternal to get callback */
+	virtual void OnProcessInputs_External(int32 PhysicsStep)
+	{
+		// check(false); // Enable when FNetworkPhysicsCallback no longer calls this
+	}
+
 	/**
-	* Called once when callback object is registered and after input data have been marshalled
+	* Called once when callback object is registered and after input data have been marshaled
 	*/
 	virtual void OnPostInitialize_Internal() { }
 
+	/** Called before any sim callbacks are triggered but after async inputs has marshaled over
+	* Register with ESimCallbackOptions::ProcessInputsInternal to get callback */
+	virtual void OnProcessInputs_Internal(int32 PhysicsStep)
+	{
+		// check(false); // Enable when FNetworkPhysicsCallback no longer calls this
+	}
+
 	/**
 	* Called before simulation step (NOTE: not once per sub-step when sub-stepping is enabled)
 	*/
diff --git a/Engine/Source/Runtime/Experimental/Chaos/Public/RewindData.h b/Engine/Source/Runtime/Experimental/Chaos/Public/RewindData.h
--- a/Engine/Source/Runtime/Experimental/Chaos/Public/RewindData.h
+++ b/Engine/Source/Runtime/Experimental/Chaos/Public/RewindData.h
@@ -2325,10 +2325,17 @@ class IRewindCallback
 {
 public:
 	virtual ~IRewindCallback() = default;
+	
+	/** Called just before ProcessInputs_Internal */
+	virtual void PreProcessInputs_Internal(int32 PhysicsStep) {}
+
 	/** Called before any sim callbacks are triggered but after physics data has marshaled over
 	*   This means brand new physics particles are already created for example, and any pending game thread modifications have happened
 	*   See ISimCallbackObject for recording inputs to callbacks associated with this PhysicsStep */
 	virtual void ProcessInputs_Internal(int32 PhysicsStep, const TArray<FSimCallbackInputAndObject>& SimCallbackInputs) {}
+	
+	/** Called just after ProcessInputs_Internal */
+	virtual void PostProcessInputs_Internal(int32 PhysicsStep) {}
 
 	/** Called after any presim callbacks are triggered and after physics data has marshaled over in order to modify the sim callback outputs */
 	virtual void ApplyCallbacks_Internal(int32 PhysicsStep, const TArray<ISimCallbackObject*>& SimCallbackObjects) {}
