Test multithreading tuning for collision detection

#Synced-CL 50361196

diff --git a/Engine/Source/Runtime/Experimental/Chaos/Private/Chaos/Framework/Parallel.cpp b/Engine/Source/Runtime/Experimental/Chaos/Private/Chaos/Framework/Parallel.cpp
--- a/Engine/Source/Runtime/Experimental/Chaos/Private/Chaos/Framework/Parallel.cpp
+++ b/Engine/Source/Runtime/Experimental/Chaos/Private/Chaos/Framework/Parallel.cpp
@@ -16,7 +16,7 @@ namespace Chaos
 	CHAOS_API int32 MaxNumWorkers = 100;
 	CHAOS_API int32 SmallBatchSize = 1;
 	CHAOS_API int32 LargeBatchSize = 100;
-	CHAOS_API int32 MinParallelTaskSize = 0;
+	CHAOS_API int32 MinParallelTaskSize = 1;
 	
 	CHAOS_API bool bDisablePhysicsParallelFor = false;
 	CHAOS_API bool bDisableParticleParallelFor = false;
