IAD Cache - on startup, verify current block and reset if needed.

#Synced-CL 50391842

diff --git a/Engine/Source/Runtime/Experimental/IoStore/OnDemand/Private/OnDemandInstallCache.cpp b/Engine/Source/Runtime/Experimental/IoStore/OnDemand/Private/OnDemandInstallCache.cpp
--- a/Engine/Source/Runtime/Experimental/IoStore/OnDemand/Private/OnDemandInstallCache.cpp
+++ b/Engine/Source/Runtime/Experimental/IoStore/OnDemand/Private/OnDemandInstallCache.cpp
@@ -731,6 +731,7 @@ struct FCas
 	FResult						Initialize(FStringView Directory, bool bDeleteExisting = false);
 	FCasLocation				FindChunk(const FIoHash& Hash) const;
 	FCasBlockId					CreateBlock();
+	bool						ContainsBlock(FCasBlockId BlockId) const;
 	FResult						DeleteBlock(FCasBlockId BlockId, TFunctionRef<void(const FCasAddr&)> OutAddrs);
 	FString						GetBlockFilename(FCasBlockId BlockId) const;
 	TResult<FSharedFileHandle>	OpenRead(FCasBlockId BlockId);
@@ -914,6 +915,13 @@ FCasBlockId FCas::CreateBlock()
 	return Out;
 }
 
+bool FCas::ContainsBlock(FCasBlockId BlockId) const
+{
+	UE::TUniqueLock Lock(Mutex);
+
+	return BlockIds.Contains(BlockId);
+}
+
 FResult FCas::DeleteBlock(FCasBlockId BlockId, TFunctionRef<void(const FCasAddr&)> OutAddrs)
 {
 	using namespace UE::UnifiedError::IoStoreOnDemand;
@@ -2261,7 +2269,7 @@ FOnDemandInstallCache::FOnDemandInstallCache(const FOnDemandInstallCacheConfig&
 	// Try to reset the cache if something has gone wrong 
 	if (InitResult.HasError())
 	{
-		UE_LOGFMT(LogIoStoreOnDemand, Error, "Resetting install cache, reason: {Error}", InitResult.GetError());
+		UE_LOGFMT(LogIoStoreOnDemand, Error, "Resetting install cache, reason: {Error}", *InitResult.GetError().GetErrorMessage(true).ToString());
 
 		FOnDemandInstallCacheStats::OnStartupError(InitResult);
 		InitResult = Reset();
@@ -2275,7 +2283,7 @@ FOnDemandInstallCache::FOnDemandInstallCache(const FOnDemandInstallCacheConfig&
 	}
 	else
 	{
-		UE_LOGFMT(LogIoStoreOnDemand, Error, "Failed to initialize install cache, reason: {Error}", InitResult.GetError());
+		UE_LOGFMT(LogIoStoreOnDemand, Error, "Failed to initialize install cache, reason: {Error}", *InitResult.GetError().GetErrorMessage(true).ToString());
 	}
 }
 
@@ -3136,6 +3144,8 @@ FResult FOnDemandInstallCache::Reset()
 		return Result;
 	}
 
+	CurrentBlock = FCasBlockId::Invalid;
+
 	UE_LOG(LogIoStoreOnDemand, Log, TEXT("Created CAS journal '%s'"), *JournalFile);
 	return MakeValue();
 }
@@ -3211,6 +3221,15 @@ FResult FOnDemandInstallCache::InitialVerify()
 		}
 	}
 
+	if (!Cas.ContainsBlock(CurrentBlock))
+	{
+		const FCasBlockId OldCurrentBlock = CurrentBlock.exchange(FCasBlockId::Invalid);
+		if (OldCurrentBlock != FCasBlockId::Invalid)
+		{
+			UE_LOG(LogIoStoreOnDemand, Warning, TEXT("Current Block %u was not found, resetting"), OldCurrentBlock.Id);
+		}
+	}
+
 	return MakeValue();
 }
 
