Dataflow : made dataflow compiler to split node sequence at terminal node boundaries in prepartion for better handling of multiple terminal node graphs
- Also fixed the simulation node not properly reporting the IsTerminal flag
- Also added logging to the Dataflow compiler for better debugging

#Synced-CL 50391333

diff --git a/Engine/Source/Runtime/Experimental/Dataflow/Engine/Private/Dataflow/DataflowCompiler.cpp b/Engine/Source/Runtime/Experimental/Dataflow/Engine/Private/Dataflow/DataflowCompiler.cpp
--- a/Engine/Source/Runtime/Experimental/Dataflow/Engine/Private/Dataflow/DataflowCompiler.cpp
+++ b/Engine/Source/Runtime/Experimental/Dataflow/Engine/Private/Dataflow/DataflowCompiler.cpp
@@ -10,6 +10,8 @@
 
 #define LOCTEXT_NAMESPACE "DataflowCompiler"
 
+DEFINE_LOG_CATEGORY_STATIC(LogDataflowCompiler, Log, All);
+
 namespace UE::Dataflow
 {
 	namespace Private
@@ -64,6 +66,7 @@ namespace UE::Dataflow
 
 			static void CutFromUpstream(FNodeSequenceRef SequenceToCut)
 			{
+				SequenceToCut->Log(TEXT("CutFromUpstream"));
 				for (FNodeSequenceRef UpstreamSequence : SequenceToCut->GetUpstreamDependencies())
 				{
 					UpstreamSequence->RemoveDownstreamDependency(SequenceToCut);
@@ -90,6 +93,7 @@ namespace UE::Dataflow
 
 			void AddToCompiledGraph(FCompiledGraph& OutCompiledGraph)
 			{
+				Log(TEXT("AddToCompiledGraph"));
 				const int32 NewTaskIndex = OutCompiledGraph.StartTask();
 				{
 					for (FConstNodeRef Node : Nodes)
@@ -141,6 +145,20 @@ namespace UE::Dataflow
 				UpstreamDependencies.Reset();
 			}
 
+			void Log(const TCHAR* Message) const
+			{
+				if (UE_LOG_ACTIVE(LogDataflowCompiler, Verbose))
+				{
+					TStringBuilder<256> NodeList;
+					for (const FConstNodeRef& Node : Nodes)
+					{
+						NodeList.Append(Node->GetName().ToString());
+						NodeList.Append(", ");
+					}
+					UE_LOG(LogDataflowCompiler, Verbose, TEXT("%s - [%p] - nodes:[%s]"), Message, this, *NodeList);
+				}
+			}
+
 		private:
 			int32 TaskIndex = INDEX_NONE;
 			int32 ExecutionIndex = INDEX_NONE;
@@ -162,6 +180,12 @@ namespace UE::Dataflow
 			FConstNodePtr NodeToProcess = StartNode;
 			while (NodeToProcess)
 			{
+				// if the current node is a terminal we stop the sequence right away
+				if (NodeToProcess->IsTerminal())
+				{
+					break;
+				}
+
 				// check if we have a singular node connected to this one
 				FCompiler::FindDownstreamNodes(NodeToProcess.ToSharedRef(), DownstreamNodes);
 				if (DownstreamNodes.Num() != 1)
@@ -182,6 +206,7 @@ namespace UE::Dataflow
 				NodeToProcess = Nodes.Last();
 				DownstreamNodes.Reset();
 			}
+			Log(TEXT("Built"));
 		}
 	}
 
@@ -227,12 +252,16 @@ namespace UE::Dataflow
 	{
 		static FGuid NoGuid;
 
+		UE_LOG(LogDataflowCompiler, Verbose, TEXT("Dataflow Compiler - START"));
+
 		const double StartCompilationTimeMs = FPlatformTime::ToMilliseconds64(FPlatformTime::Cycles64());
 
 		// Get dataflow Object to compile
 		TStrongObjectPtr<const UDataflow> DataflowObject = Params.DataflowObjectToCompile.Pin();
 		if (DataflowObject)
 		{
+			UE_LOG(LogDataflowCompiler, Verbose, TEXT("Dataflow Compiler - Asset [%s]"), *DataflowObject->GetFName().ToString());
+
 			TArray<FConstNodeRef> NodesInGraph;
 
 			// compile main graph 
@@ -255,6 +284,8 @@ namespace UE::Dataflow
 		const double EndCompilationTimeMs = FPlatformTime::ToMilliseconds64(FPlatformTime::Cycles64());
 		Results.CompilationTimeMs = (EndCompilationTimeMs - StartCompilationTimeMs);
 
+		UE_LOG(LogDataflowCompiler, Verbose, TEXT("Dataflow Compiler - END - time=%.3fms"), (float)Results.CompilationTimeMs);
+
 		return Results.Succedded();
 	}
 
