[Chaos] Fix solver actor initialisation crash when creating a solver actor on an async thread (e.g. async loading thread). Physics scene creation and event registration must happen on the game thread so the initialization flow must be moved out of the actor constructor and into the actor lifetime flow.

#Synced-CL 50199821

diff --git a/Engine/Source/Runtime/Experimental/ChaosSolverEngine/Private/Chaos/ChaosSolverActor.cpp b/Engine/Source/Runtime/Experimental/ChaosSolverEngine/Private/Chaos/ChaosSolverActor.cpp
--- a/Engine/Source/Runtime/Experimental/ChaosSolverEngine/Private/Chaos/ChaosSolverActor.cpp
+++ b/Engine/Source/Runtime/Experimental/ChaosSolverEngine/Private/Chaos/ChaosSolverActor.cpp
@@ -211,16 +211,6 @@ AChaosSolverActor::AChaosSolverActor(const FObjectInitializer& ObjectInitializer
 {
 	if(!HasAnyFlags(RF_ClassDefaultObject))
 	{
-		// Don't spawn solvers on the CDO
-
-		// @question(Benn) : Does this need to be created on the Physics thread using a queued command?
-		PhysScene = MakeShareable(new FPhysScene_Chaos(this
-#if CHAOS_DEBUG_NAME
-								  , TEXT("Solver Actor Physics")
-#endif
-		));
-		RigidSolverProxy.Solver = PhysScene->GetSolver();
-		
 		// Ticking setup for collision/breaking notifies
 		PrimaryActorTick.TickGroup = TG_PostPhysics;
 		PrimaryActorTick.bCanEverTick = true;
@@ -278,6 +268,24 @@ void AChaosSolverActor::PreInitializeComponents()
 	Super::PreInitializeComponents();
 }
 
+void AChaosSolverActor::ConditionalCreateSolver()
+{
+	if(PhysScene)
+	{
+		return;
+	}
+
+	if(!HasAnyFlags(RF_DefaultSubObject | RF_ArchetypeObject | RF_ClassDefaultObject))
+	{
+		PhysScene = MakeShareable(new FPhysScene_Chaos(this
+#if CHAOS_DEBUG_NAME
+			, TEXT("Solver Actor Physics")
+#endif
+		));
+		RigidSolverProxy.Solver = PhysScene->GetSolver();
+	}
+}
+
 void AChaosSolverActor::MigrateSolver() const
 {
 	if(GetSolver())
@@ -413,6 +421,13 @@ void AChaosSolverActor::Serialize(FArchive& Ar)
 	Ar.UsingCustomVersion(FFortniteMainBranchObjectVersion::GUID);
 }
 
+void AChaosSolverActor::PreRegisterAllComponents()
+{
+	Super::PreRegisterAllComponents();
+
+	ConditionalCreateSolver();
+}
+
 void AChaosSolverActor::PostRegisterAllComponents()
 {
 	Super::PostRegisterAllComponents();
@@ -435,12 +450,18 @@ void AChaosSolverActor::PostUnregisterAllComponents()
 	
 	// Unregister the dataflow simulation interface
 	UE::Dataflow::UnregisterSimulationInterface(this);
+
+	// At this point all components/actors that connect to the solver actor via register/beginplay
+	// should now no longer be resident within the physics scene so we can remove our reference
+	PhysScene = nullptr;
+	RigidSolverProxy.Solver = nullptr;
 }
 
 void AChaosSolverActor::PostDuplicate(EDuplicateMode::Type DuplicateMode)
 {
 	Super::PostDuplicate(DuplicateMode);
 
+	ConditionalCreateSolver();
 	MigrateSolver();
 }
 
diff --git a/Engine/Source/Runtime/Experimental/ChaosSolverEngine/Public/Chaos/ChaosSolverActor.h b/Engine/Source/Runtime/Experimental/ChaosSolverEngine/Public/Chaos/ChaosSolverActor.h
--- a/Engine/Source/Runtime/Experimental/ChaosSolverEngine/Public/Chaos/ChaosSolverActor.h
+++ b/Engine/Source/Runtime/Experimental/ChaosSolverEngine/Public/Chaos/ChaosSolverActor.h
@@ -185,6 +185,7 @@ public:
 	TSharedPtr<FPhysScene_Chaos> GetPhysicsScene() const { return PhysScene; }
 	Chaos::FPhysicsSolver* GetSolver() const { return RigidSolverProxy.Solver; }
 
+	CHAOSSOLVERENGINE_API virtual void PreRegisterAllComponents() override;
 	CHAOSSOLVERENGINE_API virtual void PostRegisterAllComponents() override;
 	CHAOSSOLVERENGINE_API virtual void PreInitializeComponents() override;
 	CHAOSSOLVERENGINE_API virtual void PostUnregisterAllComponents() override;
@@ -216,6 +217,9 @@ private:
 	/** Rigid solver proxy used in dataflow simulation */
 	FDataflowRigidSolverProxy RigidSolverProxy;
 
+	/** Ensure that the solver actor has a valid solver and scene */
+	void ConditionalCreateSolver();
+
 	/** Migrate the solver onto the right owner (World vs Actor)*/
 	void MigrateSolver() const;
 
