Dstm: Chaos vehicle support 2.1: Add migration for particle dirty bits. Process GT dirty data after migrated PT data is applied. Reset geometry on migrated BodyInstance to be dirty so it gets pushed into PT.

#Synced-CL 50453497

diff --git a/Engine/Source/Runtime/Engine/Classes/PhysicsEngine/BodyInstance.h b/Engine/Source/Runtime/Engine/Classes/PhysicsEngine/BodyInstance.h
--- a/Engine/Source/Runtime/Engine/Classes/PhysicsEngine/BodyInstance.h
+++ b/Engine/Source/Runtime/Engine/Classes/PhysicsEngine/BodyInstance.h
@@ -1515,6 +1515,8 @@ public:
 
 private:
 
+	void MarkNonFrequentDirty();
+
 	ENGINE_API void UpdateDebugRendering();
 
 	/** Used to map between shapes and welded bodies. We do not create entries if the owning body instance is root*/
diff --git a/Engine/Source/Runtime/Engine/Private/PhysicsEngine/BodyInstance.cpp b/Engine/Source/Runtime/Engine/Private/PhysicsEngine/BodyInstance.cpp
--- a/Engine/Source/Runtime/Engine/Private/PhysicsEngine/BodyInstance.cpp
+++ b/Engine/Source/Runtime/Engine/Private/PhysicsEngine/BodyInstance.cpp
@@ -3497,6 +3497,24 @@ void FBodyInstance::ApplySerializedState(Chaos::FSerializedDataBufferPtr&& InSer
 	{
 		FChaosEngineInterface::ApplySerializedStateToParticle(Actor, MoveTemp(InSerializedState), DataContext);
 	}, LockContext, EPhysicsInterfaceScopedTransactionMode::Normal);
+
+	// Collision geometry and shape data are not serialized for Dstm (see PushToPhysicsStateImp). 
+	// So NonFrequentData (which is contains them) is marked as dirty to push those data into PT. 
+	// This ensures ShapesArray is populated before shape properties are updated
+	if (DataContext == Chaos::ESerializedDataContext::External)
+	{
+		MarkNonFrequentDirty();
+	}
+}
+
+void FBodyInstance::MarkNonFrequentDirty()
+{
+	FPhysicsActorHandle PhysicsActor = GetPhysicsActor();
+	if (PhysicsActor)
+	{
+		Chaos::FRigidBodyHandle_External& BodyExternal = PhysicsActor->GetGameThreadAPI();
+		BodyExternal.MarkNonFrequentDataDirty();
+	}
 }
 
 void FBodyInstance::SerializeState(Chaos::FSerializedDataBuffer& OutSerializedState, Chaos::ESerializedDataContext DataContext)
diff --git a/Engine/Source/Runtime/Experimental/Chaos/Private/Chaos/Serialization/SolverSerializer.cpp b/Engine/Source/Runtime/Experimental/Chaos/Private/Chaos/Serialization/SolverSerializer.cpp
--- a/Engine/Source/Runtime/Experimental/Chaos/Private/Chaos/Serialization/SolverSerializer.cpp
+++ b/Engine/Source/Runtime/Experimental/Chaos/Private/Chaos/Serialization/SolverSerializer.cpp
@@ -242,6 +242,8 @@ namespace Chaos
 			RigidParticle->SetObjectState(static_cast<EObjectStateType>(InParticleState.ParticleDynamicsMisc.MObjectState));
 		}
 
+		GTParticle->DirtyFlags().SetBits(InParticleState.DirtyFlagsBits);
+
 		// TODO: Add support for Cluster unions and Geometry Collections
 	}
 
diff --git a/Engine/Source/Runtime/Experimental/Chaos/Private/PhysicsProxy/SingleParticlePhysicsProxy.cpp b/Engine/Source/Runtime/Experimental/Chaos/Private/PhysicsProxy/SingleParticlePhysicsProxy.cpp
--- a/Engine/Source/Runtime/Experimental/Chaos/Private/PhysicsProxy/SingleParticlePhysicsProxy.cpp
+++ b/Engine/Source/Runtime/Experimental/Chaos/Private/PhysicsProxy/SingleParticlePhysicsProxy.cpp
@@ -72,6 +72,7 @@ void PushToPhysicsStateImp(const Chaos::FDirtyPropertiesManager& Manager, Chaos:
 	const FParticleNonFrequentData* NewNonFrequentData = ParticleData.FindNonFrequentData(Manager, DataIdx);
 
 #if UE_WITH_REMOTE_OBJECT_HANDLE
+	bool bIsMigrated = false;
 	if (FSerializedDataBufferPtr PendingStateData = Dirty.Proxy && !Dirty.Proxy->IsInitialized() ? Solver.GetSerializer().PopPendingInternalSerializedStateForProxy(Dirty.Proxy) : nullptr)
 	{
 		// Currently we don't serialize collision geometry and shape data, so we need to apply it first
@@ -82,8 +83,9 @@ void PushToPhysicsStateImp(const Chaos::FDirtyPropertiesManager& Manager, Chaos:
 		}
 
 		Solver.GetSerializer().ApplySerializedStateToParticle(Handle, *PendingStateData);
+
+		bIsMigrated = true;
 	}
-	else
 #endif
 	{
 		if (bResimInitialized)	//todo: assumes particles are always initialized as enabled. This is not true in future versions of code, so check PushData
@@ -110,12 +112,17 @@ void PushToPhysicsStateImp(const Chaos::FDirtyPropertiesManager& Manager, Chaos:
 				Evolution.SetParticleTransform(Handle, NewXR->X(), NewXR->R(), bIsTeleport, bWakeUp);
 			}
 
-			if(NewNonFrequentData)
+#if UE_WITH_REMOTE_OBJECT_HANDLE
+			if (!bIsMigrated)
+#endif
 			{
-				// Geometry may have changed, we need to remove the particle and its collisions from the graph
-				Evolution.InvalidateParticle(Handle);
+				if (NewNonFrequentData)
+				{
+					// Geometry may have changed, we need to remove the particle and its collisions from the graph
+					Evolution.InvalidateParticle(Handle);
 
-				Handle->SetNonFrequentData(*NewNonFrequentData);
+					Handle->SetNonFrequentData(*NewNonFrequentData);
+				}
 			}
 
 			auto NewVelocities = bHasKinematicData ? ParticleData.FindVelocities(Manager, DataIdx) : nullptr;
diff --git a/Engine/Source/Runtime/Experimental/Chaos/Public/Chaos/ParticleDirtyFlags.h b/Engine/Source/Runtime/Experimental/Chaos/Public/Chaos/ParticleDirtyFlags.h
--- a/Engine/Source/Runtime/Experimental/Chaos/Public/Chaos/ParticleDirtyFlags.h
+++ b/Engine/Source/Runtime/Experimental/Chaos/Public/Chaos/ParticleDirtyFlags.h
@@ -996,6 +996,11 @@ inline FChaosArchive& operator<<(FChaosArchive& Ar,FMaterialData& Data)
 			return Bits;
 		}
 
+		void SetBits(const int32 BitsIn)
+		{
+			Bits = BitsIn;
+		}
+
 	private:
 		int32 Bits;
 	};
diff --git a/Engine/Source/Runtime/Experimental/Chaos/Public/Chaos/ParticleHandle.h b/Engine/Source/Runtime/Experimental/Chaos/Public/Chaos/ParticleHandle.h
--- a/Engine/Source/Runtime/Experimental/Chaos/Public/Chaos/ParticleHandle.h
+++ b/Engine/Source/Runtime/Experimental/Chaos/Public/Chaos/ParticleHandle.h
@@ -3079,6 +3079,11 @@ public:
 		return MDirtyFlags;
 	}
 
+	FDirtyChaosPropertyFlags& DirtyFlags()
+	{
+		return MDirtyFlags;
+	}
+
 	void ClearDirtyFlags()
 	{
 		MDirtyFlags.Clear();
diff --git a/Engine/Source/Runtime/Experimental/Chaos/Public/PhysicsProxy/SingleParticlePhysicsProxy.h b/Engine/Source/Runtime/Experimental/Chaos/Public/PhysicsProxy/SingleParticlePhysicsProxy.h
--- a/Engine/Source/Runtime/Experimental/Chaos/Public/PhysicsProxy/SingleParticlePhysicsProxy.h
+++ b/Engine/Source/Runtime/Experimental/Chaos/Public/PhysicsProxy/SingleParticlePhysicsProxy.h
@@ -1336,6 +1336,27 @@ public:
 		VerifyContext();
 		GetParticle_LowLevel()->SetGeometry(ImplicitGeometryPtr);
 	}
+
+	void MarkNonFrequentDataDirty()
+	{
+		VerifyContext();
+
+		FGeometryParticle* GTParticle = GetParticle_LowLevel();
+		if (GTParticle)
+		{
+			// Mark NonFrequentData as dirty.
+			GTParticle->ForceDirty(EChaosPropertyFlags::NonFrequentData);
+
+			// Manually add proxy to solver's dirty list (replicating what MarkDirty does)
+			if (IPhysicsProxyBase* Proxy = GetProxy())
+			{
+				if (FPhysicsSolverBase* SolverBase = Proxy->GetSolverBase())
+				{
+					SolverBase->AddDirtyProxy(Proxy);
+				}
+			}
+		}
+	}
 	
 	UE_DEPRECATED(5.4, "Use SetGeometry with FImplicitObjectPtr instead.")
 	void SetGeometry(TUniquePtr<FImplicitObject>&& UniqueGeometry)
