Remove RigidPhysics dependence on Engine.

#Synced-CL 50403248

diff --git a/Engine/Plugins/Experimental/ChaosRigidPhysicsAsync/Tests/Private/BodyTests.cpp b/Engine/Plugins/Experimental/ChaosRigidPhysicsAsync/Tests/Private/BodyTests.cpp
--- a/Engine/Plugins/Experimental/ChaosRigidPhysicsAsync/Tests/Private/BodyTests.cpp
+++ b/Engine/Plugins/Experimental/ChaosRigidPhysicsAsync/Tests/Private/BodyTests.cpp
@@ -13,6 +13,8 @@
 #include "RigidPhysics/RigidModifier.h"
 #include "RigidPhysics/RigidScene.h"
 
+#include "RigidTestUtils.h"
+
 #if UE_RIGIDPHYSICS_API_ENABLED
 
 namespace Chaos::LowLevelTest
diff --git a/Engine/Plugins/Experimental/ChaosRigidPhysicsAsync/Tests/Private/RigidTestUtils.cpp b/Engine/Plugins/Experimental/ChaosRigidPhysicsAsync/Tests/Private/RigidTestUtils.cpp
new file mode 100644
--- /dev/null
+++ b/Engine/Plugins/Experimental/ChaosRigidPhysicsAsync/Tests/Private/RigidTestUtils.cpp
@@ -0,0 +1,104 @@
+// Copyright Epic Games, Inc. All Rights Reserved.
+
+#include "RigidTestUtils.h"
+
+#if UE_RIGIDPHYSICS_API_ENABLED
+
+#include "Chaos/CollisionFilterData.h"
+#include "PhysicalMaterials/PhysicalMaterial.h"
+#include "Physics/Experimental/ChaosInterfaceUtils.h"
+#include "Physics/PhysicsInterfaceTypes.h"
+#include "PhysicsEngine/AggregateGeom.h"
+
+namespace UE::Physics
+{
+	FVector3f MakeSolidBoxInertia(float InMass, const FVector3f& InSize)
+	{
+		FVector3f SizeSq = InSize * InSize;
+		return (InMass / 12.0f) * FVector3f(
+			SizeSq.Y + SizeSq.Z,
+			SizeSq.X + SizeSq.Z,
+			SizeSq.X + SizeSq.Y);
+	}
+
+	FRigidShapeInstancePtr MakeShapeInstanceFromAddParams(const FGeometryAddParams& InParams)
+	{
+		TArray<Chaos::FImplicitObjectPtr> Geoms;
+		Chaos::FShapesArray Shapes;
+		ChaosInterface::CreateGeometry(InParams, Geoms, Shapes);
+
+		if (Shapes.Num() != 1)
+		{
+			UE_LOGFMT(LogRigidPhysics, Warning, "MakeShapeInstanceFromAddParams is intended to be used for single shapes only. Have {0} shapes.", Shapes.Num());
+		}
+
+		if (Shapes.Num() > 0)
+		{
+			return FRigidShapeInstancePtr(Shapes[0].Release());
+		}
+
+		return FRigidShapeInstancePtr();
+	}
+
+	FRigidShapeInstancePtr MakeColliderFromAggGeom(
+		FKAggregateGeom& InAggGeom, 
+		const FVector3d& InScale,
+		ECollisionTraceFlag CollisionTraceFlag,
+		UPhysicalMaterial* InSimpleMaterial,
+		const TArrayView<UPhysicalMaterial*>& InComplexMaterials,
+		const TArrayView<FPhysicalMaterialMaskParams>& InComplexMaterialMasks,
+		const FTransform& InRelativeTransform)
+	{
+		// Not used - filters are set up after
+		FBodyCollisionData BodyCollisionData;
+
+		FGeometryAddParams AddParams;
+		AddParams.bDoubleSided = false;
+		AddParams.CollisionData = BodyCollisionData;
+		AddParams.CollisionTraceType = CollisionTraceFlag;
+		AddParams.Scale = InScale;
+		AddParams.SimpleMaterial = InSimpleMaterial;
+		AddParams.ComplexMaterials = InComplexMaterials;
+		AddParams.ComplexMaterialMasks = InComplexMaterialMasks;
+		AddParams.LocalTransform = InRelativeTransform;
+		AddParams.WorldTransform = FTransform::Identity;
+		AddParams.Geometry = &InAggGeom;
+		//AddParams.TriMeshGeometries = MakeArrayView(TriMeshGeometries);
+
+		return MakeShapeInstanceFromAddParams(AddParams);
+	}
+
+	FRigidShapeInstancePtr MakeBoxShape(const FVector3f& InSize, const FVector3f& InCenter, UPhysicalMaterial* InMaterial)
+	{
+		using namespace Chaos;
+
+		// Single Box
+		FKBoxElem BoxElem = FKBoxElem(InSize.X, InSize.Y, InSize.Z);
+		BoxElem.Center = FVector(InCenter);
+
+		FKAggregateGeom AggGeom;
+		AggGeom.AddElement(BoxElem);
+
+		FRigidShapeInstancePtr Collider = MakeColliderFromAggGeom(
+			AggGeom, 
+			FVector3d::One(),
+			ECollisionTraceFlag::CTF_UseSimpleAsComplex,
+			InMaterial,
+			{},
+			{},
+			FTransform3d::Identity);
+
+		if (Collider.IsValid())
+		{
+			const EFilterFlags FilterFlags = EFilterFlags::SimpleCollision | EFilterFlags::ContactNotify | EFilterFlags::ModifyContacts;
+			const Filter::FShapeFilterData ShapeFilter = Filter::FShapeFilterBuilder::BuildBlockAll(FilterFlags);
+			Collider->GetPerShapeData()->SetQueryEnabled(true);
+			Collider->GetPerShapeData()->SetSimEnabled(true);
+			Collider->GetPerShapeData()->SetShapeFilterData(ShapeFilter);
+		}
+
+		return Collider;
+	}
+}
+
+#endif
\ No newline at end of file
diff --git a/Engine/Plugins/Experimental/ChaosRigidPhysicsAsync/Tests/Private/RigidTestUtils.h b/Engine/Plugins/Experimental/ChaosRigidPhysicsAsync/Tests/Private/RigidTestUtils.h
new file mode 100644
--- /dev/null
+++ b/Engine/Plugins/Experimental/ChaosRigidPhysicsAsync/Tests/Private/RigidTestUtils.h
@@ -0,0 +1,28 @@
+// Copyright Epic Games, Inc. All Rights Reserved.
+
+#pragma once
+
+#if UE_RIGIDPHYSICS_API_ENABLED
+
+#include "Chaos/ImplicitFwd.h"
+#include "Math/Bounds.h"
+#include "Math/Transform.h"
+#include "RigidPhysics/RigidCollider.h"
+#include "RigidPhysics/RigidLog.h"
+
+class UPhysicalMaterial;
+
+namespace Chaos
+{
+	class FPerShapeData;
+}
+
+namespace UE::Physics
+{
+	// Utility for Unit Tests until the API if more fully implemented
+	FRigidShapeInstancePtr MakeBoxShape(const FVector3f& InSize, const FVector3f& InCenter = FVector3f::Zero(), UPhysicalMaterial* InMaterial = nullptr);
+	FVector3f MakeSolidBoxInertia(float InMass, const FVector3f& InSize);
+
+}
+
+#endif
\ No newline at end of file
diff --git a/Engine/Plugins/Experimental/ChaosRigidPhysicsAsync/Tests/Private/SceneLockBenchmarks.cpp b/Engine/Plugins/Experimental/ChaosRigidPhysicsAsync/Tests/Private/SceneLockBenchmarks.cpp
--- a/Engine/Plugins/Experimental/ChaosRigidPhysicsAsync/Tests/Private/SceneLockBenchmarks.cpp
+++ b/Engine/Plugins/Experimental/ChaosRigidPhysicsAsync/Tests/Private/SceneLockBenchmarks.cpp
@@ -12,6 +12,8 @@
 #include "RigidPhysics/RigidCollider.h"
 #include "RigidPhysics/RigidContext.h"
 
+#include "RigidTestUtils.h"
+
 #if UE_RIGIDPHYSICS_API_ENABLED
 
 namespace Chaos::LowLevelTest
diff --git a/Engine/Plugins/Experimental/ChaosRigidPhysicsAsync/Tests/Private/SceneTests.cpp b/Engine/Plugins/Experimental/ChaosRigidPhysicsAsync/Tests/Private/SceneTests.cpp
--- a/Engine/Plugins/Experimental/ChaosRigidPhysicsAsync/Tests/Private/SceneTests.cpp
+++ b/Engine/Plugins/Experimental/ChaosRigidPhysicsAsync/Tests/Private/SceneTests.cpp
@@ -15,6 +15,8 @@
 #include "RigidPhysics/RigidModifier.h"
 #include "RigidPhysics/RigidScene.h"
 
+#include "RigidTestUtils.h"
+
 #if UE_RIGIDPHYSICS_API_ENABLED
 
 namespace Chaos::LowLevelTest
diff --git a/Engine/Source/Runtime/Experimental/RigidPhysics/Private/RigidPhysics/RigidCollider.cpp b/Engine/Source/Runtime/Experimental/RigidPhysics/Private/RigidPhysics/RigidCollider.cpp
--- a/Engine/Source/Runtime/Experimental/RigidPhysics/Private/RigidPhysics/RigidCollider.cpp
+++ b/Engine/Source/Runtime/Experimental/RigidPhysics/Private/RigidPhysics/RigidCollider.cpp
@@ -1,10 +1,6 @@
 // Copyright Epic Games, Inc. All Rights Reserved.
 #include "RigidPhysics/RigidCollider.h"
-#include "Chaos/CollisionFilterData.h"
-#include "PhysicalMaterials/PhysicalMaterial.h"
 #include "Physics/Experimental/ChaosInterfaceUtils.h"
-#include "Physics/PhysicsInterfaceTypes.h"
-#include "PhysicsEngine/AggregateGeom.h"
 #include "RigidPhysics/RigidLog.h"
 
 #if UE_RIGIDPHYSICS_API_ENABLED
@@ -41,97 +37,4 @@ namespace UE::Physics
 
 }
 
-
-namespace UE::Physics
-{
-	FVector3f MakeSolidBoxInertia(float InMass, const FVector3f& InSize)
-	{
-		FVector3f SizeSq = InSize * InSize;
-		return (InMass / 12.0f) * FVector3f(
-			SizeSq.Y + SizeSq.Z,
-			SizeSq.X + SizeSq.Z,
-			SizeSq.X + SizeSq.Y);
-	}
-
-	FRigidShapeInstancePtr MakeShapeInstanceFromAddParams(const FGeometryAddParams& InParams)
-	{
-		TArray<Chaos::FImplicitObjectPtr> Geoms;
-		Chaos::FShapesArray Shapes;
-		ChaosInterface::CreateGeometry(InParams, Geoms, Shapes);
-
-		if (Shapes.Num() != 1)
-		{
-			UE_LOGFMT(LogRigidPhysics, Warning, "MakeShapeInstanceFromAddParams is intended to be used for single shapes only. Have {0} shapes.", Shapes.Num());
-		}
-
-		if (Shapes.Num() > 0)
-		{
-			return FRigidShapeInstancePtr(Shapes[0].Release());
-		}
-
-		return FRigidShapeInstancePtr();
-	}
-
-	FRigidShapeInstancePtr MakeColliderFromAggGeom(
-		FKAggregateGeom& InAggGeom, 
-		const FVector3d& InScale,
-		ECollisionTraceFlag CollisionTraceFlag,
-		UPhysicalMaterial* InSimpleMaterial,
-		const TArrayView<UPhysicalMaterial*>& InComplexMaterials,
-		const TArrayView<FPhysicalMaterialMaskParams>& InComplexMaterialMasks,
-		const FTransform& InRelativeTransform)
-	{
-		// Not used - filters are set up after
-		FBodyCollisionData BodyCollisionData;
-
-		FGeometryAddParams AddParams;
-		AddParams.bDoubleSided = false;
-		AddParams.CollisionData = BodyCollisionData;
-		AddParams.CollisionTraceType = CollisionTraceFlag;
-		AddParams.Scale = InScale;
-		AddParams.SimpleMaterial = InSimpleMaterial;
-		AddParams.ComplexMaterials = InComplexMaterials;
-		AddParams.ComplexMaterialMasks = InComplexMaterialMasks;
-		AddParams.LocalTransform = InRelativeTransform;
-		AddParams.WorldTransform = FTransform::Identity;
-		AddParams.Geometry = &InAggGeom;
-		//AddParams.TriMeshGeometries = MakeArrayView(TriMeshGeometries);
-
-		return MakeShapeInstanceFromAddParams(AddParams);
-	}
-
-	FRigidShapeInstancePtr MakeBoxShape(const FVector3f& InSize, const FVector3f& InCenter, UPhysicalMaterial* InMaterial)
-	{
-		using namespace Chaos;
-
-		// Single Box
-		FKBoxElem BoxElem = FKBoxElem(InSize.X, InSize.Y, InSize.Z);
-		BoxElem.Center = FVector(InCenter);
-
-		FKAggregateGeom AggGeom;
-		AggGeom.AddElement(BoxElem);
-
-		FRigidShapeInstancePtr Collider = MakeColliderFromAggGeom(
-			AggGeom, 
-			FVector3d::One(),
-			ECollisionTraceFlag::CTF_UseSimpleAsComplex,
-			InMaterial,
-			{},
-			{},
-			FTransform3d::Identity);
-
-		if (Collider.IsValid())
-		{
-			const EFilterFlags FilterFlags = EFilterFlags::SimpleCollision | EFilterFlags::ContactNotify | EFilterFlags::ModifyContacts;
-			const Filter::FShapeFilterData ShapeFilter = Filter::FShapeFilterBuilder::BuildBlockAll(FilterFlags);
-			Collider->GetPerShapeData()->SetQueryEnabled(true);
-			Collider->GetPerShapeData()->SetSimEnabled(true);
-			Collider->GetPerShapeData()->SetShapeFilterData(ShapeFilter);
-		}
-
-		return Collider;
-	}
-}
-
-
 #endif // UE_RIGIDPHYSICS_API_ENABLED
\ No newline at end of file
diff --git a/Engine/Source/Runtime/Experimental/RigidPhysics/Public/RigidPhysics/RigidCollider.h b/Engine/Source/Runtime/Experimental/RigidPhysics/Public/RigidPhysics/RigidCollider.h
--- a/Engine/Source/Runtime/Experimental/RigidPhysics/Public/RigidPhysics/RigidCollider.h
+++ b/Engine/Source/Runtime/Experimental/RigidPhysics/Public/RigidPhysics/RigidCollider.h
@@ -113,13 +113,4 @@ namespace UE::Physics
 	};
 }
 
-namespace UE::Physics
-{
-	// Utility for Unit Tests until the API if more fully implemented
-	extern UE_INTERNAL RIGIDPHYSICS_API FRigidShapeInstancePtr MakeBoxShape(const FVector3f& InSize, const FVector3f& InCenter = FVector3f::Zero(), UPhysicalMaterial* InMaterial = nullptr);
-	extern UE_INTERNAL RIGIDPHYSICS_API FVector3f MakeSolidBoxInertia(float InMass, const FVector3f& InSize);
-
-}
-
-
 #endif
diff --git a/Engine/Source/Runtime/Experimental/RigidPhysics/RigidPhysics.Build.cs b/Engine/Source/Runtime/Experimental/RigidPhysics/RigidPhysics.Build.cs
--- a/Engine/Source/Runtime/Experimental/RigidPhysics/RigidPhysics.Build.cs
+++ b/Engine/Source/Runtime/Experimental/RigidPhysics/RigidPhysics.Build.cs
@@ -13,7 +13,6 @@ namespace UnrealBuildTool.Rules
                 new string[] {
                 "Core",
                 "CoreUObject",
-				"Engine",
 				}
 			);
 
