Allow the IAX CDN hit/miss detection code to run in test configuration.

#rnx

#Synced-CL 50282971

diff --git a/Engine/Source/Runtime/Experimental/IoStore/OnDemand/Private/OnDemandHttpClient.cpp b/Engine/Source/Runtime/Experimental/IoStore/OnDemand/Private/OnDemandHttpClient.cpp
--- a/Engine/Source/Runtime/Experimental/IoStore/OnDemand/Private/OnDemandHttpClient.cpp
+++ b/Engine/Source/Runtime/Experimental/IoStore/OnDemand/Private/OnDemandHttpClient.cpp
@@ -415,9 +415,9 @@ FMultiEndpointHttpClient::FHttpTicketId FMultiEndpointHttpClient::IssueRequest(F
 
 	UE::IoStore::HTTP::FRequest HttpRequest = EventLoop.Get(Url, ConnectionPool, &RequestParams);
 
-#if !UE_BUILD_TEST
+#if UE_TRACK_CDN_HIT_STATUS
 	HttpRequest.Header("pragma", "akamai-x-cache-on");
-#endif
+#endif // UE_TRACK_CDN_HIT_STATUS
 
 	check(Request.RequestHeaders.IsEmpty() || ((Request.RequestHeaders.Num() % 2) == 0));
 	if (!Request.RequestHeaders.IsEmpty())
diff --git a/Engine/Source/Runtime/Experimental/IoStore/OnDemand/Private/Statistics.h b/Engine/Source/Runtime/Experimental/IoStore/OnDemand/Private/Statistics.h
--- a/Engine/Source/Runtime/Experimental/IoStore/OnDemand/Private/Statistics.h
+++ b/Engine/Source/Runtime/Experimental/IoStore/OnDemand/Private/Statistics.h
@@ -17,7 +17,7 @@
 #endif
 
 #ifndef UE_TRACK_CDN_HIT_STATUS
-	#if UE_BUILD_TEST == 0 && UE_BUILD_SHIPPING == 0
+	#if UE_BUILD_SHIPPING == 0
 		#define UE_TRACK_CDN_HIT_STATUS 1
 	#else
 		#define UE_TRACK_CDN_HIT_STATUS 0
