Dataflow : Add TaskGraph and System to have a centralized place to launch tasks fro evaluators
- the taskgraph allows for launching both GT and worker thread tasks and have a tick processing with a budget

#Synced-CL 50431047

diff --git a/Engine/Source/Runtime/Experimental/Dataflow/Core/Private/Dataflow/DataflowSystem.cpp b/Engine/Source/Runtime/Experimental/Dataflow/Core/Private/Dataflow/DataflowSystem.cpp
new file mode 100644
--- /dev/null
+++ b/Engine/Source/Runtime/Experimental/Dataflow/Core/Private/Dataflow/DataflowSystem.cpp
@@ -0,0 +1,12 @@
+// Copyright Epic Games, Inc. All Rights Reserved.
+
+#include "Dataflow/DataflowSystem.h"
+
+namespace UE::Dataflow
+{
+	/*static*/ Dataflow::FSystem& FSystem::Get()
+	{
+		static Dataflow::FSystem Instance;
+		return Instance;
+	}
+}
diff --git a/Engine/Source/Runtime/Experimental/Dataflow/Core/Private/Dataflow/DataflowTaskGraph.cpp b/Engine/Source/Runtime/Experimental/Dataflow/Core/Private/Dataflow/DataflowTaskGraph.cpp
new file mode 100644
--- /dev/null
+++ b/Engine/Source/Runtime/Experimental/Dataflow/Core/Private/Dataflow/DataflowTaskGraph.cpp
@@ -0,0 +1,140 @@
+// Copyright Epic Games, Inc. All Rights Reserved.
+
+#include "Dataflow/DataflowTaskGraph.h"
+#include "Algo/Transform.h"
+
+#define LOCTEXT_NAMESPACE "DataflowTaskGraph"
+
+DEFINE_LOG_CATEGORY_STATIC(LogDataflowTaskGraph, Log, All);
+
+namespace UE::Dataflow
+{
+	bool FTask::IsCompleted() const
+	{
+		return Task.IsCompleted();
+	}
+
+	bool FTask::IsValid() const
+	{
+		return Task.IsValid();
+	}
+
+	bool FTaskGraph::FGameThreadTask::CanExecute() const
+	{
+		for (const FTask& Prerequisite : Prerequisites)
+		{
+			if (!Prerequisite.Task.IsCompleted())
+			{
+				return false;
+			}
+		}
+		return true;
+	}
+
+	bool FTaskGraph::FGameThreadTask::IsCancelled() const
+	{
+		return CancellationToken && CancellationToken->IsCanceled();
+	}
+
+	bool FTaskGraph::FGameThreadTask::TryExecute()
+	{
+		if (CanExecute())
+		{
+			const bool bIsCancelled = IsCancelled();
+			if (TaskBody && !bIsCancelled)
+			{
+				TaskBody();
+			}
+			CompletionTaskEvent.Trigger();
+			return true;
+		}
+		return false;
+	}
+
+	///////////////////////////////////////////////////////////////////////////////////////////
+
+	FTaskGraph::FTaskGraph()
+	{
+		// TODO : launch ticker only when first gamethread task is requested
+		TickDelegateHandle = FTSTicker::GetCoreTicker().AddTicker(FTickerDelegate::CreateRaw(this, &FTaskGraph::OnTick));
+	}
+
+	FTaskGraph::~FTaskGraph()
+	{
+		FTSTicker::GetCoreTicker().RemoveTicker(TickDelegateHandle);
+		TickDelegateHandle.Reset();
+	}
+
+	FTask FTaskGraph::LaunchGameThreadTask(const TCHAR* DebugName, TUniqueFunction<void()>&& TaskBody, const TArray<FTask>& Prerequisites, FCancellationTokenPtr CancellationToken)
+	{
+		GameThreadTasks.Emplace(
+			MakeShared<FGameThreadTask>(FGameThreadTask
+			{
+				.DebugName = DebugName,
+				.TaskBody = MoveTemp(TaskBody),
+				.Prerequisites = Prerequisites,
+				.CancellationToken = CancellationToken,
+				.CompletionTaskEvent = UE::Tasks::FTaskEvent{ DebugName },
+			}));
+		FTask ReturnTask;
+		ReturnTask.Task = GameThreadTasks.Last()->CompletionTaskEvent;
+		return ReturnTask;
+	}
+
+	FTask FTaskGraph::LaunchAnyThreadTask(const TCHAR* DebugName, TUniqueFunction<void()>&& TaskBody, const TArray<FTask>& Prerequisites, FCancellationTokenPtr CancellationToken)
+	{
+		TArray<UE::Tasks::FTask> UETasksPrerequisites;
+		Algo::Transform(Prerequisites, UETasksPrerequisites, [](const FTask& Task) { return Task.Task; });
+		FTask ReturnTask;
+		ReturnTask.Task = UE::Tasks::Launch(DebugName, MoveTemp(TaskBody), UETasksPrerequisites);
+		return ReturnTask;
+	}
+
+	bool FTaskGraph::OnTick(float DeltaTime)
+	{
+		UE_LOG(LogDataflowTaskGraph, Verbose, TEXT("FTaskGraph::Tick -- START"));
+
+		// TODO : should this be a parameter of the taskgraph ?
+		constexpr double TickTimeBudget = 5.0;
+
+		const double StartTimeMs = FPlatformTime::ToMilliseconds64(FPlatformTime::Cycles64());
+		double CurrentTimeMs = StartTimeMs;
+
+		TArray<FGameThreadTaskRef> NonExecutedTasks;
+
+		const int32 NumTasks = GameThreadTasks.Num();
+		NonExecutedTasks.Reserve(NumTasks);
+
+		bool bOverBudget = false;
+		for (int32 TaskIndex = 0; TaskIndex < NumTasks; ++TaskIndex)
+		{
+			if (bOverBudget)
+			{
+				NonExecutedTasks.Add(GameThreadTasks[TaskIndex]);
+			}
+			else if (!GameThreadTasks[TaskIndex]->TryExecute())
+			{
+				NonExecutedTasks.Add(GameThreadTasks[TaskIndex]);
+			}
+			else
+			{
+				// we only check time if the task has executed to make sure we don't run out of time just checking that tasks can run 
+				// and be stuck at always evelauting the same one with no progress
+				CurrentTimeMs = FPlatformTime::ToMilliseconds64(FPlatformTime::Cycles64());
+				if ((CurrentTimeMs - StartTimeMs) > TickTimeBudget)
+				{
+					bOverBudget = true;
+				}
+
+			}
+		}
+
+		// flip the arrays
+		GameThreadTasks = MoveTemp(NonExecutedTasks);
+
+		UE_LOG(LogDataflowTaskGraph, Verbose, TEXT("FTaskGraph::Tick -- END : time = [%.3f ms]"), (float)(CurrentTimeMs - StartTimeMs));
+		return true;
+	}
+};
+
+#undef LOCTEXT_NAMESPACE
diff --git a/Engine/Source/Runtime/Experimental/Dataflow/Core/Private/Dataflow/DataflowTaskGraph.h b/Engine/Source/Runtime/Experimental/Dataflow/Core/Private/Dataflow/DataflowTaskGraph.h
new file mode 100644
--- /dev/null
+++ b/Engine/Source/Runtime/Experimental/Dataflow/Core/Private/Dataflow/DataflowTaskGraph.h
@@ -0,0 +1,77 @@
+// Copyright Epic Games, Inc. All Rights Reserved.
+
+#pragma once
+
+#include "Tasks/Task.h"
+#include "Containers/Ticker.h"
+#include "Templates/SharedPointer.h"
+
+namespace UE::Dataflow
+{
+	// encapsulation of a task so that we cannot for example wait on them 
+	struct FTask
+	{
+	public:
+		bool IsCompleted() const;
+		bool IsValid() const;
+
+		// Explicitly delete the wait method to avoid future implementer to add it 
+		// Dtaaflow Tasks should not be awaited and should be listen to using a completion callback/task instead
+		void Wait() = delete;
+
+	private:
+		friend struct FTaskGraph;
+		UE::Tasks::FTask Task;
+	};
+
+	struct FTaskGraph
+	{
+	public:
+		FTaskGraph();
+		~FTaskGraph();
+
+		using FCancellationTokenPtr = TSharedPtr<UE::Tasks::FCancellationToken>;
+
+		/** 
+		* Launch a task to run on the game thread 
+		* Execution will happen during the next execution of tick method of the taskgraph when all dependent tasks have be completed
+		*/
+		FTask LaunchGameThreadTask(const TCHAR* DebugName, TUniqueFunction<void()>&& TaskBody, const TArray<FTask>& Prerequisites, FCancellationTokenPtr CancellationToken = {});
+
+		/**
+		* Launch a task to run on any thread
+		*/
+		FTask LaunchAnyThreadTask(const TCHAR* DebugName, TUniqueFunction<void()>&& TaskBody, const TArray<FTask>& Prerequisites, FCancellationTokenPtr CancellationToken = {});
+
+	private:
+		struct FGameThreadTask
+		{
+			/** Debug name of the task */
+			const TCHAR* DebugName = TEXT("DataflowGameThreadTask");
+
+			/** Function to execute when running the task */
+			TUniqueFunction<void()> TaskBody;
+
+			/** List of pPrerequisites */
+			TArray<FTask> Prerequisites;
+
+			/** optional cancellation token */
+			FCancellationTokenPtr CancellationToken;
+
+			/** Task event that completes when the body has executed */
+			UE::Tasks::FTaskEvent CompletionTaskEvent;
+
+			bool CanExecute() const;
+			bool IsCancelled() const;
+			bool TryExecute();
+		};
+		using FGameThreadTaskRef = TSharedRef<FGameThreadTask>;
+
+		FTSTicker::FDelegateHandle TickDelegateHandle;
+		bool OnTick(float DeltaTime);
+
+		TArray<FGameThreadTaskRef> GameThreadTasks;
+
+
+	};
+};
diff --git a/Engine/Source/Runtime/Experimental/Dataflow/Core/Public/Dataflow/DataflowSystem.h b/Engine/Source/Runtime/Experimental/Dataflow/Core/Public/Dataflow/DataflowSystem.h
new file mode 100644
--- /dev/null
+++ b/Engine/Source/Runtime/Experimental/Dataflow/Core/Public/Dataflow/DataflowSystem.h
@@ -0,0 +1,18 @@
+// Copyright Epic Games, Inc. All Rights Reserved.
+
+#pragma once
+
+#include "Dataflow/DataflowTaskGraph.h"
+
+namespace UE::Dataflow
+{
+	struct FSystem
+	{
+	public:
+		static Dataflow::FSystem& Get();
+
+		/* TaskGraph to schedule Dtaaflow specific tasks */
+		Dataflow::FTaskGraph TaskGraph;
+	};
+
+}
\ No newline at end of file
