Dataflow : Make the EvaluateNode method to return true if any of the output was evaluated

#Synced-CL 50431047

diff --git a/Engine/Source/Runtime/Experimental/Dataflow/Core/Private/Dataflow/DataflowNode.cpp b/Engine/Source/Runtime/Experimental/Dataflow/Core/Private/Dataflow/DataflowNode.cpp
--- a/Engine/Source/Runtime/Experimental/Dataflow/Core/Private/Dataflow/DataflowNode.cpp
+++ b/Engine/Source/Runtime/Experimental/Dataflow/Core/Private/Dataflow/DataflowNode.cpp
@@ -713,14 +713,16 @@ void FDataflowNode::Invalidate(const UE::Dataflow::FTimestamp& InModifiedTimesta
 	}
 }
 
-void FDataflowNode::EvaluateNode(UE::Dataflow::FContext& Context) const
+bool FDataflowNode::EvaluateNode(UE::Dataflow::FContext& Context) const
 {
 	// TODO(dataflow) : in the future we should also have the possibility to evaluate the watched output if needed by the context 
+	bool bWasEvaluatedOrForwarded = false;
 	if (IsActive())
 	{
 		if (ExpandedOutputs.IsEmpty())
 		{
 			Evaluate(Context, nullptr);
+			bWasEvaluatedOrForwarded = true;
 		}
 		else if (!IsFrozen())
 		{
@@ -732,6 +734,7 @@ void FDataflowNode::EvaluateNode(UE::Dataflow::FContext& Context) const
 					if (!Context.HasData(Output->CacheKey(), GetTimestamp()))
 					{
 						Evaluate(Context, Output);
+						bWasEvaluatedOrForwarded = true;
 					}
 				}
 			}
@@ -746,10 +749,12 @@ void FDataflowNode::EvaluateNode(UE::Dataflow::FContext& Context) const
 				if (const FDataflowInput* const PassthroughInput = Output->GetPassthroughInput())
 				{
 					Output->ForwardInput(PassthroughInput, Context);
+					bWasEvaluatedOrForwarded = true;
 				}
 			}
 		}
 	}
+	return bWasEvaluatedOrForwarded;
 }
 
 const FProperty* FDataflowNode::FindProperty(const UStruct* Struct, const void* InProperty, const FName& PropertyName, TArray<const FProperty*>* OutPropertyChain) const
diff --git a/Engine/Source/Runtime/Experimental/Dataflow/Core/Public/Dataflow/DataflowNode.h b/Engine/Source/Runtime/Experimental/Dataflow/Core/Public/Dataflow/DataflowNode.h
--- a/Engine/Source/Runtime/Experimental/Dataflow/Core/Public/Dataflow/DataflowNode.h
+++ b/Engine/Source/Runtime/Experimental/Dataflow/Core/Public/Dataflow/DataflowNode.h
@@ -360,8 +360,9 @@ PRAGMA_ENABLE_DEPRECATION_WARNINGS
 	* This evaluate the node based on the existing state of the outputs
 	* if the node has no output then evaluate will be called with a null output 
 	* otherwise Evaluate will be called for each connected output if it has an outdated cached value in teh context
+	* Return true if at least one output was evaluated or forwarded
 	*/
-	void EvaluateNode(UE::Dataflow::FContext& Context) const;
+	bool EvaluateNode(UE::Dataflow::FContext& Context) const;
 
 	virtual void Evaluate(UE::Dataflow::FContext& Context, const FDataflowOutput*) const { ensure(false); }
 
