Dataflow : change Dataflow taskgraph evaluator to use the new dataflow taskgraph and properly handle dependent assets
- Change the engine context EvaluateGraph to evaluate terminal node
- Fracture editor now use the taskgraph evaluator
- Made compiled graph to store a referenc eto the graph they were compiled from

#Synced-CL 50431047

diff --git a/Engine/Plugins/Experimental/ChaosEditor/Source/FractureEditor/Private/FractureEditorModeToolkit.cpp b/Engine/Plugins/Experimental/ChaosEditor/Source/FractureEditor/Private/FractureEditorModeToolkit.cpp
--- a/Engine/Plugins/Experimental/ChaosEditor/Source/FractureEditor/Private/FractureEditorModeToolkit.cpp
+++ b/Engine/Plugins/Experimental/ChaosEditor/Source/FractureEditor/Private/FractureEditorModeToolkit.cpp
@@ -353,17 +353,17 @@ FReply FFractureEditorModeToolkit::OnDataflowOverridesUpdateAsset()
 	{
 		if (FractureEditorModeToolkit::Private::bUseNewDataflowEvaluator)
 		{
-			auto OnEvaluationCompleted = [DataflowContext = DataflowContext](UE::Dataflow::FContext&)
+			/*auto OnEvaluationCompleted = [DataflowContext = DataflowContext](UE::Dataflow::FContext&)
 				{
 					if (DataflowContext)
 					{
 						DataflowContext->UpdateAssetFromTerminalNodes(true);
 					}
-				};
+				};*/
 			DataflowEvaluator.Reset();
 			if (DataflowContext)
 			{
-				DataflowEvaluator = DataflowContext->EvaluateGraph<UE::Dataflow::FDataflowTickEvaluator>(OnEvaluationCompleted);
+				DataflowEvaluator = DataflowContext->EvaluateGraph<UE::Dataflow::FDataflowTaskGraphEvaluator>(/*OnEvaluationCompleted*/{});
 				if (!DataflowEvaluator.IsValid())
 				{
 					UE_LOG(LogFractureTool, Error, TEXT("FractureMode: Failed to start evaluating Dataflow graph for geometry collection %s"), *FractureSettings->RestCollection->GetFName().ToString());
diff --git a/Engine/Source/Runtime/Experimental/Dataflow/Core/Private/Dataflow/DataflowCompiledGraph.cpp b/Engine/Source/Runtime/Experimental/Dataflow/Core/Private/Dataflow/DataflowCompiledGraph.cpp
--- a/Engine/Source/Runtime/Experimental/Dataflow/Core/Private/Dataflow/DataflowCompiledGraph.cpp
+++ b/Engine/Source/Runtime/Experimental/Dataflow/Core/Private/Dataflow/DataflowCompiledGraph.cpp
@@ -4,6 +4,15 @@
 
 namespace UE::Dataflow
 {
+	FCompiledGraph::FCompiledGraph(TSharedRef<const FGraph> InSourceGraph)
+		: SourceGraph(InSourceGraph)
+	{}
+
+	TSharedRef<const FGraph> FCompiledGraph::GetSourceGraph() const
+	{
+		return SourceGraph;
+	}
+
 	int32 FCompiledGraph::StartTask()
 	{
 		return Tasks.Emplace(FTask{ .CommandIndex = Commands.Num() });
diff --git a/Engine/Source/Runtime/Experimental/Dataflow/Core/Private/Dataflow/DataflowContextEvaluator.cpp b/Engine/Source/Runtime/Experimental/Dataflow/Core/Private/Dataflow/DataflowContextEvaluator.cpp
--- a/Engine/Source/Runtime/Experimental/Dataflow/Core/Private/Dataflow/DataflowContextEvaluator.cpp
+++ b/Engine/Source/Runtime/Experimental/Dataflow/Core/Private/Dataflow/DataflowContextEvaluator.cpp
@@ -4,6 +4,7 @@
 #include "Dataflow/DataflowGraph.h"
 #include "Dataflow/DataflowNode.h"
 #include "Dataflow/DataflowCompiledGraph.h"
+#include "Dataflow/DataflowSystem.h"
 
 #define LOCTEXT_NAMESPACE "DataflowEvaluator"
 
@@ -279,7 +280,14 @@ namespace UE::Dataflow
 	void FThreadSafeProgressTracker::SetProgress(float InProgress, FName InNodeName)
 	{
 		FScopeLock Lock(&ProgressLock);
-		Progress = InProgress;
+		Progress = FMath::Clamp(InProgress, 0.0f, 1.0f);
+		NodeName = InNodeName;
+	}
+
+	void FThreadSafeProgressTracker::AddProgress(float InProgressDelta, FName InNodeName)
+	{
+		FScopeLock Lock(&ProgressLock);
+		Progress = FMath::Clamp(Progress + InProgressDelta, 0.0f, 1.0f);
 		NodeName = InNodeName;
 	}
 
@@ -290,25 +298,32 @@ namespace UE::Dataflow
 	{
 	}
 
-	void FDataflowDirectEvaluator::Start(TSharedRef<const FGraph> InDataflowGraph, TSharedRef<const FCompiledGraph> InCompiledGraph, const FGuid& InNodeId, FOnPostEvaluationFunction OnPostEvaluation)
+	void FDataflowDirectEvaluator::Start(const FDataflowEvaluatorParameters& InParams)
 	{
 		bIsCanceled = false;
 		bIsRunning = true;
 
 		int32 LastNodeIndex = INDEX_NONE;
 		FGuid NodeId;
+		TSharedRef<const FGraph> DataflowGraph = InParams.CompiledGraph->GetSourceGraph();
 		do
 		{
-			NodeId = InCompiledGraph->GetNextNodeId(LastNodeIndex++);
+			NodeId = InParams.CompiledGraph->GetNextNodeId(LastNodeIndex++);
 			if (NodeId.IsValid())
 			{
-				if (TSharedPtr<const FDataflowNode> Node = InDataflowGraph->FindBaseNode(NodeId))
+				if (TSharedPtr<const FDataflowNode> Node = DataflowGraph->FindBaseNode(NodeId))
 				{
-					ProgressTracker.SetProgress((float)LastNodeIndex, (float)InCompiledGraph->GetNumNodes(), Node->GetName());
+					ProgressTracker.SetProgress((float)LastNodeIndex, (float)InParams.CompiledGraph->GetNumNodes(), Node->GetName());
 
 					if (ensure(!Node->EvaluateOnGameThreadOnly()))
 					{
-						Node->EvaluateNode(OwningContext);
+						if (Node->EvaluateNode(OwningContext))
+						{
+							if (Node->IsTerminal() && InParams.OnTerminalNodeEvaluated.IsSet())
+							{
+								InParams.OnTerminalNodeEvaluated(OwningContext, *Node);
+							}
+						}
 					}
 				}
 			}
@@ -317,9 +332,9 @@ namespace UE::Dataflow
 		bIsRunning = false;
 
 		// No more nodes to process
-		if (!bIsCanceled && OnPostEvaluation.IsSet())
+		if (!bIsCanceled && InParams.OnPostEvaluation.IsSet())
 		{
-			OnPostEvaluation(OwningContext);
+			InParams.OnPostEvaluation(OwningContext);
 		}
 	}
 
@@ -365,20 +380,22 @@ namespace UE::Dataflow
 		LastNodeIndex = INDEX_NONE;
 		NodeId = FGuid();
 		PostEvaluationFunction.Reset();
+		TerminalNodeEvaluatedFunction.Reset();
 		bIsRunning = false;
 
 		FTSTicker::GetCoreTicker().RemoveTicker(TickDelegateHandle);
 		TickDelegateHandle.Reset();
 	}
 
-	void FDataflowTickEvaluator::Start(TSharedRef<const FGraph> InDataflowGraph, TSharedRef<const FCompiledGraph> InCompiledGraph, const FGuid& InNodeId, FOnPostEvaluationFunction OnPostEvaluation)
+	void FDataflowTickEvaluator::Start(const FDataflowEvaluatorParameters& InParams)
 	{
 		Reset();
 
-		DataflowGraph = InDataflowGraph;
-		CompiledGraph = InCompiledGraph;
-		NodeId = InNodeId;
-		PostEvaluationFunction = OnPostEvaluation;
+		DataflowGraph = InParams.CompiledGraph->GetSourceGraph();
+		CompiledGraph = InParams.CompiledGraph;
+		NodeId = InParams.NodeId;
+		PostEvaluationFunction = InParams.OnPostEvaluation;
+		TerminalNodeEvaluatedFunction = InParams.OnTerminalNodeEvaluated;
 		LastNodeIndex = INDEX_NONE;
 
 		TickDelegateHandle = FTSTicker::GetCoreTicker().AddTicker(FTickerDelegate::CreateRaw(this, &FDataflowTickEvaluator::Tick));
@@ -448,7 +465,13 @@ namespace UE::Dataflow
 
 				const double StartEvalTimeMs = FPlatformTime::ToMilliseconds64(FPlatformTime::Cycles64());
 
-				Node->EvaluateNode(OwningContext);
+				if (Node->EvaluateNode(OwningContext))
+				{
+					if (Node->IsTerminal() && TerminalNodeEvaluatedFunction.IsSet())
+					{
+						TerminalNodeEvaluatedFunction(OwningContext, *Node);
+					}
+				}
 
 				const double EndEvalTimeMs = FPlatformTime::ToMilliseconds64(FPlatformTime::Cycles64());
 
@@ -480,33 +503,35 @@ namespace UE::Dataflow
 	FDataflowTaskGraphEvaluator::FDataflowTaskGraphEvaluator(FContext& InOwningContext)
 		: OwningContext(InOwningContext)
 		, CancellationToken(MakeShared<UE::Tasks::FCancellationToken>())
+		, ProgressTracker(MakeShared<FThreadSafeProgressTracker>())
 	{
 	}
 
-	void FDataflowTaskGraphEvaluator::Start(TSharedRef<const FGraph> InDataflowGraph, TSharedRef<const FCompiledGraph> InCompiledGraph, const FGuid& InNodeId, FOnPostEvaluationFunction OnPostEvaluation)
+	void FDataflowTaskGraphEvaluator::Start(const FDataflowEvaluatorParameters& InParams)
 	{
-		check(!IsInGameThread());
-
 		// build the taskgraph from the compiled graph
-		const int32 NumTasks = InCompiledGraph->GetNumTasks();
+		const int32 NumTasks = InParams.CompiledGraph->GetNumTasks();
 
 		CancellationToken = MakeShared<UE::Tasks::FCancellationToken>();
-		TArray<UE::Tasks::FTask> AllTasks;
+		TArray<FTask> AllTasks;
 		AllTasks.SetNum(NumTasks);
 
 		TArray<FGuid> NodeIds;
 		TArray<int32> ParentTaskIds;
 		TArray<TWeakPtr<const FDataflowNode>> NodeWeakPtrs;
-		TArray<UE::Tasks::FTask> PreRequisites;
+		TArray<FTask> PreRequisites;
 
+		TSharedRef<const FGraph> DataflowGraph = InParams.CompiledGraph->GetSourceGraph();
+
+		const float PerTaskProgress = (NumTasks > 0) ? (1.f / (float)NumTasks) : 0;
 		for (int32 TaskIndex = 0; TaskIndex < NumTasks; ++TaskIndex)
 		{
-			InCompiledGraph->GetTasksNodesAndDependencies(TaskIndex, NodeIds, ParentTaskIds);
+			InParams.CompiledGraph->GetTasksNodesAndDependencies(TaskIndex, NodeIds, ParentTaskIds);
 
 			NodeWeakPtrs.Reset();
 			for (const FGuid& NodeId : NodeIds)
 			{
-				if (TSharedPtr<const FDataflowNode> Node = InDataflowGraph->FindBaseNode(NodeId))
+				if (TSharedPtr<const FDataflowNode> Node = DataflowGraph->FindBaseNode(NodeId))
 				{
 					NodeWeakPtrs.Add(Node);
 				}
@@ -520,35 +545,53 @@ namespace UE::Dataflow
 					PreRequisites.Add(AllTasks[ParentTaskId]);
 				}
 			}
-			// TODO : maybe we shaould have the context as a shared Ptr
-			auto TaskBody = [NodeWeakPtrs = MoveTemp(NodeWeakPtrs), CancellationToken = CancellationToken, &OwningContext = OwningContext]()
+			// TODO : maybe we should have the context as a shared Ptr
+			auto TaskBody = 
+				[NodeWeakPtrs = MoveTemp(NodeWeakPtrs), CancellationToken = CancellationToken, 
+				PerTaskProgress, ProgressTracker = ProgressTracker, &OwningContext = OwningContext, OnTerminalNodeEvaluated = InParams.OnTerminalNodeEvaluated]()
 				{
-					if (!CancellationToken->IsCanceled())
+					UE_LOG(LogDataflowEvaluator, Verbose, TEXT("FContextEvaluator::Task -- START"));
+
+					const int32 NumNodes = NodeWeakPtrs.Num();
+					const float PerNodeProgress = (NumNodes > 0) ? (PerTaskProgress / (float)NumNodes): 0;
+					for (TWeakPtr<const FDataflowNode> NodeWeakPtr : NodeWeakPtrs)
 					{
-						for (TWeakPtr<const FDataflowNode> NodeWeakPtr : NodeWeakPtrs)
+						if (TSharedPtr<const FDataflowNode> NodePtr = NodeWeakPtr.Pin())
 						{
-							if (TSharedPtr<const FDataflowNode> NodePtr = NodeWeakPtr.Pin())
+							UE_LOG(LogDataflowEvaluator, Verbose, TEXT("FContextEvaluator::Task -- EVALNODE [%s]"), *NodePtr->GetName().ToString());
+							ProgressTracker->AddProgress(PerNodeProgress, NodePtr->GetName());
+							
+							if (NodePtr->EvaluateNode(OwningContext))
 							{
-								NodePtr->EvaluateNode(OwningContext);
-							}
-							if (CancellationToken->IsCanceled())
-							{
-								break;
+								// todo : may need to ensure this is running on the game thread once we can execute on aritrary threads
+								if (NodePtr->IsTerminal() && OnTerminalNodeEvaluated.IsSet())
+								{
+									OnTerminalNodeEvaluated(OwningContext, *NodePtr);
+								}
 							}
 						}
+						if (CancellationToken->IsCanceled())
+						{
+							UE_LOG(LogDataflowEvaluator, Verbose, TEXT("FContextEvaluator::Task -- CANCELLED"));
+							break;
+						}
 					}
+					UE_LOG(LogDataflowEvaluator, Verbose, TEXT("FContextEvaluator::Task -- END"));
 				};
-			AllTasks[TaskIndex] = UE::Tasks::Launch(TEXT("DataflowEvaluationTask"), MoveTemp(TaskBody), PreRequisites);
+			AllTasks[TaskIndex] = Dataflow::FSystem::Get().TaskGraph.LaunchGameThreadTask(TEXT("DataflowEvaluationTask"), MoveTemp(TaskBody), PreRequisites, CancellationToken);
 		}
 
-		auto CompletionTaskBody = [OnPostEvaluation, &OwningContext = OwningContext]()
+		auto CompletionTaskBody = [OnPostEvaluation = InParams.OnPostEvaluation, &OwningContext = OwningContext]()
 			{
+				UE_LOG(LogDataflowEvaluator, Verbose, TEXT("FContextEvaluator::PostEvalTask -- START"));
 				if (OnPostEvaluation.IsSet())
 				{
+					UE_LOG(LogDataflowEvaluator, Verbose, TEXT("FContextEvaluator::PostEvalTask -- CALL"));
 					OnPostEvaluation(OwningContext);
 				}
+				UE_LOG(LogDataflowEvaluator, Verbose, TEXT("FContextEvaluator::PostEvalTask -- END"));
 			};
-		CompletionTask = UE::Tasks::Launch(TEXT("DataflowCompletionTask"), MoveTemp(CompletionTaskBody), AllTasks);
+		CompletionTask = MakeShared<FTask>(Dataflow::FSystem::Get().TaskGraph.LaunchGameThreadTask(TEXT("DataflowCompletionTask"), MoveTemp(CompletionTaskBody), AllTasks));
 	}
 
 	void FDataflowTaskGraphEvaluator::Cancel()
@@ -558,17 +601,17 @@ namespace UE::Dataflow
 
 	bool FDataflowTaskGraphEvaluator::IsRunning() const
 	{
-		return !CompletionTask.IsCompleted();
+		return CompletionTask && !CompletionTask->IsCompleted();
 	}
 
 	float FDataflowTaskGraphEvaluator::GetProgress() const
 	{
-		return ProgressTracker.GetProgress();
+		return ProgressTracker->GetProgress();
 	}
 
 	FText FDataflowTaskGraphEvaluator::GetProgressMessage() const
 	{
-		return ProgressTracker.GetProgressMessage();
+		return ProgressTracker->GetProgressMessage();
 	}
 };
 
diff --git a/Engine/Source/Runtime/Experimental/Dataflow/Core/Public/Dataflow/DataflowCompiledGraph.h b/Engine/Source/Runtime/Experimental/Dataflow/Core/Public/Dataflow/DataflowCompiledGraph.h
--- a/Engine/Source/Runtime/Experimental/Dataflow/Core/Public/Dataflow/DataflowCompiledGraph.h
+++ b/Engine/Source/Runtime/Experimental/Dataflow/Core/Public/Dataflow/DataflowCompiledGraph.h
@@ -5,9 +5,12 @@
 #include "Containers/Array.h"
 #include "HAL/Platform.h"
 #include "Misc/Guid.h"
+#include "Templates/SharedPointer.h"
 
 namespace UE::Dataflow
 {
+	class FGraph;
+
 	/** 
 	* This structure stores a compiled graph
 	* A compiled graph allow a dataflow evaluator to evaluate a graph efficiently without having to parse the graph again
@@ -22,6 +25,8 @@ namespace UE::Dataflow
 	struct FCompiledGraph
 	{
 	public:
+		DATAFLOWCORE_API FCompiledGraph(TSharedRef<const FGraph> InSourceGraph);
+
 		DATAFLOWCORE_API int32 StartTask();
 		DATAFLOWCORE_API void AddNode(FGuid NodeId);
 		DATAFLOWCORE_API void AddDependency(int32 TaskId);
@@ -33,6 +38,8 @@ namespace UE::Dataflow
 		DATAFLOWCORE_API int32 GetNumTasks() const;
 		DATAFLOWCORE_API void GetTasksNodesAndDependencies(int32 TaskId, TArray<FGuid>& OutNodeIds, TArray<int32>& OutParentTaskIds) const;
 
+		DATAFLOWCORE_API TSharedRef<const FGraph> GetSourceGraph() const;
+		
 	private:
 		enum class ECommandType
 		{
@@ -76,5 +83,7 @@ namespace UE::Dataflow
 
 		/** List of commands (with reference to the task and node arrays) */
 		TArray<FCommand> Commands;
+
+		TSharedRef<const FGraph> SourceGraph;
 	};
 }
\ No newline at end of file
diff --git a/Engine/Source/Runtime/Experimental/Dataflow/Core/Public/Dataflow/DataflowContextEvaluator.h b/Engine/Source/Runtime/Experimental/Dataflow/Core/Public/Dataflow/DataflowContextEvaluator.h
--- a/Engine/Source/Runtime/Experimental/Dataflow/Core/Public/Dataflow/DataflowContextEvaluator.h
+++ b/Engine/Source/Runtime/Experimental/Dataflow/Core/Public/Dataflow/DataflowContextEvaluator.h
@@ -4,11 +4,11 @@
 
 #include "Async/TaskGraphInterfaces.h"
 #include "Containers/Array.h"
+#include "Containers/Ticker.h"
 #include "Misc/Guid.h"
+#include "Tasks/Task.h"
 #include "Templates/Function.h"
 #include "Templates/SharedPointer.h"
-#include "Containers/Ticker.h"
-#include "Tasks/Task.h"
 
 struct FDataflowNode;
 struct FDataflowOutput;
@@ -21,6 +21,7 @@ namespace UE::Dataflow
 	class FContext;
 
 	using FOnPostEvaluationFunction = TFunction<void(FContext&)>;
+	using FOnTerminalNodeEvaluatedFunction = TFunction<void(FContext&, const FDataflowNode& Node)>;
 
 	/*
 	* Asynchronously evaluate a dataflow nodes
@@ -84,10 +85,39 @@ namespace UE::Dataflow
 		TSet<FNodeId> CompletedTasks;
 	};
 
+	struct FDataflowEvaluatorParameters
+	{
+		FDataflowEvaluatorParameters(TSharedRef<const FCompiledGraph> InCompiledGraph)
+			: CompiledGraph(InCompiledGraph)
+		{}
+
+		/** Compiled graph to evaluate */
+		TSharedRef<const FCompiledGraph> CompiledGraph;
+
+		/** 
+		* Optional Id of the node to evaluate in particular
+		* Evaluator will try to its best to evaluate only the relevant nodes that leads to this one 
+		* If not specified the entire graph will be evaluated
+		*/
+		FGuid NodeId;
+
+		/** 
+		* Optional function to call when evaluation is done 
+		*/
+		FOnPostEvaluationFunction OnPostEvaluation;
+
+		/** 
+		* Optional function to call each time after a "terminal" node has been evaluated
+		* A terminal node is a type of node that has its IsTerminal() method to return true 
+		* (see FDataflowTerminalNode, FDataflowSubGraphOutputNode, FDataflowExecutionNode)
+		*/
+		FOnTerminalNodeEvaluatedFunction OnTerminalNodeEvaluated;
+	};
+
 	struct IDataflowEvaluator
 	{
 		virtual ~IDataflowEvaluator() = default;
-		virtual void Start(TSharedRef<const FGraph> InDataflowGraph, TSharedRef<const FCompiledGraph> InCompiledGraph, const FGuid& InNodeId, FOnPostEvaluationFunction OnPostEvaluation) = 0;
+		virtual void Start(const FDataflowEvaluatorParameters& InParams) = 0;
 		virtual void Cancel() = 0;
 
 		virtual bool IsRunning() const = 0;
@@ -108,6 +138,7 @@ namespace UE::Dataflow
 
 		void SetProgress(float InCurrent, float InTotal, FName InNodeName);
 		void SetProgress(float InProgress, FName InNodeName);
+		void AddProgress(float InProgressDelta, FName InNodeName);
 	
 	private:
 		mutable FCriticalSection ProgressLock;
@@ -128,7 +159,7 @@ namespace UE::Dataflow
 		virtual ~FDataflowDirectEvaluator() = default;
 
 		//~ Begin IDataflowEvaluator interface
-		virtual void Start(TSharedRef<const FGraph> InDataflowGraph, TSharedRef<const FCompiledGraph> InCompiledGraph, const FGuid& InNodeId, FOnPostEvaluationFunction OnPostEvaluation) override;
+		virtual void Start(const FDataflowEvaluatorParameters& StartParams) override;
 		virtual void Cancel() override;
 		virtual bool IsRunning() const override;
 		virtual float GetProgress() const override;
@@ -155,7 +186,7 @@ namespace UE::Dataflow
 		DATAFLOWCORE_API virtual ~FDataflowTickEvaluator();
 
 		//~ Begin IDataflowEvaluator interface
-		virtual void Start(TSharedRef<const FGraph> InDataflowGraph, TSharedRef<const FCompiledGraph> InCompiledGraph, const FGuid& InNodeId, FOnPostEvaluationFunction OnPostEvaluation) override;
+		virtual void Start(const FDataflowEvaluatorParameters& InParams) override;
 		virtual void Cancel() override;
 		virtual bool IsRunning() const override;
 		virtual float GetProgress() const override;
@@ -172,6 +203,7 @@ namespace UE::Dataflow
 		TSharedPtr<const FCompiledGraph> CompiledGraph;
 		TSharedPtr<const FGraph> DataflowGraph;
 		FOnPostEvaluationFunction PostEvaluationFunction;
+		FOnTerminalNodeEvaluatedFunction TerminalNodeEvaluatedFunction;
 		FGuid NodeId;
 
 		FTSTicker::FDelegateHandle TickDelegateHandle;
@@ -184,6 +216,8 @@ namespace UE::Dataflow
 		FThreadSafeProgressTracker ProgressTracker;
 	};
 
+	struct FTask;
+
 	/**
 	* This evaluator will evaluate the graph in one go when Start is called
 	* This is useful for running simulation graph for example that needs to execute in exactly one frame
@@ -197,7 +231,7 @@ namespace UE::Dataflow
 		virtual ~FDataflowTaskGraphEvaluator() = default;
 
 		//~ Begin IDataflowEvaluator interface
-		virtual void Start(TSharedRef<const FGraph> InDataflowGraph, TSharedRef<const FCompiledGraph> InCompiledGraph, const FGuid& InNodeId, FOnPostEvaluationFunction OnPostEvaluation) override;
+		virtual void Start(const FDataflowEvaluatorParameters& InParams) override;
 		virtual void Cancel() override;
 		virtual bool IsRunning() const override;
 		virtual float GetProgress() const override;
@@ -207,10 +241,10 @@ namespace UE::Dataflow
 	private:
 		FContext& OwningContext;
 
-		UE::Tasks::FTask CompletionTask;
+		TSharedPtr<FTask> CompletionTask;
 		TSharedRef<UE::Tasks::FCancellationToken> CancellationToken;
 
-		FThreadSafeProgressTracker ProgressTracker;
+		TSharedRef<FThreadSafeProgressTracker> ProgressTracker;
 	};
 }
 
diff --git a/Engine/Source/Runtime/Experimental/Dataflow/Engine/Private/Dataflow/DataflowCompiler.cpp b/Engine/Source/Runtime/Experimental/Dataflow/Engine/Private/Dataflow/DataflowCompiler.cpp
--- a/Engine/Source/Runtime/Experimental/Dataflow/Engine/Private/Dataflow/DataflowCompiler.cpp
+++ b/Engine/Source/Runtime/Experimental/Dataflow/Engine/Private/Dataflow/DataflowCompiler.cpp
@@ -291,8 +291,14 @@ namespace UE::Dataflow
 
 	/*static*/ bool FCompiler::CompileGraph(const UEdGraph& GraphObject, const FCompilerParameters& Params, FCompilerResults& Results)
 	{
+		const UDataflow* Dataflow = UDataflow::GetDataflowAssetFromEdGraph(&GraphObject);
+		if (!Dataflow || !Dataflow->Dataflow.IsValid())
+		{
+			return false;
+		}
+
 		// Now build the command buffer
-		TSharedRef<FCompiledGraph> CompiledGraph = MakeShared<FCompiledGraph>();
+		TSharedRef<FCompiledGraph> CompiledGraph = MakeShared<FCompiledGraph>(Dataflow->Dataflow.ToSharedRef());
 		Results.AddCompiledGraph(&GraphObject, CompiledGraph);
 
 		// Find first the left nost thread as start node for the node thread
diff --git a/Engine/Source/Runtime/Experimental/Dataflow/Engine/Private/Dataflow/DataflowObjectInterface.cpp b/Engine/Source/Runtime/Experimental/Dataflow/Engine/Private/Dataflow/DataflowObjectInterface.cpp
--- a/Engine/Source/Runtime/Experimental/Dataflow/Engine/Private/Dataflow/DataflowObjectInterface.cpp
+++ b/Engine/Source/Runtime/Experimental/Dataflow/Engine/Private/Dataflow/DataflowObjectInterface.cpp
@@ -12,57 +12,45 @@ namespace UE::Dataflow
 
 	namespace Private
 	{
-		bool EvaluateGraphInternal(TSharedRef<IDataflowEvaluator> Evaluator, UObject* Owner, FOnPostEvaluationFunction OnPostEvaluation)
+		void UpdateAssetFromNode(FContext& Context, const FDataflowNode& Node, UObject* Asset)
 		{
-			if (UDataflow* DataflowObject = InstanceUtils::GetDataflowAssetFromObject(Owner))
+			if (const FDataflowTerminalNode* TerminalNode = Node.AsType<const FDataflowTerminalNode>())
 			{
-				if (TSharedPtr<FGraph> Graph = DataflowObject->GetDataflow())
+				// Note: If the node is deactivated and has any outputs, then these outputs might still need to be forwarded.
+				//       Therefore the Evaluate method has to be called for whichever value of bActive.
+				//       This however isn't the case of SetAssetValue() for which the active state needs to be checked before the call.
+				if (TerminalNode->IsActive())
 				{
-					if (TSharedPtr<const UE::Dataflow::FCompiledGraph> CompiledGraph = DataflowObject->CompileGraphIfNeeded())
-					{
-						Evaluator->Start(Graph.ToSharedRef(), CompiledGraph.ToSharedRef(), FGuid(), OnPostEvaluation);
-						return true;
-					}
+					TerminalNode->SetAssetValue(Asset, Context);
 				}
 			}
-			return false;
 		}
 
-		void UpdateAssetFromTerminalNodes(FContext& Context, UObject* Asset, bool bUseAllTerminals)
+		bool EvaluateGraphInternal(TSharedRef<IDataflowEvaluator> Evaluator, UObject* Owner, FOnPostEvaluationFunction OnPostEvaluation)
 		{
-			if (UDataflow* DataflowObject = InstanceUtils::GetDataflowAssetFromObject(Asset))
+			if (UDataflow* DataflowObject = InstanceUtils::GetDataflowAssetFromObject(Owner))
 			{
 				if (TSharedPtr<FGraph> Graph = DataflowObject->GetDataflow())
 				{
-					TArray<TSharedPtr<FDataflowNode>> TerminalNodes;
-					if (bUseAllTerminals)
-					{
-						TerminalNodes = Graph->GetFilteredNodes(FDataflowTerminalNode::StaticType());
-					}
-					else
+					if (TSharedPtr<const UE::Dataflow::FCompiledGraph> CompiledGraph = DataflowObject->CompileGraphIfNeeded())
 					{
-						const FName TerminalNodeName = InstanceUtils::GetTerminalNodeNameFromObject(Asset);
-						if (const TSharedPtr<FDataflowNode> TerminalNode = Graph->FindFilteredNode(FDataflowTerminalNode::StaticType(), TerminalNodeName))
-						{
-							TerminalNodes.Add(TerminalNode);
-						}
-					}
+						TWeakObjectPtr<UObject> WeakOwner = Owner;
 
-					for (TSharedPtr<FDataflowNode> Node : TerminalNodes)
-					{
-						if (const FDataflowTerminalNode* TerminalNode = Node ? Node->AsType<const FDataflowTerminalNode>() : nullptr)
-						{
-							// Note: If the node is deactivated and has any outputs, then these outputs might still need to be forwarded.
-							//       Therefore the Evaluate method has to be called for whichever value of bActive.
-							//       This however isn't the case of SetAssetValue() for which the active state needs to be checked before the call.
-							if (TerminalNode->IsActive() && Asset)
+						FDataflowEvaluatorParameters Params(CompiledGraph.ToSharedRef());
+						Params.OnPostEvaluation = OnPostEvaluation;
+						Params.OnTerminalNodeEvaluated = [WeakOwner](FContext& Context, const FDataflowNode& Node)
 							{
-								TerminalNode->SetAssetValue(Asset, Context);
-							}
-						}
+								if (TStrongObjectPtr<UObject> Owner = WeakOwner.Pin())
+								{
+									UpdateAssetFromNode(Context, Node, Owner.Get());
+								}
+							};
+						Evaluator->Start(Params);
+						return true;
 					}
 				}
 			}
+			return false;
 		}
 	}
 }
diff --git a/Engine/Source/Runtime/Experimental/Dataflow/Engine/Public/Dataflow/DataflowObjectInterface.h b/Engine/Source/Runtime/Experimental/Dataflow/Engine/Public/Dataflow/DataflowObjectInterface.h
--- a/Engine/Source/Runtime/Experimental/Dataflow/Engine/Public/Dataflow/DataflowObjectInterface.h
+++ b/Engine/Source/Runtime/Experimental/Dataflow/Engine/Public/Dataflow/DataflowObjectInterface.h
@@ -48,11 +48,6 @@ namespace UE::Dataflow
 			return {};
 			
 		}
-
-		void UpdateAssetFromTerminalNodes(bool bUseAllTerminals)
-		{
-			Private::UpdateAssetFromTerminalNodes(*this, Owner, bUseAllTerminals);
-		}
 	};
 
 	typedef TEngineContext<FContextSingle> FEngineContext;
