Dataflow : fix missing export causing some builds to fail

#Synced-CL 50431047

diff --git a/Engine/Source/Runtime/Experimental/Dataflow/Core/Public/Dataflow/DataflowContextEvaluator.h b/Engine/Source/Runtime/Experimental/Dataflow/Core/Public/Dataflow/DataflowContextEvaluator.h
--- a/Engine/Source/Runtime/Experimental/Dataflow/Core/Public/Dataflow/DataflowContextEvaluator.h
+++ b/Engine/Source/Runtime/Experimental/Dataflow/Core/Public/Dataflow/DataflowContextEvaluator.h
@@ -13,6 +13,8 @@
 struct FDataflowNode;
 struct FDataflowOutput;
 
+#define UE_API DATAFLOWCORE_API
+
 namespace UE::Dataflow
 {
 	struct FCompiledGraph;
@@ -155,15 +157,15 @@ namespace UE::Dataflow
 	struct FDataflowDirectEvaluator : public IDataflowEvaluator
 	{
 	public:
-		DATAFLOWCORE_API FDataflowDirectEvaluator(FContext& InOwningContext);
+		UE_API FDataflowDirectEvaluator(FContext& InOwningContext);
 		virtual ~FDataflowDirectEvaluator() = default;
 
 		//~ Begin IDataflowEvaluator interface
-		virtual void Start(const FDataflowEvaluatorParameters& StartParams) override;
-		virtual void Cancel() override;
-		virtual bool IsRunning() const override;
-		virtual float GetProgress() const override;
-		virtual FText GetProgressMessage() const override;
+		UE_API virtual void Start(const FDataflowEvaluatorParameters& StartParams) override;
+		UE_API virtual void Cancel() override;
+		UE_API virtual bool IsRunning() const override;
+		UE_API virtual float GetProgress() const override;
+		UE_API virtual FText GetProgressMessage() const override;
 		//~ End IDataflowEvaluator interface
 
 	private:
@@ -182,15 +184,15 @@ namespace UE::Dataflow
 	*/
 	struct FDataflowTickEvaluator: public IDataflowEvaluator
 	{
-		DATAFLOWCORE_API FDataflowTickEvaluator(FContext& InOwningContext);
-		DATAFLOWCORE_API virtual ~FDataflowTickEvaluator();
+		UE_API FDataflowTickEvaluator(FContext& InOwningContext);
+		UE_API virtual ~FDataflowTickEvaluator();
 
 		//~ Begin IDataflowEvaluator interface
-		virtual void Start(const FDataflowEvaluatorParameters& InParams) override;
-		virtual void Cancel() override;
-		virtual bool IsRunning() const override;
-		virtual float GetProgress() const override;
-		virtual FText GetProgressMessage() const override;
+		UE_API virtual void Start(const FDataflowEvaluatorParameters& InParams) override;
+		UE_API virtual void Cancel() override;
+		UE_API virtual bool IsRunning() const override;
+		UE_API virtual float GetProgress() const override;
+		UE_API virtual FText GetProgressMessage() const override;
 		//~ End IDataflowEvaluator interface
 
 	private:
@@ -227,15 +229,15 @@ namespace UE::Dataflow
 	struct FDataflowTaskGraphEvaluator : public IDataflowEvaluator
 	{
 	public:
-		DATAFLOWCORE_API FDataflowTaskGraphEvaluator(FContext& InOwningContext);
+		UE_API FDataflowTaskGraphEvaluator(FContext& InOwningContext);
 		virtual ~FDataflowTaskGraphEvaluator() = default;
 
 		//~ Begin IDataflowEvaluator interface
-		virtual void Start(const FDataflowEvaluatorParameters& InParams) override;
-		virtual void Cancel() override;
-		virtual bool IsRunning() const override;
-		virtual float GetProgress() const override;
-		virtual FText GetProgressMessage() const override;
+		UE_API virtual void Start(const FDataflowEvaluatorParameters& InParams) override;
+		UE_API virtual void Cancel() override;
+		UE_API virtual bool IsRunning() const override;
+		UE_API virtual float GetProgress() const override;
+		UE_API virtual FText GetProgressMessage() const override;
 		//~ End IDataflowEvaluator interface
 
 	private:
@@ -248,4 +250,4 @@ namespace UE::Dataflow
 	};
 }
 
-
+#undef UE_API
