Dataflow:
- Added FDataflowPoints type
- Added FPointsFacade

#Synced-CL 50527124

diff --git a/Engine/Source/Runtime/Experimental/Chaos/Private/GeometryCollection/Facades/PointsFacade.cpp b/Engine/Source/Runtime/Experimental/Chaos/Private/GeometryCollection/Facades/PointsFacade.cpp
new file mode 100644
--- /dev/null
+++ b/Engine/Source/Runtime/Experimental/Chaos/Private/GeometryCollection/Facades/PointsFacade.cpp
@@ -0,0 +1,583 @@
+// Copyright Epic Games, Inc. All Rights Reserved.
+
+#include "GeometryCollection/Facades/PointsFacade.h"
+#include "GeometryCollection/GeometryCollection.h"
+#include "GeometryCollection/ManagedArrayCollection.h"
+
+namespace GeometryCollection::Facades
+{
+	const FName FPointsFacade::PointsGroup = "Points";
+	const FName FPointsFacade::PointAttribute("Point");
+
+	FPointsFacade::FPointsFacade(FManagedArrayCollection& InCollection)
+		: ConstCollection(InCollection)
+		, MutableCollection(&InCollection)
+	{
+	}
+
+	FPointsFacade::FPointsFacade(const FManagedArrayCollection& InCollection)
+		: ConstCollection(InCollection)
+		, MutableCollection(nullptr)
+	{
+	}
+
+	/* ------------------------------------------------------------------------------------------------------------- */
+
+	bool FPointsFacade::IsValid() const
+	{
+		if (ConstCollection.HasGroup(FPointsFacade::PointsGroup) &&
+			ConstCollection.HasAttribute(FPointsFacade::PointAttribute, FPointsFacade::PointsGroup))
+		{
+			return true;
+		}
+
+		return false;
+	}
+
+	/* ------------------------------------------------------------------------------------------------------------- */
+
+	void FPointsFacade::DefineSchema()
+	{
+		if (ensure(MutableCollection))
+		{
+			MutableCollection->AddGroup(FPointsFacade::PointsGroup);
+			MutableCollection->AddAttribute<FVector3f>(FPointsFacade::PointAttribute, FPointsFacade::PointsGroup);
+		}
+	}
+
+	/* ------------------------------------------------------------------------------------------------------------- */
+
+	int32 FPointsFacade::GetNumPoints() const
+	{
+		if (IsValid())
+		{
+			const TManagedArray<FVector3f>& Point = ConstCollection.GetAttribute<FVector3f>(FPointsFacade::PointAttribute, FPointsFacade::PointsGroup);
+			return Point.Num();
+		}
+
+		return 0;
+	}
+
+	/* ------------------------------------------------------------------------------------------------------------- */
+
+	int32 FPointsFacade::GetNumAttrs() const
+	{
+		if (IsValid())
+		{
+			return ConstCollection.NumAttributes(FPointsFacade::PointsGroup);
+		}
+
+		return 0;
+	}
+
+	/* ------------------------------------------------------------------------------------------------------------- */
+
+	TArray<FName> FPointsFacade::GetAttrs() const
+	{
+		TArray<FName> AttrArr;
+
+		if (IsValid())
+		{
+			AttrArr = ConstCollection.AttributeNames(FPointsFacade::PointsGroup);
+		}
+
+		return AttrArr;
+	}
+
+	/* ------------------------------------------------------------------------------------------------------------- */
+
+	FString FPointsFacade::GetAttrsAsString(bool bAddCommaSeparator) const
+	{
+		TStringBuilder<256> StringBuilder;
+
+
+		if (IsValid())
+		{
+			TArray<FName> Attrs = ConstCollection.AttributeNames(FPointsFacade::PointsGroup);
+
+			for (int32 Idx = 1; Idx < Attrs.Num(); ++Idx)
+			{
+				StringBuilder += Attrs[Idx].ToString();
+
+				if (Idx < Attrs.Num() - 1)
+				{
+					if (bAddCommaSeparator)
+					{
+						StringBuilder += FString(", ");
+					}
+					else
+					{
+						StringBuilder += FString(" ");
+					}
+				}
+			}
+		}
+
+		return StringBuilder.ToString();
+	}
+
+	/* ------------------------------------------------------------------------------------------------------------- */
+
+	void FPointsFacade::AddPoints(const TArray<FVector3f>& InPoints)
+	{
+		if (ensure(MutableCollection))
+		{
+			const int32 NumPoints = InPoints.Num();
+
+			DefineSchema();
+
+			MutableCollection->AddElements(NumPoints, FPointsFacade::PointsGroup);
+
+			TManagedArray<FVector3f>& Point = MutableCollection->ModifyAttribute<FVector3f>(FPointsFacade::PointAttribute, FPointsFacade::PointsGroup);
+
+			for (int32 Idx = 0; Idx < NumPoints; ++Idx)
+			{
+				Point[Idx] = InPoints[Idx];
+			}
+		}
+	}
+
+	/* ------------------------------------------------------------------------------------------------------------- */
+
+	void FPointsFacade::AddPoints(const TArray<FVector>& InPoints)
+	{
+		if (ensure(MutableCollection))
+		{
+			const int32 NumPoints = InPoints.Num();
+
+			DefineSchema();
+
+			MutableCollection->AddElements(NumPoints, FPointsFacade::PointsGroup);
+
+			TManagedArray<FVector3f>& Point = MutableCollection->ModifyAttribute<FVector3f>(FPointsFacade::PointAttribute, FPointsFacade::PointsGroup);
+
+			for (int32 Idx = 0; Idx < NumPoints; ++Idx)
+			{
+				Point[Idx] = FVector3f(InPoints[Idx]);
+			}
+		}
+	}
+
+	/* ------------------------------------------------------------------------------------------------------------- */
+
+	void FPointsFacade::AddFloatAttribute(const FName InAttribute, const TArray<float>& InValues)
+	{
+		if (ensure(MutableCollection))
+		{
+			if (IsValid())
+			{
+				const TManagedArray<FVector3f>& Point = MutableCollection->GetAttribute<FVector3f>(FPointsFacade::PointAttribute, FPointsFacade::PointsGroup);
+				const int32 NumPoints = Point.Num();
+
+				if (NumPoints == InValues.Num())
+				{
+					TManagedArray<float>& Attr = MutableCollection->AddAttribute<float>(InAttribute, FPointsFacade::PointsGroup);
+
+					for (int32 Idx = 0; Idx < NumPoints; ++Idx)
+					{
+						Attr[Idx] = InValues[Idx];
+					}
+				}
+			}
+		}
+	}
+
+	/* ------------------------------------------------------------------------------------------------------------- */
+
+	void FPointsFacade::AddBoolAttribute(const FName InAttribute, const TArray<bool>& InValues)
+	{
+		if (ensure(MutableCollection))
+		{
+			if (IsValid())
+			{
+				const TManagedArray<FVector3f>& Point = MutableCollection->GetAttribute<FVector3f>(FPointsFacade::PointAttribute, FPointsFacade::PointsGroup);
+				const int32 NumPoints = Point.Num();
+
+				if (NumPoints == InValues.Num())
+				{
+					TManagedArray<bool>& Attr = MutableCollection->AddAttribute<bool>(InAttribute, FPointsFacade::PointsGroup);
+
+					for (int32 Idx = 0; Idx < NumPoints; ++Idx)
+					{
+						Attr[Idx] = InValues[Idx];
+					}
+				}
+			}
+		}
+	}
+
+	/* ------------------------------------------------------------------------------------------------------------- */
+
+	void FPointsFacade::AddIntAttribute(const FName InAttribute, const TArray<int32>& InValues)
+	{
+		if (ensure(MutableCollection))
+		{
+			if (IsValid())
+			{
+				const TManagedArray<FVector3f>& Point = MutableCollection->GetAttribute<FVector3f>(FPointsFacade::PointAttribute, FPointsFacade::PointsGroup);
+				const int32 NumPoints = Point.Num();
+
+				if (NumPoints == InValues.Num())
+				{
+					TManagedArray<int32>& Attr = MutableCollection->AddAttribute<int32>(InAttribute, FPointsFacade::PointsGroup);
+
+					for (int32 Idx = 0; Idx < NumPoints; ++Idx)
+					{
+						Attr[Idx] = InValues[Idx];
+					}
+				}
+			}
+		}
+	}
+
+	/* ------------------------------------------------------------------------------------------------------------- */
+
+	void FPointsFacade::AddVector2Attribute(const FName InAttribute, const TArray<FVector2f>& InValues)
+	{
+		if (ensure(MutableCollection))
+		{
+			if (IsValid())
+			{
+				const TManagedArray<FVector3f>& Point = MutableCollection->GetAttribute<FVector3f>(FPointsFacade::PointAttribute, FPointsFacade::PointsGroup);
+				const int32 NumPoints = Point.Num();
+
+				if (NumPoints == InValues.Num())
+				{
+					TManagedArray<FVector2f>& Attr = MutableCollection->AddAttribute<FVector2f>(InAttribute, FPointsFacade::PointsGroup);
+
+					for (int32 Idx = 0; Idx < NumPoints; ++Idx)
+					{
+						Attr[Idx] = InValues[Idx];
+					}
+				}
+			}
+		}
+	}
+
+	/* ------------------------------------------------------------------------------------------------------------- */
+
+	void FPointsFacade::AddVector3Attribute(const FName InAttribute, const TArray<FVector3f>& InValues)
+	{
+		if (ensure(MutableCollection))
+		{
+			if (IsValid())
+			{
+				const TManagedArray<FVector3f>& Point = MutableCollection->GetAttribute<FVector3f>(FPointsFacade::PointAttribute, FPointsFacade::PointsGroup);
+				const int32 NumPoints = Point.Num();
+
+				const int32 Num = FMath::Min(NumPoints, InValues.Num());
+
+				TManagedArray<FVector3f>& Attr = MutableCollection->AddAttribute<FVector3f>(InAttribute, FPointsFacade::PointsGroup);
+
+				for (int32 Idx = 0; Idx < Num; ++Idx)
+				{
+					Attr[Idx] = InValues[Idx];
+				}
+			}
+		}
+	}
+
+	/* ------------------------------------------------------------------------------------------------------------- */
+
+	bool FPointsFacade::HasAttribute(const FName& InAttribute) const
+	{
+		if (IsValid())
+		{
+			return ConstCollection.HasAttribute(InAttribute, FPointsFacade::PointsGroup);
+		}
+
+		return false;
+	}
+
+	/* ------------------------------------------------------------------------------------------------------------- */
+
+	EManagedArrayType FPointsFacade::GetAttributeType(const FName& InAttribute) const
+	{
+		if (IsValid())
+		{
+			if (ConstCollection.HasAttribute(InAttribute, FPointsFacade::PointsGroup))
+			{
+				return ConstCollection.GetAttributeType(InAttribute, FPointsFacade::PointsGroup);
+			}
+		}
+
+		return EManagedArrayType::FNoneType;
+	}
+
+	/* ------------------------------------------------------------------------------------------------------------- */
+
+	void FPointsFacade::FilterByFloatAttributeToCollection(const FName InAttribute, const EFilterByAttributeOperation InOperation, float InValue, FManagedArrayCollection& OutCollection) const
+	{
+		if (IsValid())
+		{
+			const TManagedArray<FVector3f>& Point = ConstCollection.GetAttribute<FVector3f>(FPointsFacade::PointAttribute, FPointsFacade::PointsGroup);
+			const int32 NumPoints = Point.Num();
+
+			TArray<int32> OrigIdx;
+			TArray<FVector>OutPoints;
+
+			const TManagedArray<float>* AttrArr = ConstCollection.FindAttribute<float>(InAttribute, FPointsFacade::PointsGroup);
+			if (AttrArr)
+			{
+				for (int32 Idx = 0; Idx < NumPoints; ++Idx)
+				{
+					if ((InOperation == EFilterByAttributeOperation::Equal && (*AttrArr)[Idx] == InValue) ||
+						(InOperation == EFilterByAttributeOperation::NotEqual && (*AttrArr)[Idx] != InValue) ||
+						(InOperation == EFilterByAttributeOperation::Greater && (*AttrArr)[Idx] > InValue) ||
+						(InOperation == EFilterByAttributeOperation::GreaterOrEqual && (*AttrArr)[Idx] >= InValue) ||
+						(InOperation == EFilterByAttributeOperation::Smaller && (*AttrArr)[Idx] < InValue) ||
+						(InOperation == EFilterByAttributeOperation::SmallerOrEqual && (*AttrArr)[Idx] <= InValue))
+					{
+						OutPoints.Add(FVector(Point[Idx]));
+						OrigIdx.Add(Idx);
+					}
+				}
+			}
+
+			const int32 NumPointsInOutputCollection = OutPoints.Num();
+
+			if (OutCollection.HasGroup(FPointsFacade::PointsGroup))
+			{
+				OutCollection.RemoveGroup(FPointsFacade::PointsGroup);
+			}
+
+			OutCollection.AddGroup(FPointsFacade::PointsGroup);
+			OutCollection.AddElements(NumPointsInOutputCollection, FPointsFacade::PointsGroup);
+
+			if (!OutCollection.HasAttribute(FPointsFacade::PointAttribute, FPointsFacade::PointsGroup))
+			{
+				OutCollection.AddAttribute<FVector3f>(FPointsFacade::PointAttribute, FPointsFacade::PointsGroup);
+			}
+
+			// Set values for the Point attribute
+			TManagedArray<FVector3f>& PointAttr = OutCollection.ModifyAttribute<FVector3f>(FPointsFacade::PointAttribute, FPointsFacade::PointsGroup);
+
+
+			for (int32 Idx = 0; Idx < NumPointsInOutputCollection; ++Idx)
+			{
+				PointAttr[Idx] = FVector3f(OutPoints[Idx]);
+			}
+
+			// Add other attributes
+			TArray<FName> Attrs = ConstCollection.AttributeNames(FPointsFacade::PointsGroup);
+
+			for (int32 Idx = 1; Idx < Attrs.Num(); ++Idx)
+			{
+				const FName Attr = Attrs[Idx];
+
+				if (ConstCollection.GetAttributeType(Attr, FPointsFacade::PointsGroup) == EManagedArrayType::FFloatType)
+				{
+					OutCollection.AddAttribute<float>(Attr, FPointsFacade::PointsGroup);
+
+					const TManagedArray<float>& FloatAttr = ConstCollection.GetAttribute<float>(Attr, FPointsFacade::PointsGroup);
+					TManagedArray<float>& NewAttr = OutCollection.ModifyAttribute<float>(Attr, FPointsFacade::PointsGroup);
+
+					for (int32 Idx1 = 0; Idx1 < NumPointsInOutputCollection; ++Idx1)
+					{
+						NewAttr[Idx1] = FloatAttr[OrigIdx[Idx1]];
+					}
+				}
+				else if (ConstCollection.GetAttributeType(Attrs[Idx], FPointsFacade::PointsGroup) == EManagedArrayType::FBoolType)
+				{
+					OutCollection.AddAttribute<bool>(Attr, FPointsFacade::PointsGroup);
+
+					const TManagedArray<bool>& FloatAttr = ConstCollection.GetAttribute<bool>(Attr, FPointsFacade::PointsGroup);
+					TManagedArray<bool>& NewAttr = OutCollection.ModifyAttribute<bool>(Attr, FPointsFacade::PointsGroup);
+
+					for (int32 Idx1 = 0; Idx1 < NumPointsInOutputCollection; ++Idx1)
+					{
+						NewAttr[Idx1] = FloatAttr[OrigIdx[Idx1]];
+					}
+				}
+				else if (ConstCollection.GetAttributeType(Attrs[Idx], FPointsFacade::PointsGroup) == EManagedArrayType::FInt32Type)
+				{
+					OutCollection.AddAttribute<int32>(Attr, FPointsFacade::PointsGroup);
+
+					const TManagedArray<int32>& FloatAttr = ConstCollection.GetAttribute<int32>(Attr, FPointsFacade::PointsGroup);
+					TManagedArray<int32>& NewAttr = OutCollection.ModifyAttribute<int32>(Attr, FPointsFacade::PointsGroup);
+
+					for (int32 Idx1 = 0; Idx1 < NumPointsInOutputCollection; ++Idx1)
+					{
+						NewAttr[Idx1] = FloatAttr[OrigIdx[Idx1]];
+					}
+				}
+				else if (ConstCollection.GetAttributeType(Attrs[Idx], FPointsFacade::PointsGroup) == EManagedArrayType::FVectorType)
+				{
+					OutCollection.AddAttribute<FVector3f>(Attr, FPointsFacade::PointsGroup);
+
+					const TManagedArray<FVector3f>& FloatAttr = ConstCollection.GetAttribute<FVector3f>(Attr, FPointsFacade::PointsGroup);
+					TManagedArray<FVector3f>& NewAttr = OutCollection.ModifyAttribute<FVector3f>(Attr, FPointsFacade::PointsGroup);
+
+					for (int32 Idx1 = 0; Idx1 < NumPointsInOutputCollection; ++Idx1)
+					{
+						NewAttr[Idx1] = FloatAttr[OrigIdx[Idx1]];
+					}
+				}
+				else if (ConstCollection.GetAttributeType(Attrs[Idx], FPointsFacade::PointsGroup) == EManagedArrayType::FVector2DType)
+				{
+					OutCollection.AddAttribute<FVector2f>(Attr, FPointsFacade::PointsGroup);
+
+					const TManagedArray<FVector2f>& FloatAttr = ConstCollection.GetAttribute<FVector2f>(Attr, FPointsFacade::PointsGroup);
+					TManagedArray<FVector2f>& NewAttr = OutCollection.ModifyAttribute<FVector2f>(Attr, FPointsFacade::PointsGroup);
+
+					for (int32 Idx1 = 0; Idx1 < NumPointsInOutputCollection; ++Idx1)
+					{
+						NewAttr[Idx1] = FloatAttr[OrigIdx[Idx1]];
+					}
+				}
+			}
+		}
+	}
+
+	/* ------------------------------------------------------------------------------------------------------------- */
+
+	void FPointsFacade::FilterByFloatAttributeToArray(const FName InAttribute, const EFilterByAttributeOperation InOperation, float InValue, TArray<FVector>& OutPoints) const
+	{
+		if (IsValid())
+		{
+			const TManagedArray<FVector3f>& Point = ConstCollection.GetAttribute<FVector3f>(FPointsFacade::PointAttribute, FPointsFacade::PointsGroup);
+			const int32 NumPoints = Point.Num();
+
+			const TManagedArray<float>* AttrArr = ConstCollection.FindAttribute<float>(InAttribute, FPointsFacade::PointsGroup);
+			if (AttrArr)
+			{
+				for (int32 Idx = 0; Idx < NumPoints; ++Idx)
+				{
+					if ((InOperation == EFilterByAttributeOperation::Equal && (*AttrArr)[Idx] == InValue) ||
+						(InOperation == EFilterByAttributeOperation::NotEqual && (*AttrArr)[Idx] != InValue) ||
+						(InOperation == EFilterByAttributeOperation::Greater && (*AttrArr)[Idx] > InValue) ||
+						(InOperation == EFilterByAttributeOperation::GreaterOrEqual && (*AttrArr)[Idx] >= InValue) ||
+						(InOperation == EFilterByAttributeOperation::Smaller && (*AttrArr)[Idx] < InValue) ||
+						(InOperation == EFilterByAttributeOperation::SmallerOrEqual && (*AttrArr)[Idx] <= InValue))
+					{
+						OutPoints.Add(FVector(Point[Idx]));
+					}
+				}
+			}
+		}
+	}
+
+	/* ------------------------------------------------------------------------------------------------------------- */
+
+	void FPointsFacade::FilterByFloatAttributeToSelection(const FName InAttribute, const EFilterByAttributeOperation InOperation, float InValue, TArray<int32>& OutPointSelection) const
+	{
+		if (IsValid())
+		{
+			const TManagedArray<FVector3f>& Point = ConstCollection.GetAttribute<FVector3f>(FPointsFacade::PointAttribute, FPointsFacade::PointsGroup);
+			const int32 NumPoints = Point.Num();
+
+			const TManagedArray<float>* AttrArr = ConstCollection.FindAttribute<float>(InAttribute, FPointsFacade::PointsGroup);
+			if (AttrArr)
+			{
+				for (int32 Idx = 0; Idx < NumPoints; ++Idx)
+				{
+					if ((InOperation == EFilterByAttributeOperation::Equal && (*AttrArr)[Idx] == InValue) ||
+						(InOperation == EFilterByAttributeOperation::NotEqual && (*AttrArr)[Idx] != InValue) ||
+						(InOperation == EFilterByAttributeOperation::Greater && (*AttrArr)[Idx] > InValue) ||
+						(InOperation == EFilterByAttributeOperation::GreaterOrEqual && (*AttrArr)[Idx] >= InValue) ||
+						(InOperation == EFilterByAttributeOperation::Smaller && (*AttrArr)[Idx] < InValue) ||
+						(InOperation == EFilterByAttributeOperation::SmallerOrEqual && (*AttrArr)[Idx] <= InValue))
+					{
+						OutPointSelection.Add(Idx);
+					}
+				}
+			}
+		}
+	}
+
+	/* ------------------------------------------------------------------------------------------------------------- */
+
+	TArray<FVector> FPointsFacade::GetPointsAsArray() const
+	{
+		TArray<FVector> OutArray;
+		
+		if (IsValid())
+		{
+			const TManagedArray<FVector3f>& Point = ConstCollection.GetAttribute<FVector3f>(FPointsFacade::PointAttribute, FPointsFacade::PointsGroup);
+			const int32 NumPoints = Point.Num();
+
+			OutArray.SetNumUninitialized(NumPoints);
+
+			for (int32 Idx = 0; Idx < NumPoints; ++Idx)
+			{
+				OutArray[Idx] = FVector(Point[Idx]);
+			}
+		}
+		
+		return OutArray;
+	}
+
+	/* ------------------------------------------------------------------------------------------------------------- */
+
+	const TArray<float>& FPointsFacade::GetFloatAttributeValues(const FName InAttribute) const
+	{
+		static TArray<float> EmptyArray;
+
+		if (const TManagedArray<float>* Attr = ConstCollection.FindAttribute<float>(InAttribute, FPointsFacade::PointsGroup))
+		{
+			return Attr->GetConstArray();
+		}
+
+		return EmptyArray;
+	}
+
+	/* ------------------------------------------------------------------------------------------------------------- */
+
+	const TBitArray<>& FPointsFacade::GetBoolAttributeValues(const FName InAttribute) const
+	{
+		static TBitArray<> EmptyArray;
+
+		if (const TManagedArray<bool>* Attr = ConstCollection.FindAttribute<bool>(InAttribute, FPointsFacade::PointsGroup))
+		{
+			return Attr->GetConstArray();
+		}
+
+		return EmptyArray;
+	}
+
+	/* ------------------------------------------------------------------------------------------------------------- */
+
+	const TArray<int32>& FPointsFacade::GetIntAttributeValues(const FName InAttribute) const
+	{
+		static TArray<int32> EmptyArray;
+
+		if (const TManagedArray<int32>* Attr = ConstCollection.FindAttribute<int32>(InAttribute, FPointsFacade::PointsGroup))
+		{
+			return Attr->GetConstArray();
+		}
+
+		return EmptyArray;
+	}
+
+	/* ------------------------------------------------------------------------------------------------------------- */
+
+	const TArray<FVector2f>& FPointsFacade::GetVector2AttributeValues(const FName InAttribute) const
+	{
+		static TArray<FVector2f> EmptyArray;
+		
+		if (const TManagedArray<FVector2f>* Attr = ConstCollection.FindAttribute<FVector2f>(InAttribute, FPointsFacade::PointsGroup))
+		{
+			return Attr->GetConstArray();
+		}
+
+		return EmptyArray;
+	}
+
+	/* ------------------------------------------------------------------------------------------------------------- */
+
+	const TArray<FVector3f>& FPointsFacade::GetVector3AttributeValues(const FName InAttribute) const
+	{
+		static TArray<FVector3f> EmptyArray;
+
+		if (const TManagedArray<FVector3f>* Attr = ConstCollection.FindAttribute<FVector3f>(InAttribute, FPointsFacade::PointsGroup))
+		{
+			return Attr->GetConstArray();
+		}
+
+		return EmptyArray;
+	}
+
+
+}
+
+
diff --git a/Engine/Source/Runtime/Experimental/Chaos/Public/GeometryCollection/Facades/PointsFacade.h b/Engine/Source/Runtime/Experimental/Chaos/Public/GeometryCollection/Facades/PointsFacade.h
new file mode 100644
--- /dev/null
+++ b/Engine/Source/Runtime/Experimental/Chaos/Public/GeometryCollection/Facades/PointsFacade.h
@@ -0,0 +1,120 @@
+// Copyright Epic Games, Inc. All Rights Reserved.
+
+#pragma once
+#include "GeometryCollection/ManagedArrayCollection.h"
+#include "UObject/NameTypes.h"
+
+#include "PointsFacade.generated.h"
+
+UENUM()
+enum class EFilterByAttributeOperation : uint8
+{
+	Equal UMETA(DisplayName = "=="),
+	NotEqual UMETA(DisplayName = "!="),
+	Greater UMETA(DisplayName = ">"),
+	GreaterOrEqual UMETA(DisplayName = ">="),
+	Smaller UMETA(DisplayName = "<"),
+	SmallerOrEqual UMETA(DisplayName = "<="),
+};
+
+namespace GeometryCollection::Facades
+{
+	/**
+	* FPointsFacade
+	* 
+	* Defines common API for storing points with attributes in a collection.
+	* The points/attributes are stored in the Points group.
+	* 
+	*/
+
+	class FPointsFacade
+	{
+	public:
+		CHAOS_API static const FName PointsGroup;
+		CHAOS_API static const FName PointAttribute;
+
+		/**
+		* FPointsFacade Constuctor
+		*/
+		CHAOS_API FPointsFacade(FManagedArrayCollection& InSelf);
+		CHAOS_API FPointsFacade(const FManagedArrayCollection& InSelf);
+
+		FPointsFacade(const FPointsFacade&) = delete;
+		FPointsFacade& operator= (const FPointsFacade&) = delete;
+
+		/** Is the facade defined constant. */
+		bool IsConst() const { return MutableCollection == nullptr; }
+
+		/** Is the facade defined constant. */
+		CHAOS_API bool IsValid() const;
+
+		/** Adds neccessary group/Attribute */
+		void DefineSchema();
+
+		/**
+		* Get number of points in the collection
+		*/
+		CHAOS_API int32 GetNumPoints() const;
+
+		/**
+		* Get number of attributes on the points
+		*/
+		CHAOS_API int32 GetNumAttrs() const;
+
+		/**
+		* Get list of attributes on the points
+		*/
+		CHAOS_API TArray<FName> GetAttrs() const;
+		CHAOS_API FString GetAttrsAsString(bool bAddCommaSeparator = true) const;
+
+		/**
+		* Add points to the collection
+		*/
+		CHAOS_API void AddPoints(const TArray<FVector3f>& InPoints);
+		CHAOS_API void AddPoints(const TArray<FVector>& InPoints);
+
+		/**
+		* Add attribute to the points
+		*/
+		CHAOS_API void AddFloatAttribute(const FName InAttribute, const TArray<float>& InValues);
+		CHAOS_API void AddBoolAttribute(const FName InAttribute, const TArray<bool>& InValues);
+		CHAOS_API void AddIntAttribute(const FName InAttribute, const TArray<int32>& InValues);
+		CHAOS_API void AddVector2Attribute(const FName InAttribute, const TArray<FVector2f>& InValues);
+		CHAOS_API void AddVector3Attribute(const FName InAttribute, const TArray<FVector3f>& InValues);
+
+		/**
+		* Check if attribute exist
+		*/
+		CHAOS_API bool HasAttribute(const FName& InAttribute) const;
+
+		/**
+		* Get attribute type
+		*/
+		CHAOS_API EManagedArrayType GetAttributeType(const FName& InAttribute) const;
+
+		/**
+		* Filter points by a float attribute value
+		*/
+		CHAOS_API void FilterByFloatAttributeToCollection(const FName InAttribute, const EFilterByAttributeOperation InOperation, float InValue, FManagedArrayCollection& OutCollection) const;
+		CHAOS_API void FilterByFloatAttributeToArray(const FName InAttribute, const EFilterByAttributeOperation InOperation, float InValue, TArray<FVector>& OutPoints) const;
+		CHAOS_API void FilterByFloatAttributeToSelection(const FName InAttribute, const EFilterByAttributeOperation InOperation, float InValue, TArray<int32>& OutPointSelection) const;
+
+		/**
+		* Get the points as a TArray<FVector> (it returns the point positions, all attributes will be lost)
+		*/
+		CHAOS_API TArray<FVector> GetPointsAsArray() const;
+
+		/**
+		* Get the values of a specified attribute as an array
+		*/
+		CHAOS_API const TArray<float>& GetFloatAttributeValues(const FName InAttribute) const;
+		CHAOS_API const TBitArray<>& GetBoolAttributeValues(const FName InAttribute) const;
+		CHAOS_API const TArray<int32>& GetIntAttributeValues(const FName InAttribute) const;
+		CHAOS_API const TArray<FVector2f>& GetVector2AttributeValues(const FName InAttribute) const;
+		CHAOS_API const TArray<FVector3f>& GetVector3AttributeValues(const FName InAttribute) const;
+
+	private:
+		const FManagedArrayCollection& ConstCollection;
+		FManagedArrayCollection* MutableCollection = nullptr;
+	};
+}
diff --git a/Engine/Source/Runtime/Experimental/Dataflow/Core/Private/Dataflow/DataflowPoints.cpp b/Engine/Source/Runtime/Experimental/Dataflow/Core/Private/Dataflow/DataflowPoints.cpp
new file mode 100644
--- /dev/null
+++ b/Engine/Source/Runtime/Experimental/Dataflow/Core/Private/Dataflow/DataflowPoints.cpp
@@ -0,0 +1,23 @@
+// Copyright Epic Games, Inc. All Rights Reserved.
+
+#include "Dataflow/DataflowPoints.h"
+#include "GeometryCollection/Facades/PointsFacade.h"
+
+#include UE_INLINE_GENERATED_CPP_BY_NAME(DataflowPoints)
+
+const GeometryCollection::Facades::FPointsFacade FDataflowPoints::GetPointsFacade() const
+{
+	return GeometryCollection::Facades::FPointsFacade(Collection);
+}
+
+GeometryCollection::Facades::FPointsFacade FDataflowPoints::GetPointsFacade()
+{
+	return GeometryCollection::Facades::FPointsFacade(Collection);
+}
+
+void FDataflowPoints::Set(FManagedArrayCollection& InCollection) 
+{ 
+	InCollection.CopyTo(&Collection); 
+}
+
+
diff --git a/Engine/Source/Runtime/Experimental/Dataflow/Core/Public/Dataflow/DataflowPoints.h b/Engine/Source/Runtime/Experimental/Dataflow/Core/Public/Dataflow/DataflowPoints.h
new file mode 100644
--- /dev/null
+++ b/Engine/Source/Runtime/Experimental/Dataflow/Core/Public/Dataflow/DataflowPoints.h
@@ -0,0 +1,28 @@
+// Copyright Epic Games, Inc. All Rights Reserved.
+
+#pragma once
+
+#include "GeometryCollection/Facades/PointsFacade.h"
+#include "GeometryCollection/ManagedArrayCollection.h"
+
+#include "DataflowPoints.generated.h"
+
+/** 
+* Represents a pointcloud with attributes for dataflow
+* Data stored in Points group
+*/
+USTRUCT()
+struct FDataflowPoints
+{
+	GENERATED_USTRUCT_BODY()
+
+public:
+	DATAFLOWCORE_API const GeometryCollection::Facades::FPointsFacade GetPointsFacade() const;
+	DATAFLOWCORE_API GeometryCollection::Facades::FPointsFacade GetPointsFacade();
+
+	DATAFLOWCORE_API void Set(FManagedArrayCollection& InCollection);
+
+private:
+	UPROPERTY() 
+	FManagedArrayCollection Collection;
+};
