Extended FIoHttpOptions with metadata for better debugging/trace information.

#rnx

#Synced-CL 50524817

diff --git a/Engine/Source/Runtime/Experimental/IoStore/OnDemand/Private/OnDemandContentInstaller.cpp b/Engine/Source/Runtime/Experimental/IoStore/OnDemand/Private/OnDemandContentInstaller.cpp
--- a/Engine/Source/Runtime/Experimental/IoStore/OnDemand/Private/OnDemandContentInstaller.cpp
+++ b/Engine/Source/Runtime/Experimental/IoStore/OnDemand/Private/OnDemandContentInstaller.cpp
@@ -1063,6 +1063,7 @@ void FOnDemandContentInstaller::ExecuteInstallRequest(FRequest& Request)
 		TRACE_CPUPROFILER_EVENT_SCOPE(FOnDemandContentInstaller::IssueRequest);
 		FSharedOnDemandContainer& Container		= InstallRequest.ResolvedChunks[HttpRequest.ContainerIndex].Container;
 		FOnDemandChunkInfo ChunkInfo			= FOnDemandChunkInfo(Container, Container->ChunkEntries[HttpRequest.EntryIndex]);
+		const FIoChunkId& ChunkId				= Container->ChunkIds[HttpRequest.EntryIndex];
 		const uint8 HttpCategory				= 1; // see FOnDemandHttpThread
 
 		FChunkDownloadRequest* PendingChunkRequest = nullptr;
@@ -1092,6 +1093,7 @@ void FOnDemandContentInstaller::ExecuteInstallRequest(FRequest& Request)
 			FIoHttpOptions HttpOptions = FIoHttpOptions(Request.Priority, HttpRetryCount, EIoHttpFlags::None);
 			HttpOptions.SetCategory(HttpCategory);
 			HttpOptions.SetRange(FIoHttpRange(0, ChunkInfo.Entry.GetDiskSize()));
+			HttpOptions.SetMetadata(FIoHttpMetadata(ChunkId));
 
 			UE_LOG(LogIoStoreOnDemand, VeryVerbose, TEXT("Created request for chunk: ChunkHash=%s, RequestId=%u"),
 				*TStringBuilder<64>(InPlace, ChunkHash), PendingChunkRequest->RequestId);
diff --git a/Engine/Source/Runtime/Experimental/IoStore/OnDemand/Private/OnDemandHttpIoDispatcherBackend.cpp b/Engine/Source/Runtime/Experimental/IoStore/OnDemand/Private/OnDemandHttpIoDispatcherBackend.cpp
--- a/Engine/Source/Runtime/Experimental/IoStore/OnDemand/Private/OnDemandHttpIoDispatcherBackend.cpp
+++ b/Engine/Source/Runtime/Experimental/IoStore/OnDemand/Private/OnDemandHttpIoDispatcherBackend.cpp
@@ -220,7 +220,7 @@ void FHttpIoDispatcherBackend::ResolveIoRequests(FIoRequestList Requests, FIoReq
 			ResolvedRequest->ChunkInfo.HostGroupName(),
 			ResolvedRequest->ChunkInfo.RelativeUrl(),
 			FIoHttpHeaders(),
-			FIoHttpOptions(Request->Priority, HttpRetryCount, EIoHttpFlags::Default, HttpRange),
+			FIoHttpOptions(Request->Priority, HttpRetryCount, EIoHttpFlags::Default, HttpRange, FIoHttpMetadata(Request->ChunkId)),
 			ResolvedRequest->ChunkInfo.ChunkEntry().Hash,
 			[this, ResolvedRequest](FIoHttpResponse&& HttpResponse)
 			{
diff --git a/Engine/Source/Runtime/Experimental/IoStore/OnDemandCore/Private/HttpIoDispatcher.cpp b/Engine/Source/Runtime/Experimental/IoStore/OnDemandCore/Private/HttpIoDispatcher.cpp
--- a/Engine/Source/Runtime/Experimental/IoStore/OnDemandCore/Private/HttpIoDispatcher.cpp
+++ b/Engine/Source/Runtime/Experimental/IoStore/OnDemandCore/Private/HttpIoDispatcher.cpp
@@ -1,7 +1,9 @@
 // Copyright Epic Games, Inc. All Rights Reserved.
 
 #include "IO/HttpIoDispatcher.h"
+
 #include "Containers/AnsiString.h"
+#include "IO/IoChunkId.h"
 #include "Misc/StringBuilder.h"
 
 DEFINE_LOG_CATEGORY(LogHttpIoDispatcher);
@@ -13,6 +15,48 @@ namespace UE
 TSharedPtr<IHttpIoDispatcher> GHttpIoDispatcher;
 const FIoHttpOptions FIoHttpOptions::Default = FIoHttpOptions();
 
+////////////////////////////////////////////////////////////////////////////////
+#if UE_HTTP_METADATA_ENABLED
+struct FIoHttpMetadata::FMetadata
+{
+	FIoChunkId ChunkId = FIoChunkId::InvalidChunkId;
+};
+#endif
+
+////////////////////////////////////////////////////////////////////////////////
+FIoHttpMetadata::FIoHttpMetadata()
+#if UE_HTTP_METADATA_ENABLED
+	: Metadata(MakeShared<FMetadata>())
+#endif
+{
+}
+
+FIoHttpMetadata::FIoHttpMetadata(const FIoChunkId& ChunkId)
+	: FIoHttpMetadata()
+{
+	SetChunkId(ChunkId);
+}
+
+FIoHttpMetadata::~FIoHttpMetadata()
+{
+}
+
+const FIoChunkId& FIoHttpMetadata::GetChunkId() const
+{
+#if UE_HTTP_METADATA_ENABLED
+	return Metadata->ChunkId;
+#else
+	return FIoChunkId::InvalidChunkId;
+#endif
+}
+
+void FIoHttpMetadata::SetChunkId(const FIoChunkId& ChunkId)
+{
+#if UE_HTTP_METADATA_ENABLED
+	Metadata->ChunkId = ChunkId;
+#endif
+}
+
 ////////////////////////////////////////////////////////////////////////////////
 FIoHttpHeaders& FIoHttpHeaders::Add(FAnsiString&& HeaderName, FAnsiString&& HeaderValue)
 {
diff --git a/Engine/Source/Runtime/Experimental/IoStore/OnDemandCore/Public/IO/HttpIoDispatcher.h b/Engine/Source/Runtime/Experimental/IoStore/OnDemandCore/Public/IO/HttpIoDispatcher.h
--- a/Engine/Source/Runtime/Experimental/IoStore/OnDemandCore/Public/IO/HttpIoDispatcher.h
+++ b/Engine/Source/Runtime/Experimental/IoStore/OnDemandCore/Public/IO/HttpIoDispatcher.h
@@ -18,6 +18,17 @@
 
 UE_API DECLARE_LOG_CATEGORY_EXTERN(LogHttpIoDispatcher, Log, All);
 
+////////////////////////////////////////////////////////////////////////////////////////////////////
+#if !defined(UE_HTTP_METADATA_ENABLED)
+	#if !UE_BUILD_SHIPPING
+		#define UE_HTTP_METADATA_ENABLED UE_TRACE_ENABLED
+	#else
+		#define UE_HTTP_METADATA_ENABLED 0
+	#endif
+#endif
+
+class FIoChunkId;
+
 namespace UE
 {
 
@@ -112,10 +123,28 @@ FIoHttpRange FIoHttpRange::FromOffsetAndLength(const FIoOffsetAndLength& OffsetA
 	return FIoHttpRange(IntCastChecked<uint32>(OffsetAndLength.GetOffset()), IntCastChecked<uint32>(End));
 }
 
+/** Non shipping metadata about an HTTP request. */
+class FIoHttpMetadata
+{
+public:
+	UE_API						FIoHttpMetadata();
+	UE_API						FIoHttpMetadata(const FIoChunkId& ChunkId);
+	UE_API						~FIoHttpMetadata();
+	UE_API void					SetChunkId(const FIoChunkId& ChunkId);
+	UE_API const FIoChunkId&	GetChunkId() const;
+
+private:
+#if UE_HTTP_METADATA_ENABLED
+	struct FMetadata;
+	TSharedPtr<FMetadata>		Metadata;
+#endif
+};
+
 /** Options for controlling the behavior of an HTTP request. */
 struct FIoHttpOptions
 {
 public:
+	inline						FIoHttpOptions(int32 InPriority, int32 InRetryCount, EIoHttpFlags InFlags, const FIoHttpRange& InRange, FIoHttpMetadata Metadata);
 	inline						FIoHttpOptions(int32 InPriority, int32 InRetryCount, EIoHttpFlags InFlags, const FIoHttpRange& InRange);
 	inline						FIoHttpOptions(int32 InPriority, int32 InRetryCount, EIoHttpFlags InFlags);
 	inline						FIoHttpOptions(int32 InPriority, int32 InRetryCount);
@@ -127,11 +156,13 @@ public:
 	int32						GetRetryCount()	const { return RetryCount; }
 	EIoHttpFlags				GetFlags()		const { return Flags; }
 	uint8						GetCategory()	const { return Category; }
+	const FIoHttpMetadata&		GetMetadata()	const { return Metadata; }
 
 	void						SetRange(const FIoHttpRange InRange)	{ Range = InRange; }
 	void						SetPriority(int32 InPriority)			{ Priority = InPriority; }
 	void						SetRetryCount(int32 InRetryCount)		{ RetryCount = InRetryCount; }
 	void						SetCategory(uint8 InCategory)			{ Category = InCategory; }
+	void						SetMetadata(FIoHttpMetadata&& InMeta)	{ Metadata = MoveTemp(InMeta); }
 
 	UE_API static const FIoHttpOptions Default;
 
@@ -141,8 +172,18 @@ private:
 	int32			RetryCount = 0;
 	EIoHttpFlags	Flags = EIoHttpFlags::Default;
 	uint8			Category = 0;
+	FIoHttpMetadata	Metadata;
 };
 
+FIoHttpOptions::FIoHttpOptions(int32 InPriority, int32 InRetryCount, EIoHttpFlags InFlags, const FIoHttpRange& InRange, FIoHttpMetadata InMetadata)
+	: Range(InRange)
+	, Priority(InPriority)
+	, RetryCount(InRetryCount)
+	, Flags(InFlags)
+	, Metadata(InMetadata)
+{
+}
+
 FIoHttpOptions::FIoHttpOptions(int32 InPriority, int32 InRetryCount, EIoHttpFlags InFlags, const FIoHttpRange& InRange)
 	: Range(InRange)
 	, Priority(InPriority)
