Extended HTTP I/O dispatcher to use the new HTTP trace channel.
- run client with -trace="[other channels],http"
- run client with -trace="[other channels],http,metadata,assetmetadata" to resolve package names

#rnx

#Synced-CL 50524817

diff --git a/Engine/Source/Runtime/Experimental/IoStore/OnDemand/Private/OnDemandHttpIoDispacher.cpp b/Engine/Source/Runtime/Experimental/IoStore/OnDemand/Private/OnDemandHttpIoDispacher.cpp
--- a/Engine/Source/Runtime/Experimental/IoStore/OnDemand/Private/OnDemandHttpIoDispacher.cpp
+++ b/Engine/Source/Runtime/Experimental/IoStore/OnDemand/Private/OnDemandHttpIoDispacher.cpp
@@ -23,10 +23,22 @@
 #include "Misc/Fork.h"
 #include "Misc/SingleThreadRunnable.h"
 #include "OnDemandHttpClient.h"
+#include "ProfilingDebugging/AssetMetadataTrace.h"
+#include "ProfilingDebugging/HttpTrace.h"
 #include "Statistics.h"
 
 #include <atomic>
 
+#if UE_TRACE_HTTP_ENABLED
+#include "IO/IoChunkId.h"
+#endif
+
+#if UE_TRACE_HTTP_METADATA_ENABLED
+#include "IO/IoDispatcherInternal.h"
+#include "IO/PackageId.h"
+#include "Misc/PackageName.h"
+#endif
+
 // When enabled the cvar 'iax.InvalidUrlChance' will be enabled and allow us to simulate invalid urls
 #define UE_ALLOW_INVALID_URL_DEBUGGING !UE_BUILD_SHIPPING
 #define UE_ALLOW_CACHE_POISONING !UE_BUILD_SHIPPING
@@ -597,6 +609,150 @@ void DebugPrintRange(FHttpRequest* Requests, const UE::FIoHttpRange& TotalRange)
 #endif
 }
 
+////////////////////////////////////////////////////////////////////////////////
+#if UE_TRACE_HTTP_ENABLED
+
+class FTraceContext
+{
+public:
+	void					Init(void* Dispatcher);
+	void					TraceRequestStarted(FHttpRequest* Request, const ANSICHAR* Url);
+	void					TraceRequestCompleted(
+								const FHttpRequest* Request,
+								const FMultiEndpointHttpClientResponse& Response,
+								const FIASHostGroup& HostGroup);
+	void					TracePackageName(const FIoChunkId& ChunkId);
+	static FTraceContext&	Get();
+
+private:
+	struct FPackageNameCache
+	{
+		static constexpr int32 Slack = 256;
+
+		void Reset()
+		{
+			PackageName.Empty(Slack);
+			Filename.Reset();
+		}
+
+		FString						PackageName;
+		TUtf8StringBuilder<Slack>	Filename;
+	};
+
+	uint64				Dispatcher = 0;
+	FPackageNameCache	PackageNameCache;
+};
+
+////////////////////////////////////////////////////////////////////////////////
+void FTraceContext::Init(void* InDispatcher)
+{
+	Get().Dispatcher = uint64(InDispatcher);
+	TRACE_HTTP_DISPATCHER_CREATED(uint64(InDispatcher), "HttpIoDispatcher");
+	TRACE_HTTP_CATEGORY_CREATED(uint8(EHttpRequestType::Streaming), "Streaming");
+	TRACE_HTTP_CATEGORY_CREATED(uint8(EHttpRequestType::Installed), "Install");
+}
+
+FTraceContext& FTraceContext::Get()
+{
+	static FTraceContext Ctx;
+	return Ctx;
+}
+
+void FTraceContext::TraceRequestStarted(FHttpRequest* Request, const ANSICHAR* Url)
+{
+	TRACE_HTTP_REQUEST_STARTED(Dispatcher, uint64(Request), Url, Request->Options.GetPriority(), Request->Options.GetCategory());
+
+	if (Request->ChunkHash.IsZero())
+	{
+		// Generic HTTP request
+		return;
+	}
+
+	for (const FHttpRequest* It = Request; It != nullptr; It = It->Next)
+	{
+		const FIoChunkId&	ChunkId	= It->Options.GetMetadata().GetChunkId();
+		const FIoHttpRange	Range	= It->Options.GetRange();
+		TracePackageName(ChunkId);
+		TRACE_HTTP_CHUNK_RANGE_ADDED(uint64(Request), ChunkId, Range.GetMin(), Range.GetMax());
+	}
+}
+
+void FTraceContext::TraceRequestCompleted(
+	const FHttpRequest* Request,
+	const FMultiEndpointHttpClientResponse& Response,
+	const FIASHostGroup& HostGroup)
+{
+	FAnsiStringView Host = Response.HostIndex != INDEX_NONE ? HostGroup.GetHostUrls()[Response.HostIndex] : HostGroup.GetPrimaryHostUrl();
+	TRACE_HTTP_REQUEST_COMPLETED(uint64(Request), Host.GetData(), Response.StatusCode, uint32(Response.Body.GetSize()));
+}
+
+void FTraceContext::TracePackageName(const FIoChunkId& ChunkId)
+{
+#if UE_TRACE_HTTP_METADATA_ENABLED
+	const bool bTraceHttpChannelEnabled = UE_TRACE_CHANNELEXPR_IS_ENABLED(HttpChannel);
+	if (bTraceHttpChannelEnabled == false)
+	{
+		return;
+	}
+
+	bool bPackageDataChunk = false;
+	switch (ChunkId.GetChunkType())
+	{
+		case EIoChunkType::ExportBundleData:
+		case EIoChunkType::BulkData:
+		case EIoChunkType::OptionalBulkData:
+		case EIoChunkType::MemoryMappedBulkData:
+			bPackageDataChunk = true;
+		default:
+			break;
+	};
+
+	if (bPackageDataChunk == false)
+	{
+		return;
+	}
+
+	FPackageId PackageId;
+	FMemory::Memcpy(&PackageId, &ChunkId, sizeof(FPackageId));
+
+	if (PackageId.IsValid() == false)
+	{
+		return;
+	}
+
+	PackageNameCache.Reset();
+	FUtf8StringView _;
+	FUtf8StringView Filename = FIoDispatcherInternal::GetFilename(ChunkId, PackageNameCache.Filename, _);
+
+	if (Filename.IsEmpty())
+	{
+		return;
+	}
+
+	FString Tmp(Filename);
+	if (FPackageName::TryConvertFilenameToLongPackageName(Tmp, PackageNameCache.PackageName))
+	{
+		UE_TRACE_PACKAGE_NAME(PackageId, FName(PackageNameCache.PackageName));
+	}
+#endif // UE_TRACE_HTTP_METADATA_ENABLED
+}
+
+#define TRACE_INIT(Dispatcher) \
+	FTraceContext::Get().Init(Dispatcher)
+#define TRACE_REQUEST_STARTED(Request, Url) \
+	FTraceContext::Get().TraceRequestStarted(Request, Url)
+#define TRACE_REQUEST_COMPLETED(Request, Response, HostGroup) \
+	FTraceContext::Get().TraceRequestCompleted(Request, Response, HostGroup)
+
+////////////////////////////////////////////////////////////////////////////////
+#else
+
+#define TRACE_INIT(...)
+#define TRACE_REQUEST_STARTED(...)
+#define TRACE_REQUEST_COMPLETED(...)
+
+#endif // UE_TRACE_HTTP_ENABLED
+
 ////////////////////////////////////////////////////////////////////////////////
 class FHttpQueue
 {
@@ -1148,6 +1304,7 @@ FHttpDispatcher::FHttpDispatcher(TUniquePtr<IIasCache>&& InCache)
 	const bool bAllowPreFork				= FParse::Param(FCommandLine::IsInitialized() ? FCommandLine::Get() : TEXT(""), TEXT("-Ias.EnableHttpThreadPreFork"));
 
 	Thread.Reset(FForkProcessHelper::CreateForkableThread(this, TEXT("IoService.Http"), StackSize, ThreadPriority, ThreadAffinityMask, CreateFlags, bAllowPreFork));
+	TRACE_INIT(this);
 };
 
 FHttpDispatcher::~FHttpDispatcher()
@@ -1686,10 +1843,12 @@ bool FHttpDispatcher::ProcessHttpRequests()
 				: EMultiEndpointRequestFlags::None;
 
 			// Note that the requests issued here must be completed before this method can return so the capturing of locals is safe.
-			FMultiEndpointHttpClient::FHttpTicketId TicketId =  HttpClient->Get(HostGroup.GetUnderlyingHostGroup(), Url, TotalRange.ToOffsetAndLength(), MoveTemp(Headers), Flags,
+			TRACE_REQUEST_STARTED(HttpRequest, Url.ToString());
+			FMultiEndpointHttpClient::FHttpTicketId TicketId = HttpClient->Get(HostGroup.GetUnderlyingHostGroup(), Url, TotalRange.ToOffsetAndLength(), MoveTemp(Headers), Flags,
 				[this, HttpRequest, TotalRange, &NumConcurrentRequests, HostGroup = MoveTemp(HostGroup)]
 				(FMultiEndpointHttpClientResponse&& HttpResponse) mutable
 				{
+					TRACE_REQUEST_COMPLETED(HttpRequest, HttpResponse, HostGroup);
 					FOnDemandIoBackendStats::Get()->OnHttpRequestCompleted();
 					NumConcurrentRequests--;
 
@@ -1966,6 +2125,8 @@ TSharedPtr<IOnDemandHttpIoDispatcher> MakeOnDemanHttpIoDispatcher(TUniquePtr<cla
 } // UE::IoStore
 
 #undef UE_ALLOW_CACHE_POISONING
-
 #undef UE_ALLOW_INVALID_URL_DEBUGGING
+#undef TRACE_INIT
+#undef TRACE_REQUEST_STARTED
+#undef TRACE_REQUEST_COMPLETED
 
