[TraceBasedDebuggers] replaced recording flag by the id of the application instance that requested the active recording
- allow recording controls to filter out sessions from the dropdown
- fix message traffic issue by sending the Pong message only to the sender of the Ping (i.e., no broadcast)
- RewindDebugger controls now filters out other Editor sessions
- fix session naming issue with multiple local editor instances: the session name can only be replaced by "Local Editor" by the process handling the Pong message and not when generating that message since multiple editors can handle it

#Synced-CL 50485369

diff --git a/Engine/Plugins/Animation/GameplayInsights/Source/RewindDebugger/Private/RewindDebuggerRecordingControls.cpp b/Engine/Plugins/Animation/GameplayInsights/Source/RewindDebugger/Private/RewindDebuggerRecordingControls.cpp
--- a/Engine/Plugins/Animation/GameplayInsights/Source/RewindDebugger/Private/RewindDebuggerRecordingControls.cpp
+++ b/Engine/Plugins/Animation/GameplayInsights/Source/RewindDebugger/Private/RewindDebuggerRecordingControls.cpp
@@ -87,7 +87,7 @@ void FDebuggerRecordingControls::AddToMenuInternal(FToolMenuSection& InSection)
 					if (Dir.Len() > 0)
 					{
 						TArray<FString> Files;
-						IFileManager::Get().FindFiles(Files, *Dir, TEXT(".utrace"));
+						IFileManager::Get().FindFiles(Files, *Dir, TEXT("utrace"));
 						for (const FString& File : Files)
 						{
 							const FString FullPath = FString::Printf(TEXT("%s/%s"), *Dir, *File);
@@ -172,6 +172,13 @@ bool FDebuggerRecordingControls::IsRecordingAvailableForSessionInternal(const TS
 		|| InSessionInfo->BuildTargetType != EBuildTargetType::Editor;
 }
 
+bool FDebuggerRecordingControls::ShouldFilterOutSessionInternal(const TSharedRef<TraceBasedDebuggers::FSessionInfo>& InSessionInfo) const
+{
+	// Do not allow users to see and start recordings in other Editor processes (locally or on the network)
+	return (InSessionInfo->InstanceId != TraceBasedDebuggers::FRemoteSessionsManager::LocalEditorSessionID
+		&& InSessionInfo->BuildTargetType == EBuildTargetType::Editor);
+}
+
 TNotNull<TUniquePtr<TraceBasedDebuggers::FStartRecordingCommandMessage>> FDebuggerRecordingControls::BuildStartRecordingParamsInternal()
 {
 	return MakeUnique<FStartRecordingCommand>();
diff --git a/Engine/Plugins/Animation/GameplayInsights/Source/RewindDebugger/Private/RewindDebuggerRecordingControls.h b/Engine/Plugins/Animation/GameplayInsights/Source/RewindDebugger/Private/RewindDebuggerRecordingControls.h
--- a/Engine/Plugins/Animation/GameplayInsights/Source/RewindDebugger/Private/RewindDebuggerRecordingControls.h
+++ b/Engine/Plugins/Animation/GameplayInsights/Source/RewindDebugger/Private/RewindDebuggerRecordingControls.h
@@ -1,4 +1,4 @@
-﻿// Copyright Epic Games, Inc. All Rights Reserved.
+// Copyright Epic Games, Inc. All Rights Reserved.
 
 #pragma once
 
@@ -32,6 +32,7 @@ protected:
 	virtual bool ConnectToLiveSession_RelayInternal(const FGuid& InSessionId) override;
 	virtual bool ConnectToLiveSession_DirectInternal(const FGuid& InSessionId) override;
 	virtual bool IsRecordingAvailableForSessionInternal(const TSharedRef<TraceBasedDebuggers::FSessionInfo>& InSessionInfo) const override;
+	virtual bool ShouldFilterOutSessionInternal(const TSharedRef<TraceBasedDebuggers::FSessionInfo>& InSessionInfo) const override;
 
 	virtual TNotNull<TUniquePtr<TraceBasedDebuggers::FStartRecordingCommandMessage>> BuildStartRecordingParamsInternal() override;
 	virtual void SendStartRecordingCommandInternal(const FMessageAddress& InAddress, TUniquePtr<TraceBasedDebuggers::FStartRecordingCommandMessage> InStartRecordingCommand) override;
diff --git a/Engine/Plugins/Animation/GameplayInsights/Source/RewindDebuggerRuntime/Private/RewindDebuggerEngineEditorBridge.cpp b/Engine/Plugins/Animation/GameplayInsights/Source/RewindDebuggerRuntime/Private/RewindDebuggerEngineEditorBridge.cpp
--- a/Engine/Plugins/Animation/GameplayInsights/Source/RewindDebuggerRuntime/Private/RewindDebuggerEngineEditorBridge.cpp
+++ b/Engine/Plugins/Animation/GameplayInsights/Source/RewindDebuggerRuntime/Private/RewindDebuggerEngineEditorBridge.cpp
@@ -156,7 +156,7 @@ void FRewindDebuggerEngineEditorBridge::BuildRecordingStatusInternal(TraceBasedD
 	if (const FRewindDebuggerRuntime* Runtime = FRewindDebuggerRuntime::Instance())
 	{
 		OutStatusMessage.DebuggerId = Runtime->GetDebuggerId();
-		OutStatusMessage.bIsRecording = Runtime->IsRecording();
+		OutStatusMessage.RequesterId = Runtime->GetRecordingRequesterId();
 		OutStatusMessage.ElapsedTime = Runtime->GetAccumulatedRecordingTime();
 		OutStatusMessage.BufferedDataBytesSize = Runtime->GetBufferedDataBytesSize();
 	}
diff --git a/Engine/Plugins/Animation/GameplayInsights/Source/RewindDebuggerRuntime/Private/RewindDebuggerRemoteSessionsHandler.cpp b/Engine/Plugins/Animation/GameplayInsights/Source/RewindDebuggerRuntime/Private/RewindDebuggerRemoteSessionsHandler.cpp
--- a/Engine/Plugins/Animation/GameplayInsights/Source/RewindDebuggerRuntime/Private/RewindDebuggerRemoteSessionsHandler.cpp
+++ b/Engine/Plugins/Animation/GameplayInsights/Source/RewindDebuggerRuntime/Private/RewindDebuggerRemoteSessionsHandler.cpp
@@ -75,7 +75,7 @@ void FSessionsHandler::OnBuildingFullSessionInfoResponseMessage(
 {
 	if (const ::RewindDebugger::FRewindDebuggerRuntime* Runtime = ::RewindDebugger::FRewindDebuggerRuntime::Instance())
 	{
-		InMessage.bIsRecording = Runtime->IsRecording();
+		InMessage.RecordingRequesterId = Runtime->GetRecordingRequesterId();
 		InMessage.DebuggerId = DebuggerGuid;
 	}
 
diff --git a/Engine/Plugins/Animation/GameplayInsights/Source/RewindDebuggerRuntime/Private/RewindDebuggerRuntime.cpp b/Engine/Plugins/Animation/GameplayInsights/Source/RewindDebuggerRuntime/Private/RewindDebuggerRuntime.cpp
--- a/Engine/Plugins/Animation/GameplayInsights/Source/RewindDebuggerRuntime/Private/RewindDebuggerRuntime.cpp
+++ b/Engine/Plugins/Animation/GameplayInsights/Source/RewindDebuggerRuntime/Private/RewindDebuggerRuntime.cpp
@@ -43,6 +43,7 @@ namespace RewindDebugger
 		using namespace UE::TraceBasedDebuggers;
 		FRuntimeModule::StartRecording(FStartRecordingCommandMessage
 			{
+				.InstanceId = FApp::GetInstanceId(),
 				.Target = TEXT("127.0.0.1"),
 				.RecordingMode = ERecordingMode::Live,
 				.TransportMode = ETraceTransportMode::TraceServer
diff --git a/Engine/Source/Runtime/Experimental/ChaosSolverEngine/Private/Chaos/ChaosVDEngineEditorBridge.cpp b/Engine/Source/Runtime/Experimental/ChaosSolverEngine/Private/Chaos/ChaosVDEngineEditorBridge.cpp
--- a/Engine/Source/Runtime/Experimental/ChaosSolverEngine/Private/Chaos/ChaosVDEngineEditorBridge.cpp
+++ b/Engine/Source/Runtime/Experimental/ChaosSolverEngine/Private/Chaos/ChaosVDEngineEditorBridge.cpp
@@ -1,4 +1,4 @@
-﻿// Copyright Epic Games, Inc. All Rights Reserved.
+// Copyright Epic Games, Inc. All Rights Reserved.
 
 #include "Chaos/ChaosVDEngineEditorBridge.h"
 
@@ -199,7 +199,7 @@ void FChaosVDEngineEditorBridge::BuildRecordingStatusInternal(UE::TraceBasedDebu
 	const FChaosVDRuntimeModule& RuntimeModule = FChaosVDRuntimeModule::Get();
 
 	OutStatusMessage.DebuggerId = RuntimeModule.GetDebuggerId();
-	OutStatusMessage.bIsRecording = RuntimeModule.IsRecording();
+	OutStatusMessage.RequesterId = RuntimeModule.GetRecordingRequesterId();
 	OutStatusMessage.ElapsedTime = RuntimeModule.GetAccumulatedRecordingTime();
 	OutStatusMessage.BufferedDataBytesSize = RuntimeModule.GetBufferedDataBytesSize();
 }
diff --git a/Engine/Source/Runtime/Experimental/ChaosSolverEngine/Private/Chaos/ChaosVDRemoteSessionsManager.cpp b/Engine/Source/Runtime/Experimental/ChaosSolverEngine/Private/Chaos/ChaosVDRemoteSessionsManager.cpp
--- a/Engine/Source/Runtime/Experimental/ChaosSolverEngine/Private/Chaos/ChaosVDRemoteSessionsManager.cpp
+++ b/Engine/Source/Runtime/Experimental/ChaosSolverEngine/Private/Chaos/ChaosVDRemoteSessionsManager.cpp
@@ -96,7 +96,7 @@ void FChaosVDRemoteSessionsHandler::OnBuildingFullSessionInfoResponseMessage(
 {
 	using namespace Chaos::VisualDebugger;
 
-	InMessage.bIsRecording = FChaosVDRuntimeModule::Get().IsRecording();
+	InMessage.RecordingRequesterId = FChaosVDRuntimeModule::Get().GetRecordingRequesterId();
 	InMessage.DebuggerId = Chaos::VD::DebuggerGuid;
 
 	FChaosVDFullSessionInfoResponseData MessageData;
diff --git a/Engine/Source/Runtime/Experimental/TraceBasedDebuggers/Private/RemoteSessionsManager.cpp b/Engine/Source/Runtime/Experimental/TraceBasedDebuggers/Private/RemoteSessionsManager.cpp
--- a/Engine/Source/Runtime/Experimental/TraceBasedDebuggers/Private/RemoteSessionsManager.cpp
+++ b/Engine/Source/Runtime/Experimental/TraceBasedDebuggers/Private/RemoteSessionsManager.cpp
@@ -35,7 +35,7 @@ const FGuid FRemoteSessionsManager::InvalidSessionGUID = FGuid();
 
 const FString FRemoteSessionsManager::LocalEditorSessionName = TEXT("Local Editor");
 const FGuid FRemoteSessionsManager::LocalEditorSessionID = IsController() ? FApp::GetInstanceId() : InvalidSessionGUID;
-const FName FRemoteSessionsManager::MessageBusEndpointName = FName("TraceBasedDebuggerSessionManagerEndpoint");
+const FName FRemoteSessionsManager::MessageBusEndpointName = FName("FDebuggerSessionManager");
 const FString FRemoteSessionsManager::AllRemoteSessionsTargetName = TEXT("All Remote");
 const FString FRemoteSessionsManager::AllRemoteServersTargetName = TEXT("All Remote Servers");
 const FString FRemoteSessionsManager::AllRemoteClientsTargetName = TEXT("All Remote Clients");
@@ -49,13 +49,13 @@ const FTraceConnectionDetails& FSessionInfo::GetTraceConnectionDetails()
 
 bool FSessionInfo::IsAnyDebuggerRecording() const
 {
-	return LastKnownRecordingState.bIsRecording;
+	return LastKnownRecordingState.RequesterId.IsValid();
 }
 
 bool FSessionInfo::IsRecording(const FGuid& DebuggerTypeId) const
 {
 	return LastKnownRecordingState.DebuggerId == DebuggerTypeId
-		&& LastKnownRecordingState.bIsRecording;
+		&& LastKnownRecordingState.RequesterId == FRemoteSessionsManager::LocalEditorSessionID;
 }
 
 uint64 FSessionInfo::GetBufferedBytesNum() const
@@ -172,7 +172,13 @@ TSharedPtr<FSessionInfo> FRemoteSessionsManager::CreatedWrapperSessionInfo(FGuid
 
 TSharedPtr<FMessageEndpoint> FRemoteSessionsManager::CreateEndPoint(const TSharedRef<IMessageBus>& InMessageBus)
 {
-	FMessageEndpointBuilder EndpointBuilder = FMessageEndpoint::Builder(MessageBusEndpointName, InMessageBus)
+	// Create process specific name to make it easier to debug messages
+	const FName EndPointName(FString::Format(TEXT("{0} ({1} {2})"),
+		{ MessageBusEndpointName.ToString()
+		, FApp::GetProjectName()
+		, FString::FromInt(FPlatformProcess::GetCurrentProcessId()) }));
+
+	FMessageEndpointBuilder EndpointBuilder = FMessageEndpoint::Builder(EndPointName, InMessageBus)
 		.Handling<FSessionPing>(this, &FRemoteSessionsManager::HandleSessionPingMessage)
 		.Handling<FFullSessionInfoRequestMessage>(this, &FRemoteSessionsManager::HandleFullSessionStateRequestMessage);
 
@@ -258,6 +264,11 @@ void FRemoteSessionsManager::ShutdownMessagingSystem()
 	{
 		// Unsubscribe from all message types (i.e. no delegate required)
 		MessageEndpoint->Unsubscribe();
+
+		if (const TSharedPtr MessageBus = MessageBusPtr.Pin())
+		{
+			MessageBus->Unregister(MessageEndpoint->GetAddress());
+		}
 	}
 
 	MessageBusPtr.Reset();
@@ -400,7 +411,7 @@ void FRemoteSessionsManager::SendPing()
 	}
 }
 
-void FRemoteSessionsManager::SendPong(const FSessionPing& InMessage)
+void FRemoteSessionsManager::SendPong(const FSessionPing& InMessage, const TSharedRef<IMessageContext>& InContext)
 {
 	if (!ensure(MessageEndpoint))
 	{
@@ -415,22 +426,16 @@ void FRemoteSessionsManager::SendPong(const FSessionPing& InMessage)
 		const EBuildTargetType TargetType = GetBuildTargetType();
 		PongMessage->BuildTargetType = static_cast<uint8>(TargetType);
 
-		if (InMessage.ControllerInstanceId == PongMessage->InstanceId)
-		{
-			PongMessage->SessionName = LocalEditorSessionName;
-		}
-		else
-		{
-			FString AppSessionName = FApp::GetSessionName();
-			PongMessage->SessionName = AppSessionName == TEXT("None") || AppSessionName.IsEmpty()
-				? FString::Format(TEXT("{0} {1} {2}"),
-						{FApp::GetProjectName()
-						, FString(LexToString(TargetType))
-						, FString::FromInt(FPlatformProcess::GetCurrentProcessId())})
-				: AppSessionName;
-		}
+		FString AppSessionName = FApp::GetSessionName();
+		PongMessage->SessionName = AppSessionName == TEXT("None") || AppSessionName.IsEmpty()
+			? FString::Format(TEXT("{0} {1} {2}"),
+					{FApp::GetProjectName()
+					, FString(LexToString(TargetType))
+					, FString::FromInt(FPlatformProcess::GetCurrentProcessId())})
+			: AppSessionName;
 
-		MessageEndpoint->Publish(PongMessage, EMessageScope::Network);
+		// Send answer back directly to the sender of the Ping message (i.e., no broadcast)
+		MessageEndpoint->Send(PongMessage, InContext.Get().GetSender());
 	}
 }
 
@@ -477,6 +482,13 @@ void FRemoteSessionsManager::HandleSessionPongMessage(const FSessionPong& InMess
 		SessionInfoPtr->Address = InContext->GetSender();
 		SessionInfoPtr->InstanceId = InMessage.InstanceId;
 		SessionInfoPtr->SessionName = InMessage.SessionName;
+
+		// We received the message from our own process, so let's use the user 
+		// friendly name for that session (i.e., Local Editor)
+		if (InMessage.InstanceId == LocalEditorSessionID)
+		{
+			SessionInfoPtr->SessionName = LocalEditorSessionName;
+		}
 		SessionInfoPtr->BuildTargetType = static_cast<EBuildTargetType>(InMessage.BuildTargetType);
 
 		RegisterSessionInMultiSessionWrapper(SessionInfoPtr.ToSharedRef());
@@ -498,7 +510,7 @@ void FRemoteSessionsManager::HandleSessionPingMessage(const FSessionPing& InMess
 {
 	// If this instance is not running trace, don't answer to the ping as the debugger instance will not be able to do anything useful with it
 #if !TRACE_BASED_DEBUGGERS_WITHOUT_TRACE
-	SendPong(InMessage);
+	SendPong(InMessage, InContext);
 #endif
 }
 
@@ -510,10 +522,10 @@ void FRemoteSessionsManager::HandleRecordingStatusUpdateMessage(const FRecording
 		check(SessionInfoPtr);
 
 		const FGuid& LastKnownDebuggerTypeId = SessionInfoPtr->LastKnownRecordingState.DebuggerId;
-		const bool bLastKnownDebuggerIsRecording = SessionInfoPtr->LastKnownRecordingState.bIsRecording;
+		const FGuid LastKnownRecordingRequestedId = SessionInfoPtr->LastKnownRecordingState.RequesterId;
 		if (LastKnownDebuggerTypeId.IsValid()
 			&& LastKnownDebuggerTypeId != Message.DebuggerId
-			&& bLastKnownDebuggerIsRecording)
+			&& LastKnownRecordingRequestedId.IsValid())
 		{
 			ensureMsgf(false, TEXT("A trace-based debugger changed is recording status while another debugger is already recording."
 				" This is not supported and indicates a problem on the runtime modules side."));
@@ -524,15 +536,20 @@ void FRemoteSessionsManager::HandleRecordingStatusUpdateMessage(const FRecording
 		// process request from their debugger Id.
 		SessionInfoPtr->LastKnownRecordingState = Message;
 
-		if (bLastKnownDebuggerIsRecording != Message.bIsRecording)
+		// Only process message if we (current application instance) instigated the recording
+		if (LastKnownRecordingRequestedId == LocalEditorSessionID
+			|| Message.RequesterId == LocalEditorSessionID)
 		{
-			if (Message.bIsRecording)
-			{
-				RecordingStartedDelegate.Broadcast(SessionInfoPtr);
-			}
-			else
+			if (LastKnownRecordingRequestedId != Message.RequesterId)
 			{
-				RecordingStoppedDelegate.Broadcast(SessionInfoPtr);
+				if (Message.RequesterId.IsValid())
+				{
+					RecordingStartedDelegate.Broadcast(SessionInfoPtr);
+				}
+				else
+				{
+					RecordingStoppedDelegate.Broadcast(SessionInfoPtr);
+				}
 			}
 		}
 	}
@@ -591,7 +608,7 @@ void FRemoteSessionsManager::HandleFullSessionStateResponseMessage(const FFullSe
 	{
 		if (const TSharedPtr<FSessionInfo>& SessionInfo = *SessionInfoPtr)
 		{
-			SessionInfo->LastKnownRecordingState.bIsRecording = InMessage.bIsRecording;
+			SessionInfo->LastKnownRecordingState.RequesterId = InMessage.RecordingRequesterId;
 			SessionInfo->LastKnownRecordingState.DebuggerId = InMessage.DebuggerId;
 
 			// Allow external handlers to process the message
diff --git a/Engine/Source/Runtime/Experimental/TraceBasedDebuggers/Private/TraceBasedDebuggerRuntime.cpp b/Engine/Source/Runtime/Experimental/TraceBasedDebuggers/Private/TraceBasedDebuggerRuntime.cpp
--- a/Engine/Source/Runtime/Experimental/TraceBasedDebuggers/Private/TraceBasedDebuggerRuntime.cpp
+++ b/Engine/Source/Runtime/Experimental/TraceBasedDebuggers/Private/TraceBasedDebuggerRuntime.cpp
@@ -1,4 +1,4 @@
-﻿// Copyright Epic Games, Inc. All Rights Reserved.
+// Copyright Epic Games, Inc. All Rights Reserved.
 
 #include "TraceBasedDebuggerRuntime.h"
 
@@ -69,6 +69,7 @@ void FRuntimeModule::StartRecording(const TConstArrayView<FString> Args)
 {
 	using namespace UE::TraceBasedDebuggers;
 	FStartRecordingCommandMessage StartRecordingCommand;
+	StartRecordingCommand.InstanceId = FApp::GetInstanceId();
 
 	// Default (i.e., no arguments) or "File" trace to file using GenerateRecordingFileName
 	if (Args.Num() == 0 || Args[0] == TEXT("File"))
@@ -152,7 +153,7 @@ void FRuntimeModule::StartRecording(const TConstArrayView<FString> Args)
 void FRuntimeModule::StartRecording(const FStartRecordingCommandMessage& InRecordingStartCommand)
 {
 	using namespace UE::TraceBasedDebuggers;
-	if (!ensureAlwaysMsgf(!bIsRecording, TEXT("Received a recording start command while a recording is already active")))
+	if (!ensureAlwaysMsgf(!IsRecording(), TEXT("Received a recording start command while a recording is already active")))
 	{
 		UE_LOG_REF(LogCategory, Log, TEXT("[%hs] There is an active recording. Attempting to stop it..."), __func__)
 		StopRecording();
@@ -161,6 +162,8 @@ void FRuntimeModule::StartRecording(const FStartRecordingCommandMessage& InRecor
 	// Start with a generic Failure reason
 	FText FailureReason = LOCTEXT("SeeLogsForErrorDetailsText", "Please see the logs for more details...");
 
+	bool bIsRecording = false;
+
 #if UE_TRACE_ENABLED || UE_TRACE_MINIMAL_ENABLED
 
 	// Other tools could be using trace
@@ -290,6 +293,24 @@ void FRuntimeModule::StartRecording(const FStartRecordingCommandMessage& InRecor
 		// Allow derived classes to perform additional actions on failure
 		OnRecordingStartFailedInternal(FailureReason);
 	}
+
+	if (bIsRecording)
+	{
+		if (ensureMsgf(InRecordingStartCommand.InstanceId.IsValid()
+			, TEXT("FStartRecordingCommandMessage sent without a valid InstanceId, using local instance Id as a fallback."
+			" Make sure all code paths properly set 'InstanceId'.")))
+		{
+			ActiveRecordingRequesterId = InRecordingStartCommand.InstanceId;
+		}
+		else
+		{
+			ActiveRecordingRequesterId = FApp::GetInstanceId();
+		}
+	}
+	else
+	{
+		ActiveRecordingRequesterId = FGuid{};
+	}
 }
 
 void FRuntimeModule::StopRecording()
@@ -302,7 +323,7 @@ void FRuntimeModule::StopRecording()
 
 	CurrentTransportMode = ETraceTransportMode::Invalid;
 
-	if (!bIsRecording)
+	if (!IsRecording())
 	{
 		UE_LOG_REF(LogCategory, Warning, TEXT("[%hs] Attempted to stop recorded when there is no recording active."), __func__);
 		return;
@@ -328,7 +349,7 @@ void FRuntimeModule::StopRecording()
 		RelayWriter->Close();
 	}
 
-	bIsRecording = false;
+	ActiveRecordingRequesterId = FGuid{};
 	AccumulatedRecordingTime = 0.0f;
 
 	// Allow derived classes to perform additional actions
@@ -490,7 +511,7 @@ void FRuntimeModule::ClearRelayDataPumpDelegates()
 
 bool FRuntimeModule::RecordingTimerTick(float DeltaTime)
 {
-	if (bIsRecording)
+	if (IsRecording())
 	{
 		AccumulatedRecordingTime += DeltaTime;
 	}
@@ -508,7 +529,7 @@ void FRuntimeModule::HandleTraceConnectionEstablished()
 		OnConnectionDelegateHandle.Reset();
 	}
 
-	if (!ensureAlwaysMsgf(bIsRecording, TEXT("Received a trace connection established callback but no trace is active. This should not happen.")))
+	if (!ensureAlwaysMsgf(IsRecording(), TEXT("Received a trace connection established callback but no trace is active. This should not happen.")))
 	{
 		return;
 	}
@@ -526,7 +547,7 @@ void FRuntimeModule::HandleTraceConnectionEstablished()
 
 void FRuntimeModule::HandleTraceStopRequest(FTraceAuxiliary::EConnectionType TraceType, const FString& TraceDestination)
 {
-	if (bIsRecording)
+	if (IsRecording())
 	{
 		if (!ensure(bRequestedStop))
 		{
diff --git a/Engine/Source/Runtime/Experimental/TraceBasedDebuggers/Public/RemoteSessionsManager.h b/Engine/Source/Runtime/Experimental/TraceBasedDebuggers/Public/RemoteSessionsManager.h
--- a/Engine/Source/Runtime/Experimental/TraceBasedDebuggers/Public/RemoteSessionsManager.h
+++ b/Engine/Source/Runtime/Experimental/TraceBasedDebuggers/Public/RemoteSessionsManager.h
@@ -72,7 +72,7 @@ struct FRecordingStatusMessage
 	FGuid DebuggerId = FGuid();
 
 	UPROPERTY()
-	bool bIsRecording = false;
+	FGuid RequesterId = FGuid();
 
 	UPROPERTY()
 	float ElapsedTime = false;
@@ -98,6 +98,9 @@ struct FStartRecordingCommandMessage
 {
 	GENERATED_BODY()
 
+	UPROPERTY()
+	FGuid InstanceId = FGuid();
+
 	UPROPERTY()
 	FString Target;
 
@@ -382,10 +385,11 @@ private:
 	void SendPing();
 
 	/**
-	 * Broadcast to the network a small subset of this instance information, in to a received session ping.
+	 * Replies directly to the ping sender with a small subset of this instance information.
 	 * @param InMessage Ping message that was received
+	 * @param InContext Message context for the received Ping
 	 */
-	void SendPong(const FSessionPing& InMessage);
+	void SendPong(const FSessionPing& InMessage, const TSharedRef<IMessageContext>& InContext);
 
 	void HandleSessionPongMessage(const FSessionPong& InMessage, const TSharedRef<IMessageContext>& InContext);
 	void HandleSessionPingMessage(const FSessionPing& InMessage, const TSharedRef<IMessageContext>& InContext);
diff --git a/Engine/Source/Runtime/Experimental/TraceBasedDebuggers/Public/SessionInfo.h b/Engine/Source/Runtime/Experimental/TraceBasedDebuggers/Public/SessionInfo.h
--- a/Engine/Source/Runtime/Experimental/TraceBasedDebuggers/Public/SessionInfo.h
+++ b/Engine/Source/Runtime/Experimental/TraceBasedDebuggers/Public/SessionInfo.h
@@ -59,7 +59,7 @@ struct FFullSessionInfoResponseMessage
 	TArray<FInstancedStruct> DebuggerSpecificData;
 
 	UPROPERTY()
-	bool bIsRecording = false;
+	FGuid RecordingRequesterId = FGuid();
 };
 
 UENUM()
@@ -129,11 +129,11 @@ public:
 
 	UE_API const FTraceConnectionDetails& GetTraceConnectionDetails();
 
-	/** @return Whether any debugger, regardless of its type, is currently recording */
+	/** @return Whether any debugger is currently recording, regardless of its type, or by which application instance it was requested */
 	UE_API virtual bool IsAnyDebuggerRecording() const;
 
 	/** 
-	 * Indicates whether a given debugger type is currently recording.
+	 * Indicates whether a given debugger type has an active recording requested by the current application instance.
 	 * @param DebuggerTypeId The identifier of the debugger type
 	 * @return Whether the provided debugger type id is currently recording
 	 */
diff --git a/Engine/Source/Runtime/Experimental/TraceBasedDebuggers/Public/TraceBasedDebuggerRuntime.h b/Engine/Source/Runtime/Experimental/TraceBasedDebuggers/Public/TraceBasedDebuggerRuntime.h
--- a/Engine/Source/Runtime/Experimental/TraceBasedDebuggers/Public/TraceBasedDebuggerRuntime.h
+++ b/Engine/Source/Runtime/Experimental/TraceBasedDebuggers/Public/TraceBasedDebuggerRuntime.h
@@ -6,6 +6,7 @@
 #include "Containers/Ticker.h"
 #include "ProfilingDebugging/TraceAuxiliary.h"
 #include "TraceDataRelayTransport.h"
+#include "Misc/Guid.h"
 
 #define UE_API TRACEBASEDDEBUGGERS_API
 
@@ -67,7 +68,13 @@ public:
 	/** Returns true if we are currently recording trace data for the current debugger */
 	bool IsRecording() const
 	{
-		return bIsRecording;
+		return ActiveRecordingRequesterId.IsValid();
+	}
+
+	/** @return Unique identifier of the application instance that requested the current recording, if active. */
+	const FGuid& GetRecordingRequesterId() const
+	{
+		return ActiveRecordingRequesterId;
 	}
 
 	/** @return Unique identifier for the current debugger type */
@@ -223,12 +230,15 @@ private:
 
 	TMap<FString, bool> OriginalTraceChannelsState;
 
+	/** Unique identifier of the debugger type associated to this instance */
 	FGuid DebuggerGuid;
 
+	/** Id of the application instance that requested the current active recording; invalid if not currently recording */
+	FGuid ActiveRecordingRequesterId = FGuid{};
+
 	float AccumulatedRecordingTime = 0.0f;
 
 	ETraceTransportMode CurrentTransportMode = ETraceTransportMode::Invalid;
-	bool bIsRecording = false;
 	bool bRequestedStop = false;
 };
 
diff --git a/Engine/Source/Runtime/Experimental/TraceBasedDebuggersAnalysis/Private/Widgets/RecordingControls.cpp b/Engine/Source/Runtime/Experimental/TraceBasedDebuggersAnalysis/Private/Widgets/RecordingControls.cpp
--- a/Engine/Source/Runtime/Experimental/TraceBasedDebuggersAnalysis/Private/Widgets/RecordingControls.cpp
+++ b/Engine/Source/Runtime/Experimental/TraceBasedDebuggersAnalysis/Private/Widgets/RecordingControls.cpp
@@ -186,6 +186,12 @@ TSharedRef<SWidget> FRecordingControls::GenerateTargetSessionDropdown()
 						return true;
 					}
 
+					// Allow derived classes to ignore some sessions
+					if (ShouldFilterOutSessionInternal(InSessionInfoRef))
+					{
+						return true;
+					}
+
 					const FText SessionNameAsText = FText::AsCultureInvariant(InSessionInfoRef->SessionName);
 					MenuBuilder.AddMenuEntry(
 						SessionNameAsText,
@@ -384,7 +390,7 @@ void FRecordingControls::CancelInFlightConnectionAttempt(const TSharedPtr<FSessi
 		}
 		else
 		{
-			UE_LOG_REF(LogCategory, Verbose, TEXT("[%hs] Attempted to cancel a connection of a Session that is no linger valid | Session ID [%s] ."), __func__, *SessionID.ToString());
+			UE_LOG_REF(LogCategory, Verbose, TEXT("[%hs] Attempted to cancel a connection of a Session that is no linger valid | Session ID [%s]."), __func__, *SessionID.ToString());
 		}
 
 		return;
@@ -399,17 +405,37 @@ void FRecordingControls::HandleRemoteSessionsUpdated()
 	const TSharedPtr<FSessionInfo> CurrentSessionPtr = GetCurrentSessionInfo();
 	if (!CurrentSessionPtr)
 	{
+		bool bValidSessionSelected = false;
+
 		// if the session that is currently selected is not valid, try to find a valid session and select that as default target session
 		if (const TSharedPtr<FRemoteSessionsManager> RemoteSessionManager = RemoteSessionManagerWeakPtr.Pin())
 		{
-			if (const TSharedPtr<FSessionInfo> FirstValidSession = RemoteSessionManager->GetFirstAvailableSingleSessionInfo().Pin())
-			{
-				SelectTargetSession(FirstValidSession->InstanceId);
-				return;
-			}
+			RemoteSessionManager->EnumerateActiveSessions([this, &bValidSessionSelected](const TSharedRef<FSessionInfo>& InSessionInfoRef)
+				{
+					if (EnumHasAnyFlags(InSessionInfoRef->GetSessionTypeAttributes(), ERemoteSessionAttributes::IsMultiSessionWrapper))
+					{
+						return true;
+					}
+
+					// Allow derived classes to ignore some sessions
+					if (ShouldFilterOutSessionInternal(InSessionInfoRef))
+					{
+						return true;
+					}
+
+					UE_LOG_REF(LogCategory, Verbose, TEXT("[%hs] Selecting first valid session [%s: %s]."), __func__, *InSessionInfoRef->SessionName, LexToString(InSessionInfoRef->BuildTargetType));
+					SelectTargetSession(InSessionInfoRef->InstanceId);
+					bValidSessionSelected = true;
+
+					// Stop the iteration
+					return false;
+				});
 		}
 
-		SelectTargetSession(FRemoteSessionsManager::InvalidSessionGUID);
+		if (!bValidSessionSelected)
+		{
+			SelectTargetSession(FRemoteSessionsManager::InvalidSessionGUID);
+		}
 	}
 }
 
@@ -597,6 +623,7 @@ void FRecordingControls::SetSessionRecordingState(const bool bIsRecording, const
 			};
 
 		TNotNull<TUniquePtr<FStartRecordingCommandMessage>> RecordingParams = BuildStartRecordingParamsInternal();
+		RecordingParams->InstanceId = FApp::GetInstanceId();
 		RecordingParams->RecordingMode = ERecordingMode::Live;
 		RecordingParams->TransportMode = DataTransportOverrideMode != ETraceTransportMode::Invalid ? DataTransportOverrideMode : ETraceTransportMode::Direct;
 
@@ -735,10 +762,15 @@ bool FRecordingControls::IsRecordingPossibleForSession(const TSharedPtr<FSession
 		return false;
 	}
 
-	// Do not allow recording if another trace-based debuggers is already recording in the current session
-	if (InSessionInfo->LastKnownRecordingState.DebuggerId.IsValid()
-		&& InSessionInfo->LastKnownRecordingState.DebuggerId != DebuggerGuid
-		&& InSessionInfo->LastKnownRecordingState.bIsRecording)
+	const FRecordingStatusMessage& LastState = InSessionInfo->LastKnownRecordingState;
+	const bool bWasRecording = LastState.DebuggerId.IsValid() && LastState.RequesterId.IsValid();
+
+	// Do not allow start/stop recording if a trace-based debuggers is already recording in the current session
+	// AND it is not associated to the current debugger type (i.e., DebuggerId)
+	// OR it was not requested by the current application instance (e.g., multiple editors trying to debug the same remote process)
+	if (bWasRecording
+		&& (LastState.DebuggerId != DebuggerGuid
+			|| LastState.RequesterId != FApp::GetInstanceId()))
 	{
 		return false;
 	}
@@ -836,6 +868,11 @@ bool FRecordingControls::IsRecordingAvailableForSessionInternal(const TSharedRef
 	return true;
 }
 
+bool FRecordingControls::ShouldFilterOutSessionInternal(const TSharedRef<FSessionInfo>& InSessionInfo) const
+{
+	return false;
+}
+
 bool FRecordingControls::IsRecording() const
 {
 	using namespace UE::TraceBasedDebuggers;
diff --git a/Engine/Source/Runtime/Experimental/TraceBasedDebuggersAnalysis/Public/Widgets/RecordingControls.h b/Engine/Source/Runtime/Experimental/TraceBasedDebuggersAnalysis/Public/Widgets/RecordingControls.h
--- a/Engine/Source/Runtime/Experimental/TraceBasedDebuggersAnalysis/Public/Widgets/RecordingControls.h
+++ b/Engine/Source/Runtime/Experimental/TraceBasedDebuggersAnalysis/Public/Widgets/RecordingControls.h
@@ -105,6 +105,11 @@ protected:
 
 	UE_API virtual bool IsRecordingAvailableForSessionInternal(const TSharedRef<FSessionInfo>& InSessionInfo) const;
 
+	/**
+	 * @return Whether a given session info should be excluded from the list exposed to the user
+	 */
+	UE_API virtual bool ShouldFilterOutSessionInternal(const TSharedRef<FSessionInfo>& InSessionInfo) const;
+
 	virtual void OnToggleRecordingStateInternal(const TSharedRef<FSessionInfo>& InSessionInfo) = 0;
 	virtual void CloseSessionByRemoteSessionIDInternal(const FGuid& InSessionId) = 0;
 
