Unshelved from pending changelist '50550921':

Fix IAD immediate event handler

diff --git a/Engine/Source/Runtime/Experimental/IoStore/OnDemand/Private/Statistics.cpp b/Engine/Source/Runtime/Experimental/IoStore/OnDemand/Private/Statistics.cpp
--- a/Engine/Source/Runtime/Experimental/IoStore/OnDemand/Private/Statistics.cpp
+++ b/Engine/Source/Runtime/Experimental/IoStore/OnDemand/Private/Statistics.cpp
@@ -88,34 +88,43 @@ extern int32 GIasHttpConcurrentRequests;
 
 ////////////////////////////////////////////////////////////////////////////////
 static FOnDemandImmediateAnalyticHandler OnDemandImmediateAnalyticEventHandler;
+static TArray<FOnDemandImmediateAnalytic> QueuedOnDemandImmediateAnalyticEventStartupEvents;
 static FMutex OnDemandImmediateAnalyticEventHandlerMutex;
 void OnDemandSetImmediateAnalyticHandler(FOnDemandImmediateAnalyticHandler&& EventHandler)
 {
 	TUniqueLock Lock(OnDemandImmediateAnalyticEventHandlerMutex);
 	OnDemandImmediateAnalyticEventHandler = MoveTemp(EventHandler);
-}
 
-static void SendImmediateAnalytic(FOnDemandImmediateAnalytic&& Event)
-{
-	// FSafeHandler avoids racing with OnDemandSetImmediateAnalyticHandler.
-	// The handler should be set once before any event is sent.
-	struct FSafeHandler
+	if (OnDemandImmediateAnalyticEventHandler.EventHandler)
 	{
-		FSafeHandler()
+		if (GIaxImmediateAnalyticsChance > 0.0f && FMath::FRand() <= GIaxImmediateAnalyticsChance)
 		{
-			TUniqueLock Lock(OnDemandImmediateAnalyticEventHandlerMutex);
-			EventHandler = MoveTemp(OnDemandImmediateAnalyticEventHandler.EventHandler);
+			for (FOnDemandImmediateAnalytic& QueuedEvent : QueuedOnDemandImmediateAnalyticEventStartupEvents)
+			{
+				OnDemandImmediateAnalyticEventHandler.EventHandler(MoveTemp(QueuedEvent));
+			}
 		}
 
-		TUniqueFunction<void(FOnDemandImmediateAnalytic)> EventHandler;
-	};
+		QueuedOnDemandImmediateAnalyticEventStartupEvents.Empty();
+	}
+}
 
+static void SendImmediateAnalytic(FOnDemandImmediateAnalytic&& Event)
+{
 	if (GIaxImmediateAnalyticsChance > 0.0f && FMath::FRand() <= GIaxImmediateAnalyticsChance)
 	{
-		static FSafeHandler SafeHandler;
-		if (SafeHandler.EventHandler)
+		TUniqueLock Lock(OnDemandImmediateAnalyticEventHandlerMutex);
+
+		if (OnDemandImmediateAnalyticEventHandler.EventHandler)
+		{
+			OnDemandImmediateAnalyticEventHandler.EventHandler(MoveTemp(Event));
+		}
+		else
 		{
-			SafeHandler.EventHandler(MoveTemp(Event));
+			if (ensure(QueuedOnDemandImmediateAnalyticEventStartupEvents.Num() < 16))
+			{
+				QueuedOnDemandImmediateAnalyticEventStartupEvents.Add(MoveTemp(Event));
+			}
 		}
 	}
 }
