TraceBasedDebuggersAnalysis: Only reference TraceServices & TraceAnalysis when building developer tools

#rnx

#Synced-CL 50554652

diff --git a/Engine/Source/Runtime/Experimental/TraceBasedDebuggersAnalysis/TraceBasedDebuggersAnalysis.Build.cs b/Engine/Source/Runtime/Experimental/TraceBasedDebuggersAnalysis/TraceBasedDebuggersAnalysis.Build.cs
--- a/Engine/Source/Runtime/Experimental/TraceBasedDebuggersAnalysis/TraceBasedDebuggersAnalysis.Build.cs
+++ b/Engine/Source/Runtime/Experimental/TraceBasedDebuggersAnalysis/TraceBasedDebuggersAnalysis.Build.cs
@@ -41,7 +41,7 @@ namespace UnrealBuildTool.Rules
 			}
 
 			// Allow debugger trace analysis on desktop platforms
-			if (Target.Platform.IsInGroup(UnrealPlatformGroup.Desktop))
+			if (Target.Platform.IsInGroup(UnrealPlatformGroup.Desktop) && Target.bBuildDeveloperTools)
 			{
 				PublicDefinitions.Add("UE_WITH_TRACE_BASED_DEBUGGERS_ANALYSIS=1");
 				PublicDependencyModuleNames.AddRange(
