Dataflow:
- Moved ensure() after logging for easier debugging

#Synced-CL 50577292

diff --git a/Engine/Source/Runtime/Experimental/Dataflow/Core/Private/Dataflow/DataflowNode.cpp b/Engine/Source/Runtime/Experimental/Dataflow/Core/Private/Dataflow/DataflowNode.cpp
--- a/Engine/Source/Runtime/Experimental/Dataflow/Core/Private/Dataflow/DataflowNode.cpp
+++ b/Engine/Source/Runtime/Experimental/Dataflow/Core/Private/Dataflow/DataflowNode.cpp
@@ -1229,9 +1229,9 @@ bool FDataflowNode::ValidateConnections()
 				{
 					if (!FindInput(PropertyIt.Value()))
 					{
-						ensure(false);
 						UE_LOG(LogChaos, Warning, TEXT("Missing dataflow RegisterInputConnection in constructor for (%s:%s)"), *GetName().ToString(), *PropName.ToString())
 							bHasValidConnections = false;
+						ensure(false);
 					}
 				}
 				if (Property->HasMetaData(FDataflowNode::DataflowOutput))
@@ -1239,9 +1239,9 @@ bool FDataflowNode::ValidateConnections()
 					const FDataflowOutput* const OutputConnection = FindOutput(PropertyIt.Value());
 					if(!OutputConnection)
 					{
-						ensure(false);
 						UE_LOG(LogChaos, Warning, TEXT("Missing dataflow RegisterOutputConnection in constructor for (%s:%s)"), *GetName().ToString(),*PropName.ToString());
 						bHasValidConnections = false;
+						ensure(false);
 					}
 
 					// If OutputConnection is valid, validate passthrough connections if they exist
@@ -1261,39 +1261,41 @@ bool FDataflowNode::ValidateConnections()
 						const FDataflowInput* const PassthroughConnectionInput = OutputConnection->GetPassthroughInput();
 						if (PassthroughConnectionInput == nullptr)
 						{
-							ensure(false);
 							UE_LOG(LogChaos, Warning, TEXT("Missing DataflowPassthrough registration for (%s:%s)"), *GetName().ToString(), *PropName.ToString());
 							bHasValidConnections = false;
+							ensure(false);
+
+							return false;
 						}
 
 						const FDataflowInput* const PassthroughConnectionInputFromMetadata = FindInput(FName(FullPassthroughName));
 
 						if(PassthroughConnectionInput != PassthroughConnectionInputFromMetadata)
 						{
-							ensure(false);
 							UE_LOG(LogChaos, Warning, TEXT("Mismatch in declared and registered DataflowPassthrough connection; (%s:%s vs %s)"), *GetName().ToString(), *FullPassthroughName, *PassthroughConnectionInput->GetName().ToString());
 							bHasValidConnections = false;
+							ensure(false);
 						}
 
 						if(!PassthroughConnectionInputFromMetadata)
 						{
-							ensure(false);
 							UE_LOG(LogChaos, Warning, TEXT("Incorrect DataflowPassthrough Connection set for (%s:%s)"), *GetName().ToString(), *PropName.ToString());
 							bHasValidConnections = false;
+							ensure(false);
 						}
 
 						else if(OutputConnection->GetType() != PassthroughConnectionInput->GetType())
 						{
-							ensure(false);
 							UE_LOG(LogChaos, Warning, TEXT("DataflowPassthrough connection types mismatch for (%s:%s)"), *GetName().ToString(), *PropName.ToString());
 							bHasValidConnections = false;
+							ensure(false);
 						}
 					}
 					else if(OutputConnection->GetPassthroughInput()) 
 					{
-						ensure(false);
 						UE_LOG(LogChaos, Warning, TEXT("Missing DataflowPassthrough declaration for (%s:%s)"), *GetName().ToString(), *PropName.ToString());
 						bHasValidConnections = false;
+						ensure(false);
 					}
 				}
 			}
