Fix compile errors caused by VectorMaskBits being in the Chaos namespace with overrides

#rnx

diff --git a/Engine/Source/Runtime/Experimental/Chaos/Public/Chaos/EPAVectorized.h b/Engine/Source/Runtime/Experimental/Chaos/Public/Chaos/EPAVectorized.h
--- a/Engine/Source/Runtime/Experimental/Chaos/Public/Chaos/EPAVectorized.h
+++ b/Engine/Source/Runtime/Experimental/Chaos/Public/Chaos/EPAVectorized.h
@@ -33,7 +33,7 @@ struct VectorTEPAEntry
 
 	FORCEINLINE bool operator>(const VectorTEPAEntry& Other) const
 	{
-		return static_cast<bool>(VectorMaskBits(VectorCompareGT(Distance, Other.Distance)));
+		return static_cast<bool>(::VectorMaskBits(VectorCompareGT(Distance, Other.Distance)));
 	}
 
 	FORCEINLINE_DEBUGGABLE bool Initialize(const VectorRegister4Float* VerticesA, const VectorRegister4Float* VerticesB, int32 InIdx0, int32 InIdx1, int32 InIdx2, const TVector<int32, 3>& InAdjFaces, const TVector<int32, 3>& InAdjEdges)
@@ -48,7 +48,7 @@ struct VectorTEPAEntry
 		const VectorRegister4Float NormLenSq = VectorDot3(Norm, Norm);
 		constexpr VectorRegister4Float Eps = MakeVectorRegisterFloatConstant(1e-4f, 1e-4f, 1e-4f, 1e-4f);
 		const VectorRegister4Float EpsGTNormLenSq = VectorCompareGT(Eps, NormLenSq);
-		if (VectorMaskBits(EpsGTNormLenSq))
+		if (::VectorMaskBits(EpsGTNormLenSq))
 		{
 			return false;
 		}
@@ -141,7 +141,7 @@ struct VectorTEPAEntry
 
 		const VectorRegister4Float IsFalse = VectorBitwiseOr(IsFalse1, IsFalse2);
 
-		return !static_cast<bool>(VectorMaskBits(IsFalse));
+		return !static_cast<bool>(::VectorMaskBits(IsFalse));
 	}
 };
 template <typename SupportALambda, typename SupportBLambda >
@@ -166,7 +166,7 @@ bool VectorInitializeEPA(TArray<VectorRegister4Float>& VertsA, TArray<VectorRegi
 
 		const VectorRegister4Float IsDist1GEDist0 = VectorCompareGE(Dist1, Dist0);
 
-		if (VectorMaskBits(IsDist1GEDist0))
+		if (::VectorMaskBits(IsDist1GEDist0))
 		{
 			VertsA.Add(A1);
 			VertsB.Add(B1);
@@ -199,7 +199,7 @@ bool VectorInitializeEPA(TArray<VectorRegister4Float>& VertsA, TArray<VectorRegi
 
 		constexpr VectorRegister4Float Eps = MakeVectorRegisterFloatConstant(1e-4f, 1e-4f, 1e-4f, 1e-4f);
 		const VectorRegister4Float Valid = VectorCompareGT(VectorDot3(Dir, Dir),  Eps);
-		if (VectorMaskBits(Valid))	//two verts given are distinct
+		if (::VectorMaskBits(Valid))	//two verts given are distinct
 		{
 			// return MakeVectorRegisterFloatFromDouble(LoadFloat3(Center));
 			const VectorRegister4Float DirAbs = VectorAbs(Dir);
@@ -303,7 +303,7 @@ bool VectorInitializeEPA(TArray<VectorRegister4Float>& VertsA, TArray<VectorRegi
 			MaxSignedDistance = VectorSelect(AbsDistGTMax, Entry.Distance, MaxSignedDistance);
 		}
 		VectorRegister4Float MaxSignedDistanceNeg = VectorCompareGT(VectorZero(), MaxSignedDistance);
-		if (VectorMaskBits(MaxSignedDistanceNeg))
+		if (::VectorMaskBits(MaxSignedDistanceNeg))
 		{
 			for (VectorTEPAEntry& Entry : OutEntries)
 			{
@@ -336,7 +336,7 @@ FORCEINLINE_DEBUGGABLE void VectorEPAComputeVisibilityBorder(TEPAWorkingArray<Ve
 		VectorTEPAEntry& Entry = Entries[FloodEntry.EntryIdx];
 		if (!Entry.bObsolete)
 		{
-			if (VectorMaskBits(VectorCompareLE(Entry.DistanceToPlane(W), Eps)))
+			if (::VectorMaskBits(VectorCompareLE(Entry.DistanceToPlane(W), Eps)))
 			{
 				//W can't see this triangle so mark the edge as a border
 				OutBorderEdges.Add(FloodEntry);
@@ -385,10 +385,10 @@ FORCEINLINE_DEBUGGABLE void VectorComputeEPAResults(const VectorRegister4Float*
 	constexpr VectorRegister4Float Eps = MakeVectorRegisterFloatConstant(EpsFloat, EpsFloat, EpsFloat, EpsFloat);
 
 	// @todo(chaos): pass in epsilon? Does it need to match anything in GJK?
-	if (VectorMaskBits(VectorCompareGT(Eps, OutPenetration)))	//if closest point is on the origin (edge case when surface is right on the origin)
+	if (::VectorMaskBits(VectorCompareGT(Eps, OutPenetration)))	//if closest point is on the origin (edge case when surface is right on the origin)
 	{
 		OutDir = Entry.PlaneNormal;	//just fall back on plane normal
-		if (VectorMaskBits(VectorCompareGT(VectorZero(), Entry.Distance)))
+		if (::VectorMaskBits(VectorCompareGT(VectorZero(), Entry.Distance)))
 		{
 			OutPenetration = VectorNegate(OutPenetration); // We are a bit outside of the shape so penetration is negative
 		}
@@ -396,7 +396,7 @@ FORCEINLINE_DEBUGGABLE void VectorComputeEPAResults(const VectorRegister4Float*
 	else
 	{
 		OutDir = VectorDivide(OutDir, OutPenetration);
-		if (VectorMaskBits(VectorCompareGT(VectorZero(), Entry.Distance)))
+		if (::VectorMaskBits(VectorCompareGT(VectorZero(), Entry.Distance)))
 		{
 			//The origin is on the outside, so the direction is reversed
 			OutDir = VectorNegate(OutDir);
@@ -448,7 +448,7 @@ EEPAResult VectorEPA(TArray<VectorRegister4Float>& VertsABuffer, TArray<VectorRe
 	TEPAWorkingArray<int32> Queue;
 	for(int32 Idx = 0; Idx < Entries.Num(); ++Idx)
 	{
-		const bool bIsZeroGTDist = static_cast<bool>(VectorMaskBits(VectorCompareGE(VectorZero(), Entries[Idx].Distance)));
+		const bool bIsZeroGTDist = static_cast<bool>(::VectorMaskBits(VectorCompareGE(VectorZero(), Entries[Idx].Distance)));
 		//ensure(Entries[Idx].Distance > -Eps);
 		// Entries[Idx].Distance <= 0.0f is true if the origin is a bit out of the polytope (we need to support this case for robustness)
 		if(bIsZeroGTDist || Entries[Idx].IsOriginProjectedInside(VertsABuffer.GetData(), VertsBBuffer.GetData(), OriginInsideEps))
@@ -493,7 +493,7 @@ EEPAResult VectorEPA(TArray<VectorRegister4Float>& VertsABuffer, TArray<VectorRe
 			continue;
 		}
 
-		if (VectorMaskBits(VectorCompareGT(Entry.Distance, UpperBound)))
+		if (::VectorMaskBits(VectorCompareGT(Entry.Distance, UpperBound)))
 		{
 			ResultStatus = EEPAResult::Ok;
 			break;
@@ -518,7 +518,7 @@ EEPAResult VectorEPA(TArray<VectorRegister4Float>& VertsABuffer, TArray<VectorRe
 		// In this case the upper bound will be negative, at which point we should just exit. 
 		const VectorRegister4Float UpperBoundTolerance = VectorAdd(VectorMultiply(OnePlusRelativeEps, VectorAbs(LowerBound)), SmallFloatEps);
 		const VectorRegister4Float UpperBoundLEUpperBoundTolerance = VectorCompareLE(UpperBound, UpperBoundTolerance);
-		if (VectorMaskBits(UpperBoundLEUpperBoundTolerance))
+		if (::VectorMaskBits(UpperBoundLEUpperBoundTolerance))
 		{
 			ResultStatus = EEPAResult::Ok;
 			LastEntry = Entry;
@@ -526,7 +526,7 @@ EEPAResult VectorEPA(TArray<VectorRegister4Float>& VertsABuffer, TArray<VectorRe
 		}
 
 		const VectorRegister4Float LowGTUp = VectorCompareGT(LowerBound, UpperBound);
-		if (VectorMaskBits(LowGTUp))
+		if (::VectorMaskBits(LowGTUp))
 		{
 			//we cannot get any better than what we saw, so just return previous face
 			ResultStatus = EEPAResult::Ok;
@@ -582,10 +582,10 @@ EEPAResult VectorEPA(TArray<VectorRegister4Float>& VertsABuffer, TArray<VectorRe
 				// However these Entries still have good normals, and needs to be included to prevent
 				// this exiting with very deep penetration results
 				
-				if (bValidTri && VectorMaskBits(VectorCompareGE(UpperBound, NewEntry.Distance)))
+				if (bValidTri && ::VectorMaskBits(VectorCompareGE(UpperBound, NewEntry.Distance)))
 				{
 					// NewEntry.Distance <= 0.0f is if the origin is a bit out of the polytope
-					if (VectorMaskBits(VectorCompareGE(VectorZero(), NewEntry.Distance))|| NewEntry.IsOriginProjectedInside(VertsABuffer.GetData(), VertsBBuffer.GetData(), OriginInsideEps))
+					if (::VectorMaskBits(VectorCompareGE(VectorZero(), NewEntry.Distance))|| NewEntry.IsOriginProjectedInside(VertsABuffer.GetData(), VertsBBuffer.GetData(), OriginInsideEps))
 					{
 						Queue.Add(NewIdx);
 						bQueueDirty = true;
diff --git a/Engine/Source/Runtime/Experimental/Chaos/Public/Chaos/GJK.h b/Engine/Source/Runtime/Experimental/Chaos/Public/Chaos/GJK.h
--- a/Engine/Source/Runtime/Experimental/Chaos/Public/Chaos/GJK.h
+++ b/Engine/Source/Runtime/Experimental/Chaos/Public/Chaos/GJK.h
@@ -246,7 +246,7 @@ namespace Chaos
 			const T SupportBSimd = VectorAdd(VectorQuaternionRotateVector(Rotation, SupportBLocalSimd), Translation);
 			const T WSimd = VectorSubtract(SupportASimd, SupportBSimd);
 
-			if (VectorMaskBits(VectorCompareGT(VectorDot3(V, WSimd), Inflation)))
+			if (::VectorMaskBits(VectorCompareGT(VectorDot3(V, WSimd), Inflation)))
 			{
 				return false;
 			}
@@ -256,11 +256,11 @@ namespace Chaos
 			V = VectorSimplexFindClosestToOrigin<T, false>(SimplexSimd, NumVerts, BarycentricSimd, nullptr, nullptr);
 
 			const T NewDist2Simd = VectorDot3(V, V);
-			bNearZero = VectorMaskBits(VectorCompareLT(NewDist2Simd, Inflation2)) != 0;
+			bNearZero = ::VectorMaskBits(VectorCompareLT(NewDist2Simd, Inflation2)) != 0;
 
 			//as simplices become degenerate we will stop making progress. This is a side-effect of precision, in that case take V as the current best approximation
 			//question: should we take previous v in case it's better?
-			bool bMadeProgress = VectorMaskBits(VectorCompareLT(NewDist2Simd, PrevDist2Simd)) != 0;
+			bool bMadeProgress = ::VectorMaskBits(VectorCompareLT(NewDist2Simd, PrevDist2Simd)) != 0;
 			bTerminate = bNearZero || !bMadeProgress;
 			PrevDist2Simd = NewDist2Simd;
 
@@ -341,7 +341,7 @@ namespace Chaos
 
 			VectorRegister4Float VDotW = VectorDot3(V, W);
 
-			if (VectorMaskBits(VectorCompareGT(VDotW, InflationSimd)))
+			if (::VectorMaskBits(VectorCompareGT(VDotW, InflationSimd)))
 			{
 				return false;
 			}
@@ -351,11 +351,11 @@ namespace Chaos
 			V = VectorSimplexFindClosestToOrigin<VectorRegister4Float, false>(Simplex, NumVerts, Barycentric, nullptr, nullptr);
 			const VectorRegister4Float NewDist2 = VectorDot3(V, V);
 			
-			bNearZero = static_cast<bool>(VectorMaskBits(VectorCompareGT(Inflation2Simd, NewDist2)));
+			bNearZero = static_cast<bool>(::VectorMaskBits(VectorCompareGT(Inflation2Simd, NewDist2)));
 
 			//as simplices become degenerate we will stop making progress. This is a side-effect of precision, in that case take V as the current best approximation
 			//question: should we take previous v in case it's better?
-			const bool bMadeProgress = static_cast<bool>(VectorMaskBits(VectorCompareGT(PrevDist2, NewDist2)));
+			const bool bMadeProgress = static_cast<bool>(::VectorMaskBits(VectorCompareGT(PrevDist2, NewDist2)));
 			bTerminate = bNearZero || !bMadeProgress;
 
 			PrevDist2 = NewDist2;
@@ -1681,12 +1681,12 @@ namespace Chaos
 
 			T VDotWGTInflationSimd = VectorCompareGT(VDotW, Inflation);
 
-			if (VectorMaskBits(VDotWGTInflationSimd))
+			if (::VectorMaskBits(VDotWGTInflationSimd))
 			{
 				const T VDotRayDir = VectorDot3(V, RayDir);
 				T VDotRayDirGEZero = VectorCompareGE(VDotRayDir, TVectorZero<T>());
 
-				if (VectorMaskBits(VDotRayDirGEZero))
+				if (::VectorMaskBits(VDotRayDirGEZero))
 				{
 					return false;
 				}
@@ -1695,10 +1695,10 @@ namespace Chaos
 				// @todo(ccaulfield): this can still overflow - the comparisons against zero above should be changed (though not sure to what yet)
 				Lambda = VectorSubtract(Lambda, VectorDivide(VectorSubtract(VDotW, Inflation), VDotRayDir));
 				T LambdaGTPreLambda = VectorCompareGT(Lambda, PreLambda);
-				if (VectorMaskBits(LambdaGTPreLambda))
+				if (::VectorMaskBits(LambdaGTPreLambda))
 				{
 					T LambdaGTRayLength = VectorCompareGT(Lambda, RayLengthSimd);
-					if (VectorMaskBits(LambdaGTRayLength))
+					if (::VectorMaskBits(LambdaGTRayLength))
 					{
 						return false;
 					}
@@ -1724,11 +1724,11 @@ namespace Chaos
 				Simplex[NumVerts++] = W;	//this is really X - P which is what we need for simplex computation
 			}
 
-			if (bInflatedCloseEnough && VectorMaskBits(VectorCompareGE(VDotW, TVectorZero<T>())))
+			if (bInflatedCloseEnough && ::VectorMaskBits(VectorCompareGE(VDotW, TVectorZero<T>())))
 			{
 				//Inflated shapes are close enough, but we want MTD so we need to find closest point on core shape
 				const T VDotW2 = VectorDot3(VDotW, VDotW);
-				bCloseEnough = static_cast<bool>(VectorMaskBits(VectorCompareGE(VectorAdd(Eps2Simd, VDotW), GJKPreDist2)));
+				bCloseEnough = static_cast<bool>(::VectorMaskBits(VectorCompareGE(VectorAdd(Eps2Simd, VDotW), GJKPreDist2)));
 			}
 
 			if (!bCloseEnough)
@@ -1737,13 +1737,13 @@ namespace Chaos
 				V = VectorSimplexFindClosestToOrigin<T>(Simplex, NumVerts, Barycentric, As, Bs);
 
 				T NewDist2 = VectorDot3(V, V);	//todo: relative error
-				bCloseEnough = static_cast<bool>(VectorMaskBits(VectorCompareGT(Inflation2Simd, NewDist2)));
-				bDegenerate = static_cast<bool>(VectorMaskBits(VectorCompareGE(NewDist2, GJKPreDist2)));
+				bCloseEnough = static_cast<bool>(::VectorMaskBits(VectorCompareGT(Inflation2Simd, NewDist2)));
+				bDegenerate = static_cast<bool>(::VectorMaskBits(VectorCompareGE(NewDist2, GJKPreDist2)));
 				GJKPreDist2 = NewDist2;
 
 				if (bComputeMTD && bCloseEnough)
 				{
-					const T LambdaEqZero = VectorCompareEQ(Lambda, TVectorZero<T>());
+					const T LambdaEqZero = ::VectorCompareEQ(Lambda, TVectorZero<T>());
 					const T InGJKPreDist2GTEps2 = VectorCompareGT(GJKPreDist2, Eps2Simd);
 					const T Inflation22GTEps2 = VectorCompareGT(Inflation2Simd, Eps2Simd);
 
@@ -1757,7 +1757,7 @@ namespace Chaos
 					//	bCloseEnough = false;
 					//	bInflatedCloseEnough = true;
 					//}
-					bInflatedCloseEnough = static_cast<bool>(VectorMaskBits(IsInflatCloseEnough)) && Is4GTNumVerts;
+					bInflatedCloseEnough = static_cast<bool>(::VectorMaskBits(IsInflatCloseEnough)) && Is4GTNumVerts;
 					bCloseEnough = !bInflatedCloseEnough;
 				}
 			}
@@ -1831,7 +1831,7 @@ namespace Chaos
 				GJKClosestB = Bs[0];
 			}
 
-			if (VectorMaskBits(IsDone))
+			if (::VectorMaskBits(IsDone))
 			{
 				const T ClosestBInA = VectorAdd(StartPoint, GJKClosestB);
 				const T InGJKPreDist = VectorSqrt(GJKPreDist2);
diff --git a/Engine/Source/Runtime/Experimental/Chaos/Public/Chaos/SimplexVectorized.h b/Engine/Source/Runtime/Experimental/Chaos/Public/Chaos/SimplexVectorized.h
--- a/Engine/Source/Runtime/Experimental/Chaos/Public/Chaos/SimplexVectorized.h
+++ b/Engine/Source/Runtime/Experimental/Chaos/Public/Chaos/SimplexVectorized.h
@@ -51,7 +51,7 @@ namespace Chaos
 		Closest = VectorSelect(IsX0, X0, VectorSelect(IsX1, X1, Closest));
 
 		const T IsX0OrX1 = VectorBitwiseOr(IsX0, IsX1);
-		NumVerts =  VectorMaskBits(IsX0OrX1) ? 1 : NumVerts;
+		NumVerts =  ::VectorMaskBits(IsX0OrX1) ? 1 : NumVerts;
 
 		if constexpr (CalculatExtraInformation)
 		{
@@ -81,7 +81,7 @@ namespace Chaos
 		const T Eps2GENormal2 = VectorCompareGE(Eps2, TriNormal2);
 	
 
-		if (VectorMaskBits(Eps2GENormal2))
+		if (::VectorMaskBits(Eps2GENormal2))
 		{
 			NumVerts = 2;
 			return VectorLineSimplexFindOrigin<T, CalculatExtraInformation>(Simplex, NumVerts, OutBarycentric, As, Bs);
@@ -97,7 +97,7 @@ namespace Chaos
 		const T IsD2SEZero = VectorCompareGE(TVectorZero<T>(), d2);
 		const T IsA = VectorBitwiseAnd(IsD1SEZero, IsD2SEZero);
 
-		if (VectorMaskBits(IsA))
+		if (::VectorMaskBits(IsA))
 		{
 			NumVerts = 1;
 			if constexpr (CalculatExtraInformation)
@@ -116,7 +116,7 @@ namespace Chaos
 		const T IsD3GED4 = VectorCompareGE(d3, d4);
 		const T IsB = VectorBitwiseAnd(IsD3GEZero, IsD3GED4);
 
-		if (VectorMaskBits(IsB))
+		if (::VectorMaskBits(IsB))
 		{
 			NumVerts = 1;
 			if constexpr (CalculatExtraInformation)
@@ -143,7 +143,7 @@ namespace Chaos
 		const T IsNDABGTZero = VectorCompareGT(NormalizationDenominatorAB, TVectorZero<T>());
 		const T IsAB = VectorBitwiseAnd(VectorBitwiseAnd(IsZeroGEvc, IsD1GEZero), VectorBitwiseAnd(IsZeroGED3, IsNDABGTZero));
 
-		if (VectorMaskBits(IsAB))
+		if (::VectorMaskBits(IsAB))
 		{
 			NumVerts = 2;
 
@@ -165,7 +165,7 @@ namespace Chaos
 		const T IsD6GED5 = VectorCompareGE(d6, d5);
 		const T IsC = VectorBitwiseAnd(IsD6GEZero, IsD6GED5);
 
-		if (VectorMaskBits(IsC))
+		if (::VectorMaskBits(IsC))
 		{
 			NumVerts = 1;
 			if constexpr (CalculatExtraInformation)
@@ -193,7 +193,7 @@ namespace Chaos
 		const T IsNDACGTZero = VectorCompareGT(NormalizationDenominatorAC, TVectorZero<T>());
 		const T IsAC = VectorBitwiseAnd(VectorBitwiseAnd(IsZeroGEvb, IsD2GEZero), VectorBitwiseAnd(IsZeroGED6, IsNDACGTZero));
 
-		if (VectorMaskBits(IsAC))
+		if (::VectorMaskBits(IsAC))
 		{
 			const T w = VectorDivide(d2, NormalizationDenominatorAC);
 			NumVerts = 2;
@@ -225,7 +225,7 @@ namespace Chaos
 		const T IsNDBCGTZero = VectorCompareGT(NormalizationDenominatorBC, TVectorZero<T>());
 		const T IsBC = VectorBitwiseAnd(VectorBitwiseAnd(IsZeroGEva, IsD4MinusD3GEZero), VectorBitwiseAnd(IsD5MinusD6GEZero, IsNDBCGTZero));
 
-		if (VectorMaskBits(IsBC))
+		if (::VectorMaskBits(IsBC))
 		{
 			NumVerts = 2;
 			const T w = VectorDivide(d4MinusD3, NormalizationDenominatorBC);
@@ -279,11 +279,11 @@ namespace Chaos
 	FORCEINLINE bool VectorSignMatch(T A, T B)
 	{
 		T OneIsZero = VectorMultiply(A, B);
-		OneIsZero = VectorCompareEQ(OneIsZero, TVectorZero<T>());
+		OneIsZero = ::VectorCompareEQ(OneIsZero, TVectorZero<T>());
 
-		const bool IsZero = static_cast<bool>(VectorMaskBits(OneIsZero));
-		const int32 MaskA = VectorMaskBits(A);
-		const int32 MaskB = VectorMaskBits(B);
+		const bool IsZero = static_cast<bool>(::VectorMaskBits(OneIsZero));
+		const int32 MaskA = ::VectorMaskBits(A);
+		const int32 MaskB = ::VectorMaskBits(B);
 		return (MaskA == MaskB) && !IsZero;
 	}
 
@@ -353,7 +353,7 @@ namespace Chaos
 
 				const T MinGTDist = VectorCompareGT(MinTriangleDist2, Dist2);
 
-				const bool bFindClosest = static_cast<bool>(VectorMaskBits(MinGTDist)) || (ClosestTriangleIdx == INDEX_NONE);
+				const bool bFindClosest = static_cast<bool>(::VectorMaskBits(MinGTDist)) || (ClosestTriangleIdx == INDEX_NONE);
 
 				MinTriangleDist2 = bFindClosest ? Dist2 :  MinTriangleDist2;
 				ClosestTriangleIdx = bFindClosest ? IdxSimd[Idx] :  ClosestTriangleIdx;
