[Partial Island Sleep] - Add CVar to turn off level-based waking

#Synced-CL 50536252

diff --git a/Engine/Source/Runtime/Experimental/Chaos/Private/Chaos/Island/IslandManager.cpp b/Engine/Source/Runtime/Experimental/Chaos/Private/Chaos/Island/IslandManager.cpp
--- a/Engine/Source/Runtime/Experimental/Chaos/Private/Chaos/Island/IslandManager.cpp
+++ b/Engine/Source/Runtime/Experimental/Chaos/Private/Chaos/Island/IslandManager.cpp
@@ -60,6 +60,10 @@ namespace Chaos::CVars
 	bool bChaosSolverPartialIslandWake = true;
 	FAutoConsoleVariableRef CVarChaosSolverPartialIslandWake(TEXT("p.Chaos.Solver.Sleep.PartialIslandSleep.PartialIslandWake"), bChaosSolverPartialIslandWake, TEXT("[EXPERIMENTAL] Safeguard only: New constraints and particle state changes will trigger partial wake events instead of waking up the entire island.[def:true]"));
 
+	/** Cvar to enable/disable level-based waking for all islands */
+	bool bChaosSolverLevelBasedWake = true;
+	FAutoConsoleVariableRef CVarChaosSolverLevelBasedWake(TEXT("p.Chaos.Solver.Sleep.PartialIslandSleep.LevelBasedWake"), bChaosSolverLevelBasedWake, TEXT("[EXPERIMENTAL] Safeguard only: Wake events will always propagate to higher levels, i.e. away from the closest kinematic.[def:true]"));
+
 	/** Cvar to enable/disable forcing a partial island waking for all islands even if all particles are already awake.
 	This is just a safeguard to force resetting the particle sleep counter when a new wake event occurs. It should generally be set to false. */
 	bool bChaosSolverForcePartialIslandWake = false;
@@ -635,8 +639,8 @@ namespace Chaos::Private
 
 	bool FPBDIslandManager::ShouldAssignLevels(const FPBDIsland* Island) const
 	{
-		// Partial island sleeping always needs levels
-		return bAssignLevels || Island->IsPartialSleepEnabled();
+		// Partial island sleeping with level-based wake always needs levels
+		return bAssignLevels || (Island->IsPartialSleepEnabled() && CVars::bChaosSolverLevelBasedWake);
 	}
 
 	void FPBDIslandManager::SetIsDeterministic(const bool bInIsDeterministic)
@@ -3402,7 +3406,7 @@ namespace Chaos::Private
 						// Wake and propagate always to the next constraint if
 						const bool bIsTriggerNode = NodeIndex < NumNodesInQueue;
 						const bool bIsConfirmedPropagation = bIsTriggerNode && (Edge->Flags.bIsNew || bTriggerNodesAlwaysPropagate);
-						const bool bIsUpwardPropagation = OtherNode->Level >= Edge->Level;	// Propagation upward through the levels, i.e. away from closest kinematic
+						const bool bIsUpwardPropagation = ShouldAssignLevels(Edge->Island) ? OtherNode->Level >= Edge->Level : false;	// Propagation upward through the levels, i.e. away from closest kinematic
 						const bool bEdgePropagatesPartialWake = (CVars::ChaosSolverPartialSleepCollisionConstraintsOnly != 0) && Edge->GetConstraint() && !Edge->GetConstraint()->SupportsPartialIslandSleeping();
 
 						// 1) We process one of the trigger nodes and propagate through a new constraint (if bTriggerNodesAlwaysPropagate == false).
