FOnDemandIoStore::DebugFindStreamingChunkIds no longer returns optional chunks if downloading them has been disabled.

#rnx

- It's actually quite tricky for higher level code in our module to find out if optional chunks have been disabled or not so it is hard for the caller to filter them out. I don't see any reason why we'd want to call DebugFindStreamingChunkIds and get back chunkids that won't work so for now it's easier to just filter them out in the method directly.

#Synced-CL 50589248

diff --git a/Engine/Source/Runtime/Experimental/IoStore/OnDemand/Private/OnDemandIoStore.cpp b/Engine/Source/Runtime/Experimental/IoStore/OnDemand/Private/OnDemandIoStore.cpp
--- a/Engine/Source/Runtime/Experimental/IoStore/OnDemand/Private/OnDemandIoStore.cpp
+++ b/Engine/Source/Runtime/Experimental/IoStore/OnDemand/Private/OnDemandIoStore.cpp
@@ -968,6 +968,11 @@ TArray<FIoChunkId> FOnDemandIoStore::DebugFindStreamingChunkIds(int32 NumToFind)
 		{
 			for(const FIoChunkId& Id : Container->ChunkIds)
 			{
+				if (EnumHasAllFlags(StreamingOptions, EOnDemandStreamingOptions::OptionalBulkDataDisabled) && Id.GetChunkType() == EIoChunkType::OptionalBulkData)
+				{
+					continue;
+				}
+
 				IdsContainer.Add(Id);
 
 				if (IdsContainer.Num() == NumToFind)
