The 'ias.RunRequestTest' now outputs the amount of data downloaded and the download rate as well as the time taken for the command to run.

#rnx

#Synced-CL 50589248

diff --git a/Engine/Source/Runtime/Experimental/IoStore/OnDemand/Private/DebugCommands.cpp b/Engine/Source/Runtime/Experimental/IoStore/OnDemand/Private/DebugCommands.cpp
--- a/Engine/Source/Runtime/Experimental/IoStore/OnDemand/Private/DebugCommands.cpp
+++ b/Engine/Source/Runtime/Experimental/IoStore/OnDemand/Private/DebugCommands.cpp
@@ -125,18 +125,24 @@ void FOnDemandDebugCommands::RunRequestTest(const TArray<FString>& Args) const
 	TArray<FIoRequest> Requests;
 
 	FIoBatch IoBatch;
+	uint64 TotalDownloadSize = 0;
 	for (const FIoChunkId& Id : RequestIds)
 	{
 		FIoRequest Request = IoBatch.Read(Id, FIoReadOptions(), IoDispatcherPriority_Medium);
 		Requests.Emplace(MoveTemp(Request));
+
+		TotalDownloadSize += OnDemandIoStore->GetStreamingChunkInfo(Id).EncodedSize();
 	}
 
 	const double StartTime = FPlatformTime::Seconds();
 
-	IoBatch.IssueWithCallback([StartTime, bQuitAfterTest]()
+	IoBatch.IssueWithCallback([StartTime, TotalDownloadSize, bQuitAfterTest]()
 		{
 			const double TotalTime = FPlatformTime::Seconds() - StartTime;
-			UE_LOG(LogIas, Display, TEXT("IAS Request Test took %.3f(s)"), TotalTime);
+			const double DataSizeMiB = static_cast<double>(TotalDownloadSize) / (1024.0 * 1024.0);
+			const double DataRate = TotalTime > 0.0 ? DataSizeMiB / TotalTime : 0.0;
+
+			UE_LOG(LogIas, Display, TEXT("IAS Request Test downloaded %.2f(MiB) in %.3f(s) %.3f(MiB/s)"), DataSizeMiB, TotalTime, DataRate);
 
 			if (bQuitAfterTest)
 			{
