Chaos++: Removing the unused IRigidSceneFactory. The actual used class is IRigidFactory.

#Synced-CL 50653012

diff --git a/Engine/Source/Runtime/Experimental/RigidPhysics/Public/RigidPhysics/RigidFactory.h b/Engine/Source/Runtime/Experimental/RigidPhysics/Public/RigidPhysics/RigidFactory.h
--- a/Engine/Source/Runtime/Experimental/RigidPhysics/Public/RigidPhysics/RigidFactory.h
+++ b/Engine/Source/Runtime/Experimental/RigidPhysics/Public/RigidPhysics/RigidFactory.h
@@ -12,18 +12,13 @@ namespace UE::Physics
 	class IRigidFactory
 	{
 	public:
-		IRigidFactory()
-		{
-		}
-
-		virtual ~IRigidFactory()
-		{
-		}
+		IRigidFactory() = default;
+		virtual ~IRigidFactory() = default;
 
 		virtual const FRigidTypeId& GetSceneTypeId() const = 0;
 		virtual FRigidSceneHandle CreateScene(const FRigidDebugName& InName, const IRigidSceneSettings* InSettings) = 0;
 		virtual void DestroyScene(const FRigidSceneHandle& InSceneHandle) = 0;
 	};
-}
+} // namespace UE::Physics
 
-#endif
\ No newline at end of file
+#endif // UE_RIGIDPHYSICS_API_ENABLED
diff --git a/Engine/Source/Runtime/Experimental/RigidPhysics/Public/RigidPhysics/RigidFwd.h b/Engine/Source/Runtime/Experimental/RigidPhysics/Public/RigidPhysics/RigidFwd.h
--- a/Engine/Source/Runtime/Experimental/RigidPhysics/Public/RigidPhysics/RigidFwd.h
+++ b/Engine/Source/Runtime/Experimental/RigidPhysics/Public/RigidPhysics/RigidFwd.h
@@ -66,11 +66,11 @@ namespace UE::Physics
 
 	class IRigidBody;
 	class IRigidBodyContainer;
+	class IRigidFactory;
 	class IRigidGeometryCollection;
 	class IRigidLockable;
 	class IRigidQuery;
 	class IRigidScene;
-	class IRigidSceneFactory;
 	class IRigidSceneModifier;
 	class IRigidSceneSettings;
 	class IRigidTyped;
diff --git a/Engine/Source/Runtime/Experimental/RigidPhysics/Public/RigidPhysics/RigidScene.h b/Engine/Source/Runtime/Experimental/RigidPhysics/Public/RigidPhysics/RigidScene.h
--- a/Engine/Source/Runtime/Experimental/RigidPhysics/Public/RigidPhysics/RigidScene.h
+++ b/Engine/Source/Runtime/Experimental/RigidPhysics/Public/RigidPhysics/RigidScene.h
@@ -15,21 +15,6 @@
 
 namespace UE::Physics
 {
-	class IRigidSceneFactory
-	{
-	public:
-		IRigidSceneFactory()
-		{
-		}
-
-		virtual ~IRigidSceneFactory()
-		{
-		}
-
-		virtual FRigidSceneHandle CreateScene(const FRigidDebugName& InName) = 0;
-		virtual void DestroyScene(const FRigidSceneHandle& InSceneHandle) = 0;
-	};
-
 	// The interface to a rigid body scene. The rigid body scene owns the physics simulations and
 	// manages a set of rigid body containers which in turn manage sets of rigid bodies.
 	class UE_INTERNAL IRigidScene : public IRigidLockable
