Chaos++: Refactored handles to separate files to minimize dependencies. This removes the need to implement some functions in different header files from the declaration.

#Synced-CL 50653012

diff --git a/Engine/Plugins/Experimental/ChaosRigidPhysicsAsync/Tests/Private/SceneLockBenchmarks.cpp b/Engine/Plugins/Experimental/ChaosRigidPhysicsAsync/Tests/Private/SceneLockBenchmarks.cpp
--- a/Engine/Plugins/Experimental/ChaosRigidPhysicsAsync/Tests/Private/SceneLockBenchmarks.cpp
+++ b/Engine/Plugins/Experimental/ChaosRigidPhysicsAsync/Tests/Private/SceneLockBenchmarks.cpp
@@ -9,6 +9,7 @@
 #include "PhysicsInterfaceDeclaresCore.h"
 #include "Physics/Experimental/PhysInterface_Chaos.h"
 #include "PhysicsProxy/SingleParticlePhysicsProxy.h"
+#include "RigidPhysics/RigidBody.h"
 #include "RigidPhysics/RigidCollider.h"
 #include "RigidPhysics/RigidContext.h"
 
diff --git a/Engine/Source/Runtime/Experimental/RigidPhysics/Private/RigidPhysics/RigidBodyContainerHandle.cpp b/Engine/Source/Runtime/Experimental/RigidPhysics/Private/RigidPhysics/RigidBodyContainerHandle.cpp
new file mode 100644
--- /dev/null
+++ b/Engine/Source/Runtime/Experimental/RigidPhysics/Private/RigidPhysics/RigidBodyContainerHandle.cpp
@@ -0,0 +1,38 @@
+// Copyright Epic Games, Inc. All Rights Reserved.
+#include "RigidPhysics/RigidBodyContainerHandle.h"
+
+#include "RigidPhysics/RigidScene.h"
+
+#if UE_RIGIDPHYSICS_API_ENABLED
+
+namespace UE::Physics
+{
+	FRigidBodyContainerHandle::FRigidBodyContainerHandle(const FRigidSceneHandle& InSceneHandle, const FRigidObjectId& InBodyContainerId)
+		: BodyContainerId(InBodyContainerId)
+		, SceneHandle(InSceneHandle)
+	{
+	}
+
+	void FRigidBodyContainerHandle::Reset()
+	{
+		BodyContainerId.Reset();
+		SceneHandle.Reset();
+	}
+
+	const FRigidObjectId& FRigidBodyContainerHandle::GetId() const
+	{
+		return BodyContainerId;
+	}
+
+	const FRigidSceneHandle FRigidBodyContainerHandle::GetSceneHandle() const
+	{
+		return SceneHandle;
+	}
+
+	IRigidBodyContainer* FRigidBodyContainerHandle::PinInternal(const IRigidScene* InScene) const
+	{
+		return InScene->GetBodyContainer(GetId());
+	}
+} // namespace UE::Physics
+
+#endif // UE_RIGIDPHYSICS_API_ENABLED
\ No newline at end of file
diff --git a/Engine/Source/Runtime/Experimental/RigidPhysics/Private/RigidPhysics/RigidBodyHandle.cpp b/Engine/Source/Runtime/Experimental/RigidPhysics/Private/RigidPhysics/RigidBodyHandle.cpp
new file mode 100644
--- /dev/null
+++ b/Engine/Source/Runtime/Experimental/RigidPhysics/Private/RigidPhysics/RigidBodyHandle.cpp
@@ -0,0 +1,38 @@
+// Copyright Epic Games, Inc. All Rights Reserved.
+#include "RigidPhysics/RigidBodyHandle.h"
+
+#include "RigidPhysics/RigidBodyContainer.h"
+
+#if UE_RIGIDPHYSICS_API_ENABLED
+
+namespace UE::Physics
+{
+	FRigidBodyHandle::FRigidBodyHandle(const FRigidBodyContainerHandle& InBodyContainerHandle, const FRigidObjectId& InBodyId)
+		: BodyId(InBodyId)
+		, BodyContainerHandle(InBodyContainerHandle)
+	{
+	}
+
+	const FRigidObjectId& FRigidBodyHandle::GetId() const
+	{
+		return BodyId;
+	}
+
+	const FRigidBodyContainerHandle& FRigidBodyHandle::GetContainerHandle() const
+	{
+		return BodyContainerHandle;
+	}
+
+	void FRigidBodyHandle::Reset()
+	{
+		BodyContainerHandle.Reset();
+		BodyId.Reset();
+	}
+
+	IRigidBody* FRigidBodyHandle::PinInternal(const IRigidBodyContainer* Container) const
+	{
+		return Container->GetBody(GetId());
+	}
+} // namespace UE::Physics
+
+#endif // UE_RIGIDPHYSICS_API_ENABLED
\ No newline at end of file
diff --git a/Engine/Source/Runtime/Experimental/RigidPhysics/Private/RigidPhysics/RigidGeometryCollectionHandle.cpp b/Engine/Source/Runtime/Experimental/RigidPhysics/Private/RigidPhysics/RigidGeometryCollectionHandle.cpp
new file mode 100644
--- /dev/null
+++ b/Engine/Source/Runtime/Experimental/RigidPhysics/Private/RigidPhysics/RigidGeometryCollectionHandle.cpp
@@ -0,0 +1,53 @@
+// Copyright Epic Games, Inc. All Rights Reserved.
+#include "RigidPhysics/RigidGeometryCollectionHandle.h"
+
+#include "RigidPhysics/RigidBodyContainer.h"
+#include "RigidPhysics/RigidGeometryCollection.h"
+
+#if UE_RIGIDPHYSICS_API_ENABLED
+
+namespace UE::Physics
+{
+	FRigidGeometryCollectionHandle::FRigidGeometryCollectionHandle(const FRigidSceneHandle& InSceneHandle, const FRigidObjectId& InBodyContainerId)
+		: BodyContainerHandle(InSceneHandle, InBodyContainerId)
+	{
+	}
+
+	void FRigidGeometryCollectionHandle::Reset()
+	{
+		BodyContainerHandle.Reset();
+	}
+
+	const FRigidObjectId& FRigidGeometryCollectionHandle::GetId() const
+	{
+		return BodyContainerHandle.GetId();
+	}
+
+	const FRigidSceneHandle FRigidGeometryCollectionHandle::GetSceneHandle() const
+	{
+		return BodyContainerHandle.GetSceneHandle();
+	}
+
+	FRigidBodyContainerHandle FRigidGeometryCollectionHandle::AsBodyContainerHandle() const
+	{
+		return BodyContainerHandle;
+	}
+
+	FRigidGeometryCollectionHandle::operator FRigidBodyContainerHandle() const
+	{
+		return AsBodyContainerHandle();
+	}
+
+	IRigidGeometryCollection* FRigidGeometryCollectionHandle::PinInternal(IRigidBodyContainer* Container) const
+	{
+		// TODO_CHAOSAPI: Could use checked cast here because a FRigidGeometryCollectionHandle will 
+		// always reference a GC or nothing. Currently it will perform an RTTI cast.
+		if (Container != nullptr)
+		{
+			return Container->AsA<IRigidGeometryCollection>();
+		}
+		return nullptr;
+	}
+} // namespace UE::Physics
+
+#endif // UE_RIGIDPHYSICS_API_ENABLED
\ No newline at end of file
diff --git a/Engine/Source/Runtime/Experimental/RigidPhysics/Private/RigidPhysics/RigidScene.cpp b/Engine/Source/Runtime/Experimental/RigidPhysics/Private/RigidPhysics/RigidScene.cpp
--- a/Engine/Source/Runtime/Experimental/RigidPhysics/Private/RigidPhysics/RigidScene.cpp
+++ b/Engine/Source/Runtime/Experimental/RigidPhysics/Private/RigidPhysics/RigidScene.cpp
@@ -14,42 +14,6 @@ namespace UE::Physics
 	UE_RIGIDPHYSICS_RIGIDTYPED_IMPL(IRigidScene, IRigidLockable);
 	UE_RIGIDPHYSICS_RIGIDTYPED_IMPL(IRigidSceneSettings, IRigidTyped);
 
-	FRigidSceneHandle::FRigidSceneHandle()
-		: SceneRaw(nullptr)
-	{
-	}
-
-	FRigidSceneHandle::FRigidSceneHandle(FRigidScene* InScene)
-		: SceneRaw(InScene)
-#if UE_RIGIDPHYSICS_SCENEHANDLE_SAFTEY_ENABLED
-		, SceneWeak(InScene)
-#endif
-	{
-	}
-
-	void FRigidSceneHandle::Reset()
-	{
-		SceneRaw = nullptr;
-#if UE_RIGIDPHYSICS_SCENEHANDLE_SAFTEY_ENABLED
-		SceneWeak.Reset();
-#endif
-	}
-
-	IRigidScene* FRigidSceneHandle::Get() const
-	{
-		return SceneRaw;
-	}
-
-	FRigidContextGameRO FRigidSceneHandle::LockRO() const
-	{
-		return Lock<FRigidContextGameRO>();
-	}
-
-	FRigidContextGameRW FRigidSceneHandle::LockRW() const
-	{
-		return Lock<FRigidContextGameRW>();
-	}
-
 	FRigidScene::FRigidScene()
 	{
 		// TODO_CHAOSAPI: registry slack should be a parameter
diff --git a/Engine/Source/Runtime/Experimental/RigidPhysics/Private/RigidPhysics/RigidSceneHandle.cpp b/Engine/Source/Runtime/Experimental/RigidPhysics/Private/RigidPhysics/RigidSceneHandle.cpp
new file mode 100644
--- /dev/null
+++ b/Engine/Source/Runtime/Experimental/RigidPhysics/Private/RigidPhysics/RigidSceneHandle.cpp
@@ -0,0 +1,58 @@
+// Copyright Epic Games, Inc. All Rights Reserved.
+#include "RigidPhysics/RigidSceneHandle.h"
+
+#include "RigidPhysics/RigidContext.h"
+#include "RigidPhysics/RigidLog.h"
+#include "RigidPhysics/RigidScene.h"
+
+#if UE_RIGIDPHYSICS_API_ENABLED
+
+namespace UE::Physics
+{
+	FRigidSceneHandle::FRigidSceneHandle(FRigidScene* InScene)
+		: SceneRaw(InScene)
+#if UE_RIGIDPHYSICS_SCENEHANDLE_SAFTEY_ENABLED
+		, SceneWeak(InScene)
+#endif
+	{
+	}
+
+	FRigidContextGameRO FRigidSceneHandle::LockRO() const
+	{
+		DoSafetyCheck();
+		return FRigidContextGameRO(SceneRaw);
+	}
+
+	FRigidContextGameRW FRigidSceneHandle::LockRW() const
+	{
+		DoSafetyCheck();
+		return FRigidContextGameRW(SceneRaw);
+	}
+
+	void FRigidSceneHandle::Reset()
+	{
+		SceneRaw = nullptr;
+#if UE_RIGIDPHYSICS_SCENEHANDLE_SAFTEY_ENABLED
+		SceneWeak.Reset();
+#endif
+	}
+
+	IRigidScene* FRigidSceneHandle::Get() const
+	{
+		return SceneRaw;
+	}
+
+	void FRigidSceneHandle::DoSafetyCheck() const
+	{
+#if UE_RIGIDPHYSICS_SCENEHANDLE_SAFTEY_ENABLED
+		// NOTE: In order to support the UE_RIGIDPHYSICS_SCENEHANDLE_SAFTEY_ENABLED=0 optimization,
+		// SceneHandles cannot be used after the scene has expired. The onus is on the 
+		// user to ensure that this is the case, but when we have handle safety
+		// enabled we can assert on this. Without the weak pointer, there's no way 
+		// to know that the raw pointer held in SceneWrapper points to a deleted object.
+		check(SceneRaw == SceneWeak.Pin().Get());
+#endif
+	}
+} // namespace UE::Physics
+
+#endif // UE_RIGIDPHYSICS_API_ENABLED
\ No newline at end of file
diff --git a/Engine/Source/Runtime/Experimental/RigidPhysics/Public/RigidPhysics/RigidBody.h b/Engine/Source/Runtime/Experimental/RigidPhysics/Public/RigidPhysics/RigidBody.h
--- a/Engine/Source/Runtime/Experimental/RigidPhysics/Public/RigidPhysics/RigidBody.h
+++ b/Engine/Source/Runtime/Experimental/RigidPhysics/Public/RigidPhysics/RigidBody.h
@@ -8,13 +8,13 @@
 #include "Math/Bounds.h"
 #include "Math/Transform.h"
 #include "RigidPhysics/RigidBodyContainer.h"
+#include "RigidPhysics/RigidBodyHandle.h"
 #include "RigidPhysics/RigidCollider.h"
 #include "RigidPhysics/RigidLog.h"
 #include "RigidPhysics/RigidObjectId.h"
 #include "RigidPhysics/RigidObjectPtr.h"
 #include "RigidPhysics/RigidScene.h"
 #include "RigidPhysics/RigidTyped.h"
-#include "Templates/SharedPointer.h"
 
 namespace UE::Physics
 {
@@ -273,9 +273,19 @@ namespace UE::Physics
 			return BodyRaw->GetTypeId();
 		}
 
-		FRigidBodyHandle GetHandle() const;
+		FRigidBodyHandle GetHandle() const
+		{
+			if (BodyRaw != nullptr)
+			{
+				return BodyRaw->GetHandle();
+			}
+			return FRigidBodyHandle();
+		}
 
-		FRigidBodyContainerHandle GetContainerHandle() const;
+		FRigidBodyContainerHandle GetContainerHandle() const
+		{
+			return GetHandle().GetContainerHandle();
+		}
 
 		bool IsActive() const
 		{
@@ -449,112 +459,6 @@ namespace UE::Physics
 		// IRigidBody implementation depends on Context type
 		IRigidBody* BodyRaw;
 	};
-
-
-	// A handle to a rigid body in a container.
-	// This handle may be copied and stored, and can be used to retrieve a
-	// body reference by "pinning" it and providing the context.
-	// 
-	// Example contexts are: 
-	// - Read-Only Game-Thread
-	// - Read-Write Game Thread
-	// - Read-Write Physics-Thread Callback
-	//
-	class FRigidBodyHandle
-	{
-	public:
-		FRigidBodyHandle() = default;
-		FRigidBodyHandle(const FRigidBodyHandle& R) = default;
-
-		UE_INTERNAL FRigidBodyHandle(const FRigidBodyContainerHandle& InBodyContainerHandle, const FRigidObjectId& InBodyId)
-			: BodyId(InBodyId)
-			, BodyContainerHandle(InBodyContainerHandle)
-		{
-		}
-
-		const FRigidObjectId& GetId() const
-		{
-			return BodyId;
-		}
-
-		const FRigidBodyContainerHandle& GetContainerHandle() const
-		{
-			return BodyContainerHandle;
-		}
-
-		void Reset()
-		{
-			BodyContainerHandle.Reset();
-			BodyId.Reset();
-		}
-
-		// Convert the handle to a pointer in the specified context
-		// Impl in RigidContext.h
-		template<typename ContextType>
-		TRigidBodyPtr<ContextType> Pin(const ContextType& InContext) const;
-
-		friend bool operator==(const FRigidBodyHandle& L, const FRigidBodyHandle& R) = default;
-		friend bool operator!=(const FRigidBodyHandle& L, const FRigidBodyHandle& R) = default;
-
-	private:
-		FRigidObjectId BodyId;
-		FRigidBodyContainerHandle BodyContainerHandle;
-	};
-
-	template<typename ContextType>
-	FRigidBodyHandle TRigidBodyPtrImpl<ContextType>::GetHandle() const
-	{
-		if (BodyRaw != nullptr)
-		{
-			return BodyRaw->GetHandle();
-		}
-		return FRigidBodyHandle();
-	}
-
-	template<typename ContextType>
-	FRigidBodyContainerHandle TRigidBodyPtrImpl<ContextType>::GetContainerHandle() const
-	{
-		return GetHandle().GetContainerHandle();
-	}
-
-	//
-	// Implementations to avoid circular dependency in headers
-	//
-
-	template<typename ContextType>
-	TRigidBodyPtr<ContextType> TRigidBodyContainerPtrImpl<ContextType>::GetBody(const FRigidObjectId& InBodyId) const
-	{
-		return BodyContainerRaw->GetBody(InBodyId);
-	}
-
-	template<typename ContextType>
-	TRigidBodyPtr<ContextType> TRigidBodyContainerPtrImpl<ContextType>::CreateBody(const FRigidDebugName& InName, ERigidMovementType InMovementType) const
-	{
-		return BodyContainerRaw->CreateBody(InName, InMovementType);
-	}
-
-	template<typename ContextType>
-	void TRigidBodyContainerPtrImpl<ContextType>::DestroyBody(TRigidBodyPtr<FContext>&& InBodyPtr)
-	{
-		BodyContainerRaw->DestroyBody(InBodyPtr.Get());
-	}
-
-	template<typename ContextType>
-	TRigidBodyPtr<ContextType> TRigidScenePtrImpl<ContextType>::CreateBody(const FRigidDebugName& InName, ERigidMovementType InMovementType)
-	{
-		return SceneRaw->CreateBody(InName, InMovementType);
-	}
-
-	template<typename ContextType>
-	void TRigidScenePtrImpl<ContextType>::DestroyBody(TRigidBodyPtr<ContextType>&& InBody)
-	{
-		if (InBody.IsValid())
-		{
-			SceneRaw->DestroyBody(InBody.Get());
-			InBody.Reset();
-		}
-	}
-
 }
 
 #endif // UE_RIGIDPHYSICS_API_ENABLED
\ No newline at end of file
diff --git a/Engine/Source/Runtime/Experimental/RigidPhysics/Public/RigidPhysics/RigidBodyContainer.h b/Engine/Source/Runtime/Experimental/RigidPhysics/Public/RigidPhysics/RigidBodyContainer.h
--- a/Engine/Source/Runtime/Experimental/RigidPhysics/Public/RigidPhysics/RigidBodyContainer.h
+++ b/Engine/Source/Runtime/Experimental/RigidPhysics/Public/RigidPhysics/RigidBodyContainer.h
@@ -5,6 +5,7 @@
 
 #if UE_RIGIDPHYSICS_API_ENABLED
 
+#include "RigidPhysics/RigidBodyContainerHandle.h"
 #include "RigidPhysics/RigidLog.h"
 #include "RigidPhysics/RigidObjectId.h"
 #include "RigidPhysics/RigidObjectPtr.h"
@@ -115,7 +116,10 @@ namespace UE::Physics
 			return BodyContainerRaw->GetId();
 		}
 
-		FRigidBodyContainerHandle GetHandle() const;
+		FRigidBodyContainerHandle GetHandle() const
+		{
+			return BodyContainerRaw->GetHandle();
+		}
 
 		bool IsInScene() const
 		{
@@ -127,72 +131,30 @@ namespace UE::Physics
 			return BodyContainerRaw->GetNumBodies();
 		}
 
-		// Impl in RigidBody.h
-		TRigidBodyPtr<FContext> GetBody(const FRigidObjectId& InBodyId) const;
-		TRigidBodyPtr<FContext> CreateBody(const FRigidDebugName& InName, ERigidMovementType InMovementType) const;
-		void DestroyBody(TRigidBodyPtr<FContext>&& InBodyPtr);
-
-		UE_INTERNAL IRigidBodyContainer* Get() const
-		{
-			return BodyContainerRaw;
-		}
-
-	private:
-		// Pointer to implementation, object depends on Context type
-		IRigidBodyContainer* BodyContainerRaw;
-	};
-
-
-	class FRigidBodyContainerHandle
-	{
-	public:
-		FRigidBodyContainerHandle()
-		{
-		}
-
-		UE_INTERNAL FRigidBodyContainerHandle(const FRigidSceneHandle& InSceneHandle, const FRigidObjectId& InBodyContainerId)
-			: BodyContainerId(InBodyContainerId)
-			, SceneHandle(InSceneHandle)
+		TRigidBodyPtr<FContext> GetBody(const FRigidObjectId& InBodyId) const
 		{
+			return BodyContainerRaw->GetBody(InBodyId);
 		}
 
-		void Reset()
+		TRigidBodyPtr<FContext> CreateBody(const FRigidDebugName& InName, ERigidMovementType InMovementType) const
 		{
-			BodyContainerId.Reset();
-			SceneHandle.Reset();
+			return BodyContainerRaw->CreateBody(InName, InMovementType);
 		}
 
-		const FRigidObjectId& GetId() const
+		void DestroyBody(TRigidBodyPtr<FContext>&& InBodyPtr)
 		{
-			return BodyContainerId;
+			BodyContainerRaw->DestroyBody(InBodyPtr.Get());
 		}
 
-		// NOTE: We don't really need the scene handle because it is provided by the Context
-		// in Pin(). It is only needed for the debugger visualizers (see RigidPhysics.natvis).
-		// TODO_CHAOSAPI: Consider making SceneHandle shipping only
-		const FRigidSceneHandle GetSceneHandle() const
+		UE_INTERNAL IRigidBodyContainer* Get() const
 		{
-			return SceneHandle;
+			return BodyContainerRaw;
 		}
 
-		// Convert the handle to a pointer in the specified context
-		// Impl in RigidContext.h
-		template<typename ContextType>
-		TRigidBodyContainerPtr<ContextType> Pin(const ContextType& InContext) const;
-
-		friend bool operator==(const FRigidBodyContainerHandle& L, const FRigidBodyContainerHandle& R) = default;
-		friend bool operator!=(const FRigidBodyContainerHandle& L, const FRigidBodyContainerHandle& R) = default;
-
 	private:
-		FRigidObjectId BodyContainerId;
-		FRigidSceneHandle SceneHandle;
+		// Pointer to implementation, object depends on Context type
+		IRigidBodyContainer* BodyContainerRaw;
 	};
-
-	template<typename ContextType>
-	FRigidBodyContainerHandle TRigidBodyContainerPtrImpl<ContextType>::GetHandle() const
-	{
-		return BodyContainerRaw->GetHandle();
-	}
 }
 
 #endif // UE_RIGIDPHYSICS_API_ENABLED
\ No newline at end of file
diff --git a/Engine/Source/Runtime/Experimental/RigidPhysics/Public/RigidPhysics/RigidBodyContainerHandle.h b/Engine/Source/Runtime/Experimental/RigidPhysics/Public/RigidPhysics/RigidBodyContainerHandle.h
new file mode 100644
--- /dev/null
+++ b/Engine/Source/Runtime/Experimental/RigidPhysics/Public/RigidPhysics/RigidBodyContainerHandle.h
@@ -0,0 +1,56 @@
+// Copyright Epic Games, Inc. All Rights Reserved.
+#pragma once
+
+#include "RigidPhysics/RigidFwd.h"
+
+#if UE_RIGIDPHYSICS_API_ENABLED
+
+#include "RigidPhysics/RigidObjectId.h"
+#include "RigidPhysics/RigidObjectPtr.h"
+#include "RigidPhysics/RigidSceneHandle.h"
+
+namespace UE::Physics
+{
+	template <typename ContextType>
+	class UE_INTERNAL TRigidBodyContainerPtrImpl;
+
+	class FRigidBodyContainerHandle
+	{
+	public:
+		FRigidBodyContainerHandle() = default;
+		UE_INTERNAL FRigidBodyContainerHandle(const FRigidSceneHandle& InSceneHandle, const FRigidObjectId& InBodyContainerId);
+
+		friend bool operator==(const FRigidBodyContainerHandle&, const FRigidBodyContainerHandle&) = default;
+		friend bool operator!=(const FRigidBodyContainerHandle&, const FRigidBodyContainerHandle&) = default;
+
+		// Convert the handle to a pointer in the specified context.
+		template <typename ContextType>
+		TRigidBodyContainerPtr<ContextType> Pin(const ContextType& InContext) const
+		{
+			// NOTE: we don't use our copy of SceneHandle here (or anywhere!)
+			// TODO @ JoshD: This seems wrong. We need some safety checking here...
+			if (InContext.IsValid())
+			{
+				const IRigidScene* Scene = InContext.GetScene().Get();
+				return PinInternal(Scene);
+			}
+			return nullptr;
+		}
+
+		UE_INTERNAL void Reset();
+		UE_INTERNAL const FRigidObjectId& GetId() const;
+
+		// NOTE: We don't really need the scene handle because it is provided by the Context
+		// in Pin(). It is only needed for the debugger visualizers (see RigidPhysics.natvis).
+		// TODO_CHAOSAPI: Consider making SceneHandle shipping only
+		UE_INTERNAL const FRigidSceneHandle GetSceneHandle() const;
+
+	private:
+		IRigidBodyContainer* PinInternal(const IRigidScene* InScene) const;
+
+		FRigidObjectId BodyContainerId;
+		FRigidSceneHandle SceneHandle;
+	};
+} // namespace UE::Physics
+
+#endif // UE_RIGIDPHYSICS_API_ENABLED
\ No newline at end of file
diff --git a/Engine/Source/Runtime/Experimental/RigidPhysics/Public/RigidPhysics/RigidBodyHandle.h b/Engine/Source/Runtime/Experimental/RigidPhysics/Public/RigidPhysics/RigidBodyHandle.h
new file mode 100644
--- /dev/null
+++ b/Engine/Source/Runtime/Experimental/RigidPhysics/Public/RigidPhysics/RigidBodyHandle.h
@@ -0,0 +1,64 @@
+// Copyright Epic Games, Inc. All Rights Reserved.
+#pragma once
+
+#include "RigidPhysics/RigidFwd.h"
+
+#if UE_RIGIDPHYSICS_API_ENABLED
+
+#include "RigidPhysics/RigidBodyContainerHandle.h"
+#include "RigidPhysics/RigidObjectId.h"
+#include "RigidPhysics/RigidObjectPtr.h"
+
+namespace UE::Physics
+{
+	class IRigidBody;
+	template <typename ContextType>
+	class UE_INTERNAL TRigidBodyPtrImpl;
+
+	// A handle to a rigid body in a container.
+	// This handle may be copied and stored, and can be used to retrieve a
+	// body reference by "pinning" it and providing the context.
+	// 
+	// Example contexts are: 
+	// - Read-Only Game-Thread
+	// - Read-Write Game Thread
+	// - Read-Write Physics-Thread Callback
+	//
+	class FRigidBodyHandle
+	{
+	public:
+		FRigidBodyHandle() = default;
+		FRigidBodyHandle(const FRigidBodyHandle&) = default;
+
+		UE_INTERNAL FRigidBodyHandle(const FRigidBodyContainerHandle& InBodyContainerHandle, const FRigidObjectId& InBodyId);
+
+		friend bool operator==(const FRigidBodyHandle&, const FRigidBodyHandle&) = default;
+		friend bool operator!=(const FRigidBodyHandle&, const FRigidBodyHandle&) = default;
+
+		// Convert the handle to a pointer in the specified context.
+		template <typename ContextType>
+		TRigidBodyPtr<ContextType> Pin(const ContextType& InContext) const
+		{
+			if (InContext.IsValid())
+			{
+				if (TRigidBodyContainerPtr<ContextType> BodyContainer = BodyContainerHandle.Pin(InContext))
+				{
+					return PinInternal(BodyContainer.Get());
+				}
+			}
+			return TRigidBodyPtr<ContextType>();
+		}
+
+		UE_INTERNAL const FRigidObjectId& GetId() const;
+		UE_INTERNAL const FRigidBodyContainerHandle& GetContainerHandle() const;
+		UE_INTERNAL void Reset();
+
+	private:
+		IRigidBody* PinInternal(const IRigidBodyContainer* Container) const;
+
+		FRigidObjectId BodyId;
+		FRigidBodyContainerHandle BodyContainerHandle;
+	};
+} // namespace UE::Physics
+
+#endif // UE_RIGIDPHYSICS_API_ENABLED
\ No newline at end of file
diff --git a/Engine/Source/Runtime/Experimental/RigidPhysics/Public/RigidPhysics/RigidContext.h b/Engine/Source/Runtime/Experimental/RigidPhysics/Public/RigidPhysics/RigidContext.h
--- a/Engine/Source/Runtime/Experimental/RigidPhysics/Public/RigidPhysics/RigidContext.h
+++ b/Engine/Source/Runtime/Experimental/RigidPhysics/Public/RigidPhysics/RigidContext.h
@@ -7,9 +7,6 @@
 
 #include "RigidPhysics/RigidLockable.h"
 #include "RigidPhysics/RigidLog.h"
-#include "RigidPhysics/RigidBody.h"
-#include "RigidPhysics/RigidBodyContainer.h"
-#include "RigidPhysics/RigidGeometryCollection.h"
 #include "RigidPhysics/RigidScene.h"
 
 namespace UE::Physics
@@ -72,61 +69,6 @@ namespace UE::Physics
 	private:
 		TRigidScenePtr<FContext> ScenePtr;
 	};
-
-
-	//
-	// Functions below are here to prevent circular dependencies
-	//
-
-	template<typename ContextType>
-	inline ContextType FRigidSceneHandle::Lock() const
-	{
-#if UE_RIGIDPHYSICS_SCENEHANDLE_SAFTEY_ENABLED
-		// NOTE: In order to support the UE_RIGIDPHYSICS_SCENEHANDLE_SAFTEY_ENABLED=0 optimization
-		// SceneHandles cannot be used after the scene has expired. The onus is on the 
-		// user to ensure that this is the case, but when we have handle safety
-		// enabled we can assert on this. Without the weak pointer, there's no way 
-		// to know that the raw pointer held in SceneWrapper points to a deleted object.
-		check(SceneRaw == SceneWeak.Pin().Get());
-#endif
-		return ContextType(SceneRaw);
-	}
-
-	template<typename ContextType>
-	inline TRigidBodyContainerPtr<ContextType> FRigidBodyContainerHandle::Pin(const ContextType& InContext) const
-	{
-		// NOTE: we don't use our copy of SceneHandle here (or anywhere!)
-		if (InContext.IsValid())
-		{
-			return InContext.GetScene()->GetBodyContainer(GetId());
-		}
-		return nullptr;
-	}
-
-	template<typename ContextType>
-	inline TRigidGeometryCollectionPtr<ContextType> FRigidGeometryCollectionHandle::Pin(const ContextType& InContext) const
-	{
-		if (InContext.IsValid())
-		{
-			// TODO_CHAOSAPI: Could use checked cast here because a FRigidGeometryCollectionHandle will 
-			// always reference a GC or nothing. Currently it will perform an RTTI cast.
-			return BodyContainerHandle.Pin(InContext);
-		}
-		return nullptr;
-	}
-
-	template<typename ContextType>
-	inline TRigidBodyPtr<ContextType> FRigidBodyHandle::Pin(const ContextType& InContext) const
-	{
-		if (InContext.IsValid())
-		{
-			if (TRigidBodyContainerPtr<ContextType> BodyContainer = BodyContainerHandle.Pin(InContext))
-			{
-				return BodyContainer->GetBody(GetId());
-			}
-		}
-		return TRigidBodyPtr<ContextType>();
-	}
 }
 
 #endif // UE_RIGIDPHYSICS_API_ENABLED
\ No newline at end of file
diff --git a/Engine/Source/Runtime/Experimental/RigidPhysics/Public/RigidPhysics/RigidGeometryCollection.h b/Engine/Source/Runtime/Experimental/RigidPhysics/Public/RigidPhysics/RigidGeometryCollection.h
--- a/Engine/Source/Runtime/Experimental/RigidPhysics/Public/RigidPhysics/RigidGeometryCollection.h
+++ b/Engine/Source/Runtime/Experimental/RigidPhysics/Public/RigidPhysics/RigidGeometryCollection.h
@@ -7,7 +7,7 @@
 
 #include "RigidPhysics/RigidBody.h"
 #include "RigidPhysics/RigidBodyContainer.h"
-
+#include "RigidPhysics/RigidGeometryCollectionHandle.h"
 
 namespace UE::Physics
 {
@@ -79,7 +79,15 @@ namespace UE::Physics
 			return GeometryCollectionRaw->GetId();
 		}
 
-		FRigidGeometryCollectionHandle GetHandle() const;
+		FRigidGeometryCollectionHandle GetHandle() const
+		{
+			if (GeometryCollectionRaw != nullptr)
+			{
+				FRigidBodyContainerHandle BodyContainerHandle = GeometryCollectionRaw->GetHandle();
+				return FRigidGeometryCollectionHandle(BodyContainerHandle.GetSceneHandle(), BodyContainerHandle.GetId());
+			}
+			return FRigidGeometryCollectionHandle();
+		}
 
 		bool IsInScene() const
 		{
@@ -91,122 +99,25 @@ namespace UE::Physics
 			return GeometryCollectionRaw->GetNumBodies();
 		}
 
-		// Impl in RigidBody.h
-		TRigidBodyPtr<FContext> GetBodyAt(int32 InBodyIndex) const;
-		TRigidBodyPtr<FContext> GetBody(const FRigidObjectId& InBodyId) const;
-
-		UE_INTERNAL IRigidGeometryCollection* Get() const
-		{
-			return GeometryCollectionRaw;
-		}
-
-	private:
-		// Pointer to implementation, object depends on Context type
-		IRigidGeometryCollection* GeometryCollectionRaw;
-	};
-
-	// A handle to a GeometryCollection. A GeometryCollection is a type of BodyContainer and the handle
-	// is auto-castable to FRigidBodyContainerHandle, but the reverse is more complex.
-	//
-	// NOTE: If the GeometryCollectionHandle has a valid BodyContainerHandle, it will always reference a GC. 
-	// There should be no way to get a GCHandle that contains a non-GC body container handle. This is enforced
-	// because the only way to create a FRigidGeometryCollectionHandle is from a TRigidGeometryCollectionPtr.
-	// 
-	// To safely convert a FRigidBodyContainerHandle to FRigidGeometryCollectionHandle we must go through
-	// pointer conversion:
-	// 
-	//		FRigidBodyContainerHandle BodyContainerHandle = ...;
-	//		TRigidGeomatryCollectionPtr<FRigidContextGameRW> GCPtr = BodyContainerHandle.Pin(Context);
-	//		FRigidGeometryCollectionHandle GCHandle = GCPtr;
-	// 
-	// This ensures that we get an empty handle if BodyContainerHandle is not a GC, because the cast from
-	// TRigidBodyContainerPtr to TRigidGeometryCollectionPtr will fail and return nullptr.
-	//
-	class FRigidGeometryCollectionHandle
-	{
-	public:
-		FRigidGeometryCollectionHandle()
-		{
-		}
-
-		UE_INTERNAL FRigidGeometryCollectionHandle(const FRigidSceneHandle& InSceneHandle, const FRigidObjectId& InBodyContainerId)
-			: BodyContainerHandle(InSceneHandle, InBodyContainerId)
-		{
-		}
-
-		void Reset()
-		{
-			BodyContainerHandle.Reset();
-		}
-
-		const FRigidObjectId& GetId() const
+		TRigidBodyPtr<FContext> GetBodyAt(int32 InBodyIndex) const
 		{
-			return BodyContainerHandle.GetId();
+			return GeometryCollectionRaw->GetBodyAt(InBodyIndex);
 		}
 
-		const FRigidSceneHandle GetSceneHandle() const
+		TRigidBodyPtr<FContext> GetBody(const FRigidObjectId& InBodyId) const
 		{
-			return BodyContainerHandle.GetSceneHandle();
+			return GeometryCollectionRaw->GetBody(InBodyId);
 		}
 
-		FRigidBodyContainerHandle AsBodyContainerHandle() const
-		{
-			return BodyContainerHandle;
-		}
-
-		operator FRigidBodyContainerHandle() const
+		UE_INTERNAL IRigidGeometryCollection* Get() const
 		{
-			return AsBodyContainerHandle();
+			return GeometryCollectionRaw;
 		}
 
-		// Convert the handle to a pointer in the specified context
-		// Impl in RigidContext.h
-		template<typename ContextType>
-		TRigidGeometryCollectionPtr<ContextType> Pin(const ContextType& InContext) const;
-
-		friend bool operator==(const FRigidBodyContainerHandle& L, const FRigidBodyContainerHandle& R) = default;
-		friend bool operator!=(const FRigidBodyContainerHandle& L, const FRigidBodyContainerHandle& R) = default;
-
 	private:
-		FRigidBodyContainerHandle BodyContainerHandle;
+		// Pointer to implementation, object depends on Context type
+		IRigidGeometryCollection* GeometryCollectionRaw;
 	};
-
-
-	template<typename ContextType>
-	FRigidGeometryCollectionHandle TRigidGeometryCollectionPtrImpl<ContextType>::GetHandle() const
-	{
-		if (GeometryCollectionRaw != nullptr)
-		{
-			FRigidBodyContainerHandle BodyContainerHandle = GeometryCollectionRaw->GetHandle();
-			return FRigidGeometryCollectionHandle(BodyContainerHandle.GetSceneHandle(), BodyContainerHandle.GetId());
-		}
-		return FRigidGeometryCollectionHandle();
-	}
-
-	template<typename ContextType>
-	TRigidBodyPtr<ContextType> TRigidGeometryCollectionPtrImpl<ContextType>::GetBodyAt(int32 InBodyIndex) const
-	{
-		return GeometryCollectionRaw->GetBodyAt(InBodyIndex);
-	}
-
-	template<typename ContextType>
-	TRigidBodyPtr<ContextType> TRigidGeometryCollectionPtrImpl<ContextType>::GetBody(const FRigidObjectId& InBodyId) const
-	{
-		return GeometryCollectionRaw->GetBody(InBodyId);
-	}
-
-	template<typename ContextType>
-	TRigidGeometryCollectionPtr<ContextType> TRigidScenePtrImpl<ContextType>::CreateGeometryCollection(const FRigidDebugName& InName, const FSimulationParameters& InParameters)
-	{
-		return SceneRaw->CreateGeometryCollection(InName, InParameters);
-	}
-
-	template<typename ContextType>
-	void TRigidScenePtrImpl<ContextType>::DestroyGeometryCollection(TRigidGeometryCollectionPtr<ContextType>&& InGC)
-	{
-		SceneRaw->DestroyGeometryCollection(InGC.Get());
-		InGC.Reset();
-	}
 }
 
 #endif // UE_RIGIDPHYSICS_API_ENABLED
\ No newline at end of file
diff --git a/Engine/Source/Runtime/Experimental/RigidPhysics/Public/RigidPhysics/RigidGeometryCollectionHandle.h b/Engine/Source/Runtime/Experimental/RigidPhysics/Public/RigidPhysics/RigidGeometryCollectionHandle.h
new file mode 100644
--- /dev/null
+++ b/Engine/Source/Runtime/Experimental/RigidPhysics/Public/RigidPhysics/RigidGeometryCollectionHandle.h
@@ -0,0 +1,69 @@
+// Copyright Epic Games, Inc. All Rights Reserved.
+#pragma once
+
+#include "RigidPhysics/RigidFwd.h"
+
+#if UE_RIGIDPHYSICS_API_ENABLED
+
+#include "RigidPhysics/RigidBodyContainerHandle.h"
+#include "RigidPhysics/RigidObjectPtr.h"
+
+namespace UE::Physics
+{
+	class UE_INTERNAL IRigidGeometryCollection;
+	template <typename ContextType>
+	class UE_INTERNAL TRigidGeometryCollectionPtrImpl;
+
+	// A handle to a GeometryCollection. A GeometryCollection is a type of BodyContainer and the handle
+	// is auto-castable to FRigidBodyContainerHandle, but the reverse is more complex.
+	//
+	// NOTE: If the GeometryCollectionHandle has a valid BodyContainerHandle, it will always reference a GC. 
+	// There should be no way to get a GCHandle that contains a non-GC body container handle. This is enforced
+	// because the only way to create a FRigidGeometryCollectionHandle is from a TRigidGeometryCollectionPtr.
+	// 
+	// To safely convert a FRigidBodyContainerHandle to FRigidGeometryCollectionHandle we must go through
+	// pointer conversion:
+	// 
+	//		FRigidBodyContainerHandle BodyContainerHandle = ...;
+	//		TRigidGeometryCollectionPtrImpl<FRigidContextGameRW> GCPtr = BodyContainerHandle.Pin(Context);
+	//		FRigidGeometryCollectionHandle GCHandle = GCPtr;
+	// 
+	// This ensures that we get an empty handle if BodyContainerHandle is not a GC, because the cast from
+	// TRigidBodyContainerPtr to TRigidGeometryCollectionPtr will fail and return nullptr.
+	//
+	class FRigidGeometryCollectionHandle
+	{
+	public:
+		FRigidGeometryCollectionHandle() = default;
+		UE_INTERNAL FRigidGeometryCollectionHandle(const FRigidSceneHandle& InSceneHandle, const FRigidObjectId& InBodyContainerId);
+
+		friend bool operator==(const FRigidBodyContainerHandle&, const FRigidBodyContainerHandle&) = default;
+		friend bool operator!=(const FRigidBodyContainerHandle&, const FRigidBodyContainerHandle&) = default;
+
+		// Convert the handle to a pointer in the specified context
+		template<typename ContextType>
+		TRigidGeometryCollectionPtr<ContextType> Pin(const ContextType& InContext) const
+		{
+			if (InContext.IsValid())
+			{
+				IRigidBodyContainer* Container = BodyContainerHandle.Pin(InContext).Get();
+				return PinInternal(Container);
+			}
+			return nullptr;
+		}
+
+		UE_INTERNAL void Reset();
+		UE_INTERNAL const FRigidObjectId& GetId() const;
+		UE_INTERNAL const FRigidSceneHandle GetSceneHandle() const;
+
+		UE_INTERNAL FRigidBodyContainerHandle AsBodyContainerHandle() const;
+		UE_INTERNAL operator FRigidBodyContainerHandle() const;
+
+	private:
+		IRigidGeometryCollection* PinInternal(IRigidBodyContainer* Container) const;
+
+		FRigidBodyContainerHandle BodyContainerHandle;
+	};
+} // namespace UE::Physics
+
+#endif // UE_RIGIDPHYSICS_API_ENABLED
\ No newline at end of file
diff --git a/Engine/Source/Runtime/Experimental/RigidPhysics/Public/RigidPhysics/RigidScene.h b/Engine/Source/Runtime/Experimental/RigidPhysics/Public/RigidPhysics/RigidScene.h
--- a/Engine/Source/Runtime/Experimental/RigidPhysics/Public/RigidPhysics/RigidScene.h
+++ b/Engine/Source/Runtime/Experimental/RigidPhysics/Public/RigidPhysics/RigidScene.h
@@ -9,6 +9,7 @@
 #include "RigidPhysics/RigidLog.h"
 #include "RigidPhysics/RigidObjectId.h"
 #include "RigidPhysics/RigidObjectPtr.h"
+#include "RigidPhysics/RigidSceneHandle.h"
 #include "RigidPhysics/RigidTyped.h"
 #include "Templates/SharedPointer.h"
 
@@ -142,7 +143,10 @@ namespace UE::Physics
 			return (SceneRaw != nullptr);
 		}
 
-		FRigidSceneHandle GetHandle() const;
+		FRigidSceneHandle GetHandle() const
+		{
+			return SceneRaw->GetHandle();
+		}
 
 		IRigidBodyContainer* GetBodyContainer(const FRigidObjectId& InBodyContainerId) const
 		{
@@ -164,13 +168,30 @@ namespace UE::Physics
 			SceneRaw->UnregisterModifier(Modifier);
 		}
 
-		// Impl in RigidBody.h
-		TRigidBodyPtr<FContext> CreateBody(const FRigidDebugName& InName, ERigidMovementType InMovementType);
-		void DestroyBody(TRigidBodyPtr<FContext>&& InBody);
+		TRigidBodyPtr<FContext> CreateBody(const FRigidDebugName& InName, ERigidMovementType InMovementType)
+		{
+			return SceneRaw->CreateBody(InName, InMovementType);
+		}
 
-		// Impl in RigidGeometryCollection.h
-		TRigidGeometryCollectionPtr<FContext> CreateGeometryCollection(const FRigidDebugName& InName, const FSimulationParameters& InParameters);
-		void DestroyGeometryCollection(TRigidGeometryCollectionPtr<ContextType>&& InGC);
+		void DestroyBody(TRigidBodyPtr<FContext>&& InBody)
+		{
+			if (InBody.IsValid())
+			{
+				SceneRaw->DestroyBody(InBody.Get());
+				InBody.Reset();
+			}
+		}
+
+		TRigidGeometryCollectionPtr<FContext> CreateGeometryCollection(const FRigidDebugName& InName, const FSimulationParameters& InParameters)
+		{
+			return SceneRaw->CreateGeometryCollection(InName, InParameters);
+		}
+
+		void DestroyGeometryCollection(TRigidGeometryCollectionPtr<ContextType>&& InGC)
+		{
+			SceneRaw->DestroyGeometryCollection(InGC.Get());
+			InGC.Reset();
+		}
 
 		bool Query(const FRigidQueryDesc& QueryDesc, TArray<FRigidQueryHit>& OutHits)  const
 		{
@@ -190,44 +211,6 @@ namespace UE::Physics
 	private:
 		IRigidScene* SceneRaw;
 	};
-
-
-	// A handle to a rigid body scene. FRigidSceneHandle is only used in the GameContext
-	// because all PhysicsContect use cases (callbacks, modifiers, etc.) will be passed their
-	// scene as a TRigidScenePtr.
-	class FRigidSceneHandle
-	{
-	public:
-		RIGIDPHYSICS_API FRigidSceneHandle();
-		UE_INTERNAL RIGIDPHYSICS_API FRigidSceneHandle(FRigidScene* InScene);
-
-		// Lock the scene on the game thread and return a locked context for use in pinning object handles
-		RIGIDPHYSICS_API FRigidContextGameRO LockRO() const;
-		RIGIDPHYSICS_API FRigidContextGameRW LockRW() const;
-
-		UE_INTERNAL RIGIDPHYSICS_API void Reset();
-		UE_INTERNAL RIGIDPHYSICS_API IRigidScene* Get() const;
-
-		friend bool operator==(const FRigidSceneHandle& L, const FRigidSceneHandle& R) = default;
-		friend bool operator!=(const FRigidSceneHandle& L, const FRigidSceneHandle& R) = default;
-
-	private:
-		// Impl in RigidContext.h
-		template<typename ContextType>
-		ContextType Lock() const;
-
-		// NOTE: Always the GT Scene, but never accessed on PT
-		FRigidScene* SceneRaw;
-#if UE_RIGIDPHYSICS_SCENEHANDLE_SAFTEY_ENABLED
-		FRigidSceneWeakPtr SceneWeak;
-#endif
-	};
-
-	template<typename ContextType>
-	FRigidSceneHandle TRigidScenePtrImpl<ContextType>::GetHandle() const
-	{
-		return SceneRaw->GetHandle();
-	}
 }
 
 #endif // UE_RIGIDPHYSICS_API_ENABLED
\ No newline at end of file
diff --git a/Engine/Source/Runtime/Experimental/RigidPhysics/Public/RigidPhysics/RigidSceneHandle.h b/Engine/Source/Runtime/Experimental/RigidPhysics/Public/RigidPhysics/RigidSceneHandle.h
new file mode 100644
--- /dev/null
+++ b/Engine/Source/Runtime/Experimental/RigidPhysics/Public/RigidPhysics/RigidSceneHandle.h
@@ -0,0 +1,42 @@
+// Copyright Epic Games, Inc. All Rights Reserved.
+#pragma once
+
+#include "RigidPhysics/RigidFwd.h"
+
+#if UE_RIGIDPHYSICS_API_ENABLED
+
+namespace UE::Physics
+{
+	class FRigidScene;
+
+	// A handle to a rigid body scene. FRigidSceneHandle is only used in the GameContext
+	// because all PhysicsContext use cases (callbacks, modifiers, etc.) will be passed their
+	// scene as a TRigidScenePtr.
+	class FRigidSceneHandle
+	{
+	public:
+		FRigidSceneHandle() = default;
+		UE_INTERNAL RIGIDPHYSICS_API FRigidSceneHandle(FRigidScene* InScene);
+
+		friend bool operator==(const FRigidSceneHandle&, const FRigidSceneHandle&) = default;
+		friend bool operator!=(const FRigidSceneHandle&, const FRigidSceneHandle&) = default;
+
+		// Lock the scene on the game thread and return a locked context for use in pinning object handles
+		RIGIDPHYSICS_API FRigidContextGameRO LockRO() const;
+		RIGIDPHYSICS_API FRigidContextGameRW LockRW() const;
+
+		UE_INTERNAL RIGIDPHYSICS_API void Reset();
+		UE_INTERNAL RIGIDPHYSICS_API IRigidScene* Get() const;
+
+	private:
+		void DoSafetyCheck() const;
+
+		// NOTE: Always the GT Scene, but never accessed on PT
+		FRigidScene* SceneRaw = nullptr;
+#if UE_RIGIDPHYSICS_SCENEHANDLE_SAFTEY_ENABLED
+		FRigidSceneWeakPtr SceneWeak;
+#endif
+	};
+} // namespace UE::Physics
+
+#endif // UE_RIGIDPHYSICS_API_ENABLED
