Geometry collection : wake up sleeping particles when apply force or impulse at position

#Synced-CL 50641382

diff --git a/Engine/Source/Runtime/Experimental/Chaos/Private/PhysicsProxy/GeometryCollectionPhysicsProxy.cpp b/Engine/Source/Runtime/Experimental/Chaos/Private/PhysicsProxy/GeometryCollectionPhysicsProxy.cpp
--- a/Engine/Source/Runtime/Experimental/Chaos/Private/PhysicsProxy/GeometryCollectionPhysicsProxy.cpp
+++ b/Engine/Source/Runtime/Experimental/Chaos/Private/PhysicsProxy/GeometryCollectionPhysicsProxy.cpp
@@ -2543,6 +2543,11 @@ void FGeometryCollectionPhysicsProxy::ApplyForceAt_External(FVector Force, FVect
 				}
 				if (ClosestHandle)
 				{
+					if (ClosestHandle->IsSleeping())
+					{
+						RBDSolver->GetEvolution()->WakeParticle(ClosestHandle);
+					}
+
 					const Chaos::FVec3 WorldCOM = Chaos::FParticleUtilitiesXR::GetCoMWorldPosition(ClosestHandle);
 					const Chaos::FVec3 WorldTorque = Chaos::FVec3::CrossProduct(WorldLocation - WorldCOM, Force);
 					ClosestHandle->AddForce(Force);
@@ -2581,6 +2586,11 @@ void FGeometryCollectionPhysicsProxy::ApplyImpulseAt_External(FVector Impulse, F
 				}
 				if (ClosestHandle)
 				{
+					if (ClosestHandle->IsSleeping())
+					{
+						RBDSolver->GetEvolution()->WakeParticle(ClosestHandle);
+					}
+
 					const Chaos::FVec3 WorldCOM = Chaos::FParticleUtilitiesXR::GetCoMWorldPosition(ClosestHandle);
 					ClosestHandle->SetLinearImpulseVelocity(ClosestHandle->LinearImpulseVelocity() + (Impulse * ClosestHandle->InvM()), false);
 
