Remove the possibility to get the package store backend from the public API.

- retrieving the backend was only ever used in the on-demand I/O store to determine whether the package store entry was on-demand or not
- added FPackageStoreInternal to extend the package store interface with internal functionality
- added internal function IterateBackends

Note: The reverted API changes never went out in a public engine release.

#rnx

#Synced-CL 50645574

diff --git a/Engine/Source/Runtime/CoreUObject/Internal/Serialization/PackageStoreInternal.h b/Engine/Source/Runtime/CoreUObject/Internal/Serialization/PackageStoreInternal.h
new file mode 100644
--- /dev/null
+++ b/Engine/Source/Runtime/CoreUObject/Internal/Serialization/PackageStoreInternal.h
@@ -0,0 +1,19 @@
+// Copyright Epic Games, Inc. All Rights Reserved.
+
+#pragma once
+
+#include "Serialization/PackageStore.h"
+
+class FPackageStoreInternal
+{
+public:
+	static void IterateBackends(FPackageStore& PackageStore, FVisitPackageStoreBackend&& Visit)
+	{
+		PackageStore.IterateBackends(MoveTemp(Visit));
+	}
+
+	static void IterateBackends(FVisitPackageStoreBackend&& Visit)
+	{
+		IterateBackends(FPackageStore::Get(), MoveTemp(Visit));
+	}
+};
diff --git a/Engine/Source/Runtime/CoreUObject/Private/Serialization/PackageStore.cpp b/Engine/Source/Runtime/CoreUObject/Private/Serialization/PackageStore.cpp
--- a/Engine/Source/Runtime/CoreUObject/Private/Serialization/PackageStore.cpp
+++ b/Engine/Source/Runtime/CoreUObject/Private/Serialization/PackageStore.cpp
@@ -562,24 +562,16 @@ void FPackageStore::Mount(TSharedRef<IPackageStoreBackend> Backend, int32 Priori
 	Backend->OnMounted(BackendContext);
 }
 
-EPackageStoreEntryStatus FPackageStore::GetPackageStoreEntry(FPackageId PackageId, FName PackageName,
-	FPackageStoreEntry& OutPackageStoreEntry, TSharedPtr<IPackageStoreBackend>* FoundBackend /*= nullptr*/)
+EPackageStoreEntryStatus FPackageStore::GetPackageStoreEntry(FPackageId PackageId, FName PackageName, FPackageStoreEntry& OutPackageStoreEntry)
 {
 	check(ThreadReadCount);
 	for (const FBackendAndPriority& Backend : Backends)
 	{
 		EPackageStoreEntryStatus Status = Backend.Value->GetPackageStoreEntry(PackageId, PackageName, OutPackageStoreEntry);
-		const bool bContinueSearch = 
-			Status == EPackageStoreEntryStatus::None || 
-			Status == EPackageStoreEntryStatus::Missing;
+		const bool bContinueSearch = Status == EPackageStoreEntryStatus::None || Status == EPackageStoreEntryStatus::Missing;
 
 		if (!bContinueSearch)
 		{
-			if (FoundBackend)
-			{
-				*FoundBackend = Backend.Value;
-			}
-
 			checkSlow((Backend.Value->GetSupportedLoaders() & OutPackageStoreEntry.LoaderType));
 			return Status;
 		}
@@ -623,3 +615,17 @@ bool FPackageStore::HasAnyBackendsMounted() const
 {
 	return !Backends.IsEmpty();
 }
+
+void FPackageStore::IterateBackends(FVisitPackageStoreBackend&& Visit)
+{
+	check(ThreadReadCount);
+	check(Visit != nullptr);
+
+	for (const FBackendAndPriority& Backend : Backends)
+	{
+		if (Visit(*Backend.Value) == false)
+		{
+			break;
+		}
+	}
+}
diff --git a/Engine/Source/Runtime/CoreUObject/Public/Serialization/PackageStore.h b/Engine/Source/Runtime/CoreUObject/Public/Serialization/PackageStore.h
--- a/Engine/Source/Runtime/CoreUObject/Public/Serialization/PackageStore.h
+++ b/Engine/Source/Runtime/CoreUObject/Public/Serialization/PackageStore.h
@@ -222,6 +222,8 @@ public:
 	}
 };
 
+using FVisitPackageStoreBackend = TFunction<bool(IPackageStoreBackend&)>;
+
 /**
  * Stores information about available packages that can be loaded.
  */
@@ -234,8 +236,7 @@ public:
 	COREUOBJECT_API void Mount(TSharedRef<IPackageStoreBackend> Backend, int32 Priority = 0);
 
 	/* Returns the package store entry data with export info and imported packages for the specified package ID. */
-	COREUOBJECT_API EPackageStoreEntryStatus GetPackageStoreEntry(FPackageId PackageId, FName PackageName, 
-		FPackageStoreEntry& OutPackageStoreEntry, TSharedPtr<IPackageStoreBackend>* FoundBackend = nullptr);
+	COREUOBJECT_API EPackageStoreEntryStatus GetPackageStoreEntry(FPackageId PackageId, FName PackageName, FPackageStoreEntry& OutPackageStoreEntry);
 
 	/* Returns the redirected package ID and source package name for the specified package ID if it's being redirected. */
 	COREUOBJECT_API bool GetPackageRedirectInfo(FPackageId PackageId, FName& OutSourcePackageName, FPackageId& OutRedirectedToPackageId);
@@ -248,9 +249,12 @@ public:
 	COREUOBJECT_API bool HasAnyBackendsMounted() const;
 
 private:
+	friend class FPackageStoreInternal;
+	friend class FPackageStoreReadScope;
+
 	FPackageStore();
 
-	friend class FPackageStoreReadScope;
+	COREUOBJECT_API void IterateBackends(FVisitPackageStoreBackend&& Visit);
 
 	TSharedRef<FPackageStoreBackendContext> BackendContext;
 	
diff --git a/Engine/Source/Runtime/Experimental/IoStore/OnDemand/Private/OnDemandContentInstaller.cpp b/Engine/Source/Runtime/Experimental/IoStore/OnDemand/Private/OnDemandContentInstaller.cpp
--- a/Engine/Source/Runtime/Experimental/IoStore/OnDemand/Private/OnDemandContentInstaller.cpp
+++ b/Engine/Source/Runtime/Experimental/IoStore/OnDemand/Private/OnDemandContentInstaller.cpp
@@ -20,6 +20,7 @@
 #include "OnDemandInstallCache.h"
 #include "OnDemandIoStore.h"
 #include "OnDemandPackageStoreBackend.h"
+#include "Serialization/PackageStoreInternal.h"
 #include "Statistics.h"
 #include <atomic>
 
@@ -108,8 +109,34 @@ int32 SecondsToMillicseconds(double Seconds)
 	return Seconds > 0.0 ? int32(Seconds * 1000.0) : 0;
 }
 
+////////////////////////////////////////////////////////////////////////////////
+static EPackageStoreEntryStatus GetPackageStoreEntry(
+	FPackageStore& PackageStore,
+	IOnDemandPackageStoreBackend& OnDemandPackageStore,
+	const FPackageId& PackageId,
+	FPackageStoreEntry& OutPackageStoreEntry,
+	bool& bOnDemand)
+{
+	EPackageStoreEntryStatus EntryStatus = EPackageStoreEntryStatus::Missing;
+	bOnDemand = false;
+
+	FPackageStoreInternal::IterateBackends(
+		PackageStore,
+		[&EntryStatus, &OutPackageStoreEntry, &bOnDemand, &OnDemandPackageStore, &PackageId](IPackageStoreBackend& Backend)
+		{
+			bOnDemand				= &OnDemandPackageStore == &Backend;
+			EntryStatus				= Backend.GetPackageStoreEntry(PackageId, NAME_None, OutPackageStoreEntry);
+			const bool bContinue	= EntryStatus == EPackageStoreEntryStatus::None || EntryStatus == EPackageStoreEntryStatus::Missing;
+
+			return bContinue;
+		});
+
+	return EntryStatus;
+}
+
 ////////////////////////////////////////////////////////////////////////////////
 void VisitPackageDependencies(
+	const TSharedPtr<IOnDemandPackageStoreBackend>& OnDemandPackageStore,
 	const TSet<FPackageId>& PackageIds, 
 	bool bIncludeSoftReferences, 
 	FPackageDependencyVisitor Visitor, 
@@ -154,10 +181,15 @@ void VisitPackageDependencies(
 		}
 
 		FPackageStoreEntry PackageStoreEntry;
-		TSharedPtr<IPackageStoreBackend> Backend;
-		const EPackageStoreEntryStatus EntryStatus = PackageStore.GetPackageStoreEntry(PackageId, NAME_None, PackageStoreEntry, &Backend);
-
-		const bool bContinue = Visitor(PackageId, EntryStatus, PackageStoreEntry, Backend);
+		bool bOnDemand = false;
+		const EPackageStoreEntryStatus EntryStatus = GetPackageStoreEntry(
+			PackageStore,
+			*OnDemandPackageStore,
+			PackageId,
+			PackageStoreEntry,
+			bOnDemand);
+
+		const bool bContinue = Visitor(PackageId, EntryStatus, PackageStoreEntry, bOnDemand);
 		if (bContinue == false)
 		{
 			break;
@@ -223,19 +255,20 @@ void ResolveChunksToInstall(
 	TSet<FPackageId> ResolvedPackageIds;
 	TSet<FIoChunkId> ShaderMapChunkIds;
 	VisitPackageDependencies(
+		OnDemandPackageStore,
 		PackageIds,
 		bIncludeSoftReferences,
 		[&ResolvedPackageIds, &ShaderMapChunkIds, &OutMissing, &OnDemandPackageStore](
 			FPackageId PackageId,
 			EPackageStoreEntryStatus EntryStatus,
 			const FPackageStoreEntry& PackageStoreEntry,
-			const TSharedPtr<IPackageStoreBackend>& Backend)
+			bool bOnDemand)
 			{
 				if (IsPackageStoreEntryMissing(EntryStatus))
 				{
 					OutMissing.Add(CreatePackageDataChunkId(PackageId));
 				}
-				else if (Backend.Get() == OnDemandPackageStore.Get())
+				else if (bOnDemand)
 				{
 					ResolvedPackageIds.Add(PackageId);
 					if (PackageStoreEntry.ShaderMapHashes.IsEmpty() == false)
diff --git a/Engine/Source/Runtime/Experimental/IoStore/OnDemand/Private/OnDemandContentInstaller.h b/Engine/Source/Runtime/Experimental/IoStore/OnDemand/Private/OnDemandContentInstaller.h
--- a/Engine/Source/Runtime/Experimental/IoStore/OnDemand/Private/OnDemandContentInstaller.h
+++ b/Engine/Source/Runtime/Experimental/IoStore/OnDemand/Private/OnDemandContentInstaller.h
@@ -48,10 +48,11 @@ using FPackageDependencyVisitor = TFunctionRef<
 	bool(FPackageId /*PackageId*/,
 		EPackageStoreEntryStatus /*EntryStatus*/,
 		const FPackageStoreEntry& /*PackageStoreEntry*/,
-		const TSharedPtr<IPackageStoreBackend>& /*Backend*/
+		bool /*bOnDemand*/
 	)>;
 
 void VisitPackageDependencies(
+	const TSharedPtr<IOnDemandPackageStoreBackend>& OnDemandPackageStore,
 	const TSet<FPackageId>& PackageIds, 
 	bool bIncludeSoftReferences, 
 	FPackageDependencyVisitor Visitor, 
diff --git a/Engine/Source/Runtime/Experimental/IoStore/OnDemand/Private/OnDemandIoStore.cpp b/Engine/Source/Runtime/Experimental/IoStore/OnDemand/Private/OnDemandIoStore.cpp
--- a/Engine/Source/Runtime/Experimental/IoStore/OnDemand/Private/OnDemandIoStore.cpp
+++ b/Engine/Source/Runtime/Experimental/IoStore/OnDemand/Private/OnDemandIoStore.cpp
@@ -856,44 +856,34 @@ FIoStatus FOnDemandIoStore::GetIsOnDemand(TConstArrayView<FOnDemandIsOnDemandArg
 		return Status;
 	}
 
-	auto PackageVisitor = [this, OutResults](
-		int32 ResultIndex,
-		FPackageId PackageId,
-		EPackageStoreEntryStatus EntryStatus,
-		const FPackageStoreEntry& PackageStoreEntry,
-		const TSharedPtr<IPackageStoreBackend>& Backend) -> bool
-	{
-		if (Backend == PackageStoreBackend)
-		{
-			OutResults[ResultIndex].bIsOnDemand = true;
-			return false; // Don't continue
-		}
-
-		return true; // continue
-	};
-
 	FPackageStoreReadScope ReadScope(FPackageStore::Get());
 
 	// TODO: Parallel for here? Is that safe?
 	for (int Index = 0; const FOnDemandIsOnDemandArgs& Args : InArgs)
 	{
-		OutResults[Index].bIsOnDemand = false;
+		FOnDemandIsOnDemandResult& Result = OutResults[Index++];
+		Result.bIsOnDemand = false;
 
 		const bool bIncludeSoftReferences = EnumHasAnyFlags(Args.Options, EOnDemandIsOnDemandOptions::IncludeSoftReferences);
-
-		VisitPackageDependencies(TSet<FPackageId>(Args.PackageIds), bIncludeSoftReferences,
-			[PackageVisitor, Index](
+		VisitPackageDependencies(
+			PackageStoreBackend,
+			TSet<FPackageId>(Args.PackageIds),
+			bIncludeSoftReferences,
+			[&Result](
 				FPackageId PackageId,
 				EPackageStoreEntryStatus EntryStatus,
 				const FPackageStoreEntry& PackageStoreEntry,
-				const TSharedPtr<IPackageStoreBackend>& Backend) 
+				bool bIsOnDemand) 
 			{
-				return PackageVisitor(Index, PackageId, EntryStatus, PackageStoreEntry, Backend);
+				if (bIsOnDemand)
+				{
+					Result.bIsOnDemand = true;
+					return false;
+				}
+				return true;
 			},
 			&ReadScope
 		);
-
-		++Index;
 	}
 
 	return EIoErrorCode::Ok;
