[TraceBasedDebugger] relaxed validation for registered debuggers for remote session discovery.
This is to better handle the scenario where session discovery lifetime is tied to a debugger widget, and that widget get recreated and reassigned to the same SharedPtr.
In that case the same debugger was trying to register twice before the release SharedPtr will delete the first instance.
#release 40.00
#rnx

#Synced-CL 50651167

diff --git a/Engine/Source/Runtime/Experimental/TraceBasedDebuggers/Private/RemoteSessionsManager.cpp b/Engine/Source/Runtime/Experimental/TraceBasedDebuggers/Private/RemoteSessionsManager.cpp
--- a/Engine/Source/Runtime/Experimental/TraceBasedDebuggers/Private/RemoteSessionsManager.cpp
+++ b/Engine/Source/Runtime/Experimental/TraceBasedDebuggers/Private/RemoteSessionsManager.cpp
@@ -345,19 +345,14 @@ void FRemoteSessionsManager::StartSessionDiscovery(const FGuid& InDebuggerGuid)
 		TickHandle = FTSTicker::GetCoreTicker().AddTicker(FTickerDelegate::CreateRaw(this, &FRemoteSessionsManager::Tick), TickInterval);
 	}
 
-	const int32 PreviousCount = ActiveDebuggersUsingSessionDiscovery.Num();
-	ActiveDebuggersUsingSessionDiscovery.AddUnique(InDebuggerGuid);
-	if (PreviousCount == ActiveDebuggersUsingSessionDiscovery.Num())
-	{
-		UE_LOG(LogTraceBasedDebuggers, Warning, TEXT("[%hs] Session discovery already started by debugger %s"), __func__, *InDebuggerGuid.ToString());
-	}
+	ActiveDebuggersUsingSessionDiscovery.Add(InDebuggerGuid);
 }
 
 void FRemoteSessionsManager::StopSessionDiscovery(const FGuid& InDebuggerGuid)
 {
-	const int32 NumRemoved = ActiveDebuggersUsingSessionDiscovery.Remove(InDebuggerGuid);
+	const int32 NumRemoved = ActiveDebuggersUsingSessionDiscovery.RemoveSingle(InDebuggerGuid);
 
-	UE_CLOG(NumRemoved == 0, LogTraceBasedDebuggers, Warning, TEXT("[%hs] Session discovery was not started by debugger %s"), __func__, *InDebuggerGuid.ToString());
+	UE_CLOG(NumRemoved == 0, LogTraceBasedDebuggers, Warning, TEXT("[%hs] Session discovery was not started by debugger %s (or stopped more often than it was started"), __func__, *InDebuggerGuid.ToString());
 
 	if (ActiveDebuggersUsingSessionDiscovery.IsEmpty())
 	{
