Chaos++: Fixed some missing _API declarations.

#Synced-CL 50653012

diff --git a/Engine/Source/Runtime/Experimental/RigidPhysics/Public/RigidPhysics/RigidBodyContainerHandle.h b/Engine/Source/Runtime/Experimental/RigidPhysics/Public/RigidPhysics/RigidBodyContainerHandle.h
--- a/Engine/Source/Runtime/Experimental/RigidPhysics/Public/RigidPhysics/RigidBodyContainerHandle.h
+++ b/Engine/Source/Runtime/Experimental/RigidPhysics/Public/RigidPhysics/RigidBodyContainerHandle.h
@@ -18,7 +18,7 @@ namespace UE::Physics
 	{
 	public:
 		FRigidBodyContainerHandle() = default;
-		UE_INTERNAL FRigidBodyContainerHandle(const FRigidSceneHandle& InSceneHandle, const FRigidObjectId& InBodyContainerId);
+		UE_INTERNAL RIGIDPHYSICS_API FRigidBodyContainerHandle(const FRigidSceneHandle& InSceneHandle, const FRigidObjectId& InBodyContainerId);
 
 		friend bool operator==(const FRigidBodyContainerHandle&, const FRigidBodyContainerHandle&) = default;
 		friend bool operator!=(const FRigidBodyContainerHandle&, const FRigidBodyContainerHandle&) = default;
@@ -37,16 +37,16 @@ namespace UE::Physics
 			return nullptr;
 		}
 
-		UE_INTERNAL void Reset();
-		UE_INTERNAL const FRigidObjectId& GetId() const;
+		UE_INTERNAL RIGIDPHYSICS_API void Reset();
+		UE_INTERNAL RIGIDPHYSICS_API const FRigidObjectId& GetId() const;
 
 		// NOTE: We don't really need the scene handle because it is provided by the Context
 		// in Pin(). It is only needed for the debugger visualizers (see RigidPhysics.natvis).
 		// TODO_CHAOSAPI: Consider making SceneHandle shipping only
-		UE_INTERNAL const FRigidSceneHandle GetSceneHandle() const;
+		UE_INTERNAL RIGIDPHYSICS_API const FRigidSceneHandle GetSceneHandle() const;
 
 	private:
-		IRigidBodyContainer* PinInternal(const IRigidScene* InScene) const;
+		RIGIDPHYSICS_API IRigidBodyContainer* PinInternal(const IRigidScene* InScene) const;
 
 		FRigidObjectId BodyContainerId;
 		FRigidSceneHandle SceneHandle;
diff --git a/Engine/Source/Runtime/Experimental/RigidPhysics/Public/RigidPhysics/RigidBodyHandle.h b/Engine/Source/Runtime/Experimental/RigidPhysics/Public/RigidPhysics/RigidBodyHandle.h
--- a/Engine/Source/Runtime/Experimental/RigidPhysics/Public/RigidPhysics/RigidBodyHandle.h
+++ b/Engine/Source/Runtime/Experimental/RigidPhysics/Public/RigidPhysics/RigidBodyHandle.h
@@ -30,7 +30,7 @@ namespace UE::Physics
 		FRigidBodyHandle() = default;
 		FRigidBodyHandle(const FRigidBodyHandle&) = default;
 
-		UE_INTERNAL FRigidBodyHandle(const FRigidBodyContainerHandle& InBodyContainerHandle, const FRigidObjectId& InBodyId);
+		UE_INTERNAL RIGIDPHYSICS_API FRigidBodyHandle(const FRigidBodyContainerHandle& InBodyContainerHandle, const FRigidObjectId& InBodyId);
 
 		friend bool operator==(const FRigidBodyHandle&, const FRigidBodyHandle&) = default;
 		friend bool operator!=(const FRigidBodyHandle&, const FRigidBodyHandle&) = default;
@@ -49,12 +49,12 @@ namespace UE::Physics
 			return TRigidBodyPtr<ContextType>();
 		}
 
-		UE_INTERNAL const FRigidObjectId& GetId() const;
-		UE_INTERNAL const FRigidBodyContainerHandle& GetContainerHandle() const;
-		UE_INTERNAL void Reset();
+		UE_INTERNAL RIGIDPHYSICS_API const FRigidObjectId& GetId() const;
+		UE_INTERNAL RIGIDPHYSICS_API const FRigidBodyContainerHandle& GetContainerHandle() const;
+		UE_INTERNAL RIGIDPHYSICS_API void Reset();
 
 	private:
-		IRigidBody* PinInternal(const IRigidBodyContainer* Container) const;
+		RIGIDPHYSICS_API IRigidBody* PinInternal(const IRigidBodyContainer* Container) const;
 
 		FRigidObjectId BodyId;
 		FRigidBodyContainerHandle BodyContainerHandle;
diff --git a/Engine/Source/Runtime/Experimental/RigidPhysics/Public/RigidPhysics/RigidGeometryCollectionHandle.h b/Engine/Source/Runtime/Experimental/RigidPhysics/Public/RigidPhysics/RigidGeometryCollectionHandle.h
--- a/Engine/Source/Runtime/Experimental/RigidPhysics/Public/RigidPhysics/RigidGeometryCollectionHandle.h
+++ b/Engine/Source/Runtime/Experimental/RigidPhysics/Public/RigidPhysics/RigidGeometryCollectionHandle.h
@@ -35,7 +35,7 @@ namespace UE::Physics
 	{
 	public:
 		FRigidGeometryCollectionHandle() = default;
-		UE_INTERNAL FRigidGeometryCollectionHandle(const FRigidSceneHandle& InSceneHandle, const FRigidObjectId& InBodyContainerId);
+		UE_INTERNAL RIGIDPHYSICS_API FRigidGeometryCollectionHandle(const FRigidSceneHandle& InSceneHandle, const FRigidObjectId& InBodyContainerId);
 
 		friend bool operator==(const FRigidBodyContainerHandle&, const FRigidBodyContainerHandle&) = default;
 		friend bool operator!=(const FRigidBodyContainerHandle&, const FRigidBodyContainerHandle&) = default;
@@ -52,15 +52,15 @@ namespace UE::Physics
 			return nullptr;
 		}
 
-		UE_INTERNAL void Reset();
-		UE_INTERNAL const FRigidObjectId& GetId() const;
-		UE_INTERNAL const FRigidSceneHandle GetSceneHandle() const;
+		UE_INTERNAL RIGIDPHYSICS_API void Reset();
+		UE_INTERNAL RIGIDPHYSICS_API const FRigidObjectId& GetId() const;
+		UE_INTERNAL RIGIDPHYSICS_API const FRigidSceneHandle GetSceneHandle() const;
 
-		UE_INTERNAL FRigidBodyContainerHandle AsBodyContainerHandle() const;
-		UE_INTERNAL operator FRigidBodyContainerHandle() const;
+		UE_INTERNAL RIGIDPHYSICS_API FRigidBodyContainerHandle AsBodyContainerHandle() const;
+		UE_INTERNAL RIGIDPHYSICS_API operator FRigidBodyContainerHandle() const;
 
 	private:
-		IRigidGeometryCollection* PinInternal(IRigidBodyContainer* Container) const;
+		RIGIDPHYSICS_API IRigidGeometryCollection* PinInternal(IRigidBodyContainer* Container) const;
 
 		FRigidBodyContainerHandle BodyContainerHandle;
 	};
