Dataflow : add a new function to evaluate an external dataflow on an asset from a utility blueprint that properly go through the terminals

#Synced-CL 50641382

diff --git a/Engine/Source/Runtime/Experimental/Dataflow/Core/Private/Dataflow/DataflowContextEvaluator.cpp b/Engine/Source/Runtime/Experimental/Dataflow/Core/Private/Dataflow/DataflowContextEvaluator.cpp
--- a/Engine/Source/Runtime/Experimental/Dataflow/Core/Private/Dataflow/DataflowContextEvaluator.cpp
+++ b/Engine/Source/Runtime/Experimental/Dataflow/Core/Private/Dataflow/DataflowContextEvaluator.cpp
@@ -303,6 +303,8 @@ namespace UE::Dataflow
 		bIsCanceled = false;
 		bIsRunning = true;
 
+		UE_LOG(LogDataflowEvaluator, Verbose, TEXT("FDataflowDirectEvaluator::Tick -- START"));
+
 		int32 LastNodeIndex = INDEX_NONE;
 		FGuid NodeId;
 		TSharedRef<const FGraph> DataflowGraph = InParams.CompiledGraph->GetSourceGraph();
@@ -315,16 +317,19 @@ namespace UE::Dataflow
 				{
 					ProgressTracker.SetProgress((float)LastNodeIndex, (float)InParams.CompiledGraph->GetNumNodes(), Node->GetName());
 
-					if (ensure(!Node->EvaluateOnGameThreadOnly()))
+					const double StartEvalTimeMs = FPlatformTime::ToMilliseconds64(FPlatformTime::Cycles64());
+
+					if (Node->EvaluateNode(OwningContext))
 					{
-						if (Node->EvaluateNode(OwningContext))
+						if (Node->IsTerminal() && InParams.OnTerminalNodeEvaluated.IsSet())
 						{
-							if (Node->IsTerminal() && InParams.OnTerminalNodeEvaluated.IsSet())
-							{
-								InParams.OnTerminalNodeEvaluated(OwningContext, *Node);
-							}
+							InParams.OnTerminalNodeEvaluated(OwningContext, *Node);
 						}
 					}
+
+					const double EndEvalTimeMs = FPlatformTime::ToMilliseconds64(FPlatformTime::Cycles64());
+
+					UE_LOG(LogDataflowEvaluator, Verbose, TEXT("FDataflowDirectEvaluator::Tick --     Node [%s] : [%.3f ms]"), *Node->GetName().ToString(), (float)(EndEvalTimeMs - StartEvalTimeMs));
 				}
 			}
 		} while (NodeId.IsValid() && !bIsCanceled);
@@ -336,6 +341,8 @@ namespace UE::Dataflow
 		{
 			InParams.OnPostEvaluation(OwningContext);
 		}
+
+		UE_LOG(LogDataflowEvaluator, Verbose, TEXT("FDataflowDirectEvaluator::Tick -- END"));
 	}
 
 	void FDataflowDirectEvaluator::Cancel()
diff --git a/Engine/Source/Runtime/Experimental/Dataflow/Core/Public/Dataflow/DataflowContextEvaluator.h b/Engine/Source/Runtime/Experimental/Dataflow/Core/Public/Dataflow/DataflowContextEvaluator.h
--- a/Engine/Source/Runtime/Experimental/Dataflow/Core/Public/Dataflow/DataflowContextEvaluator.h
+++ b/Engine/Source/Runtime/Experimental/Dataflow/Core/Public/Dataflow/DataflowContextEvaluator.h
@@ -151,8 +151,7 @@ namespace UE::Dataflow
 	/**
 	* This evaluator will evaluate the graph in one go when Start is called
 	* This is useful for running simulation graph for example that needs to execute in exactly one frame
-	* it is required to run from a worker thread and will assert if a node requires gamethread evaluation
-	* Note that it is safe to cancel or query if its running or its progress from another thread 
+	* or evaluating graph from a blueprint 
 	*/
 	struct FDataflowDirectEvaluator : public IDataflowEvaluator
 	{
diff --git a/Engine/Source/Runtime/Experimental/Dataflow/Engine/Private/Dataflow/DataflowBlueprintLibrary.cpp b/Engine/Source/Runtime/Experimental/Dataflow/Engine/Private/Dataflow/DataflowBlueprintLibrary.cpp
--- a/Engine/Source/Runtime/Experimental/Dataflow/Engine/Private/Dataflow/DataflowBlueprintLibrary.cpp
+++ b/Engine/Source/Runtime/Experimental/Dataflow/Engine/Private/Dataflow/DataflowBlueprintLibrary.cpp
@@ -37,6 +37,32 @@ void UDataflowBlueprintLibrary::EvaluateTerminalNodeByName(UDataflow* Dataflow,
 	}
 }
 
+bool UDataflowBlueprintLibrary::EvaluateDataflow(UDataflow* Dataflow, UObject* AssetToUpdate)
+{
+	if (Dataflow && Dataflow->Dataflow)
+	{
+		using namespace UE::Dataflow;
+		FEngineContext Context(AssetToUpdate);
+		Context.GetOnContextLogMulticast().AddLambda(
+			[](const FContext::FLogMessage& Message)
+			{
+				if (Message.Severity == EMessageSeverity::Error)
+				{
+					FString NodeString(TEXT("None"));
+					if (TSharedPtr<const FDataflowNode> Node = Message.Node.Pin())
+					{
+						NodeString = Node->GetName().ToString();
+					}
+					UE_LOG(LogChaosDataflow, Error, TEXT("EvaluateDataflow : (%s) %s"), *NodeString, *Message.Text.ToString());
+				}
+			});
+		Context.EvaluateGraph<FDataflowDirectEvaluator>(Dataflow, {});
+		return (Context.GetNumErrors() == 0);
+	}
+	UE_LOG(LogChaosDataflow, Error, TEXT("EvaluateDataflow : Invalid Dataflow Asset : %s"), *(Dataflow? Dataflow->GetFName().ToString(): FString(TEXT("null"))));
+	return false;
+}
+
 bool UDataflowBlueprintLibrary::RegenerateAssetFromDataflow(UObject* AssetToRegenerate, bool bRegenerateDependentAssets)
 {
 	if (IDataflowInstanceInterface* DataflowInterface = Cast<IDataflowInstanceInterface>(AssetToRegenerate))
diff --git a/Engine/Source/Runtime/Experimental/Dataflow/Engine/Private/Dataflow/DataflowObjectInterface.cpp b/Engine/Source/Runtime/Experimental/Dataflow/Engine/Private/Dataflow/DataflowObjectInterface.cpp
--- a/Engine/Source/Runtime/Experimental/Dataflow/Engine/Private/Dataflow/DataflowObjectInterface.cpp
+++ b/Engine/Source/Runtime/Experimental/Dataflow/Engine/Private/Dataflow/DataflowObjectInterface.cpp
@@ -26,23 +26,23 @@ namespace UE::Dataflow
 			}
 		}
 
-		bool EvaluateGraphInternal(TSharedRef<IDataflowEvaluator> Evaluator, UObject* Owner, FOnPostEvaluationFunction OnPostEvaluation)
+		bool EvaluateGraphInternal(TSharedRef<IDataflowEvaluator> Evaluator, UDataflow* DataflowObject, UObject* ObjectToUpdate, FOnPostEvaluationFunction OnPostEvaluation)
 		{
-			if (UDataflow* DataflowObject = InstanceUtils::GetDataflowAssetFromObject(Owner))
+			if (DataflowObject)
 			{
 				if (TSharedPtr<FGraph> Graph = DataflowObject->GetDataflow())
 				{
 					if (TSharedPtr<const UE::Dataflow::FCompiledGraph> CompiledGraph = DataflowObject->CompileGraphIfNeeded())
 					{
-						TWeakObjectPtr<UObject> WeakOwner = Owner;
+						TWeakObjectPtr<UObject> WeakObjectToUpdate = ObjectToUpdate;
 
 						FDataflowEvaluatorParameters Params(CompiledGraph.ToSharedRef());
 						Params.OnPostEvaluation = OnPostEvaluation;
-						Params.OnTerminalNodeEvaluated = [WeakOwner](FContext& Context, const FDataflowNode& Node)
+						Params.OnTerminalNodeEvaluated = [WeakObjectToUpdate](FContext& Context, const FDataflowNode& Node)
 							{
-								if (TStrongObjectPtr<UObject> Owner = WeakOwner.Pin())
+								if (TStrongObjectPtr<UObject> ObjectToUpdate = WeakObjectToUpdate.Pin())
 								{
-									UpdateAssetFromNode(Context, Node, Owner.Get());
+									UpdateAssetFromNode(Context, Node, ObjectToUpdate.Get());
 								}
 							};
 						Evaluator->Start(Params);
@@ -52,5 +52,14 @@ namespace UE::Dataflow
 			}
 			return false;
 		}
+
+		bool EvaluateGraphInternal(TSharedRef<IDataflowEvaluator> Evaluator, UObject* ObjectToUpdate, FOnPostEvaluationFunction OnPostEvaluation)
+		{
+			if (UDataflow* DataflowObject = InstanceUtils::GetDataflowAssetFromObject(ObjectToUpdate))
+			{
+				return EvaluateGraphInternal(Evaluator, DataflowObject, ObjectToUpdate, OnPostEvaluation);
+			}
+			return false;
+		}
 	}
 }
diff --git a/Engine/Source/Runtime/Experimental/Dataflow/Engine/Public/Dataflow/DataflowBlueprintLibrary.h b/Engine/Source/Runtime/Experimental/Dataflow/Engine/Public/Dataflow/DataflowBlueprintLibrary.h
--- a/Engine/Source/Runtime/Experimental/Dataflow/Engine/Public/Dataflow/DataflowBlueprintLibrary.h
+++ b/Engine/Source/Runtime/Experimental/Dataflow/Engine/Public/Dataflow/DataflowBlueprintLibrary.h
@@ -43,6 +43,17 @@ public:
 	UFUNCTION(BlueprintCallable, Category = "Dataflow", meta = (Keywords = "Dataflow graph"))
 	static DATAFLOWENGINE_API void EvaluateTerminalNodeByName(UDataflow* Dataflow, FName TerminalNodeName, UObject* ResultAsset);
 
+	/**s
+	* Evaluate a dataflow for a sepcific asset
+	* This will evaluate all the terminal nodes and generate any dependent asset if necessary 
+	* This is synchronous operation, so large dataflow may block the game thread for a while when calling it 
+	* @param Dataflow Dataflow asset to evaluate
+	* @param AssetToUpdate Asset to update from the Dataflow ( if a terminal exists for this type of asset ) 
+	* @return true if the evaluate executed with no errors
+	*/
+	UFUNCTION(BlueprintCallable, Category = "Dataflow", meta = (Keywords = "Dataflow graph"))
+	static DATAFLOWENGINE_API bool EvaluateDataflow(UDataflow* Dataflow, UObject* AssetToUpdate);
+
 	/**
 	* Regenerate an asset using its corresponding dataflow
 	* @return true if the asset was regenerated
diff --git a/Engine/Source/Runtime/Experimental/Dataflow/Engine/Public/Dataflow/DataflowObjectInterface.h b/Engine/Source/Runtime/Experimental/Dataflow/Engine/Public/Dataflow/DataflowObjectInterface.h
--- a/Engine/Source/Runtime/Experimental/Dataflow/Engine/Public/Dataflow/DataflowObjectInterface.h
+++ b/Engine/Source/Runtime/Experimental/Dataflow/Engine/Public/Dataflow/DataflowObjectInterface.h
@@ -14,6 +14,7 @@ namespace UE::Dataflow
 	{
 		DATAFLOWENGINE_API void UpdateAssetFromTerminalNodes(FContext& Context, UObject* Asset, bool bUseAllTerminals);
 		DATAFLOWENGINE_API bool EvaluateGraphInternal(TSharedRef<IDataflowEvaluator> Evaluator, UObject* Owner, FOnPostEvaluationFunction OnPostEvaluation);
+		DATAFLOWENGINE_API bool EvaluateGraphInternal(TSharedRef<IDataflowEvaluator> Evaluator, UDataflow* DataflowObject, UObject* ObjectToUpdate, FOnPostEvaluationFunction OnPostEvaluation);
 	}
 
 	template<class Base = FContextSingle>
@@ -37,6 +38,11 @@ namespace UE::Dataflow
 
 		virtual void Serialize(FArchive& Ar) { Base::Serialize(Ar); }
 
+		/**
+		* Evaluate the dataflow graph attached to the owner of this context and update the owner accordingly
+		* Returns an evaluator that can be used to query the progress of the evaluation 
+		* Note that the evaluator is not kept by this context and needs to be kept by the caller
+		*/
 		template <typename T UE_REQUIRES(std::is_base_of_v<IDataflowEvaluator, T>)>
 		TSharedPtr<IDataflowEvaluator> EvaluateGraph(FOnPostEvaluationFunction OnPostEvaluation)
 		{
@@ -48,6 +54,22 @@ namespace UE::Dataflow
 			return {};
 			
 		}
+
+		/**
+		* Evaluate an external dataflow graph and update the owner of this context with it 
+		* Returns an evaluator that can be used to query the progress of the evaluation 
+		* Note that the evaluator is not kept by this context and needs to be kept by the caller
+		*/
+		template <typename T UE_REQUIRES(std::is_base_of_v<IDataflowEvaluator, T>)>
+		TSharedPtr<IDataflowEvaluator> EvaluateGraph(UDataflow* DataflowObject, FOnPostEvaluationFunction OnPostEvaluation)
+		{
+			TSharedRef<IDataflowEvaluator> Evaluator = MakeShared<T>(*this);
+			if (Private::EvaluateGraphInternal(Evaluator, DataflowObject, Owner, OnPostEvaluation))
+			{
+				return Evaluator.ToSharedPtr();
+			}
+			return {};
+		}
 	};
 
 	typedef TEngineContext<FContextSingle> FEngineContext;
