Try convert chunk IDs to package and filename when logging missing chunks.

#rnx

#Synced-CL 50723010

diff --git a/Engine/Source/Runtime/Experimental/IoStore/OnDemand/Private/OnDemandContentInstaller.cpp b/Engine/Source/Runtime/Experimental/IoStore/OnDemand/Private/OnDemandContentInstaller.cpp
--- a/Engine/Source/Runtime/Experimental/IoStore/OnDemand/Private/OnDemandContentInstaller.cpp
+++ b/Engine/Source/Runtime/Experimental/IoStore/OnDemand/Private/OnDemandContentInstaller.cpp
@@ -13,12 +13,14 @@
 #include "IO/IoStatus.h"
 #include "IO/IoStoreOnDemand.h"
 #include "IO/OnDemandError.h"
+#include "Logging/StructuredLog.h"
 #include "Misc/App.h"
 #include "Misc/CoreDelegatesInternal.h"
 #include "Misc/Timespan.h"
 #include "OnDemandContentInstallReplay.h"
 #include "OnDemandInstallCache.h"
 #include "OnDemandIoStore.h"
+#include "OnDemandMisc.h"
 #include "OnDemandPackageStoreBackend.h"
 #include "Serialization/PackageStoreInternal.h"
 #include "Statistics.h"
@@ -913,14 +915,24 @@ void FOnDemandContentInstaller::ProcessInstallRequest(FRequest& Request)
 	{
 		FIoDispatcher& IoDispatcher = FIoDispatcher::Get();
 		TArray<FIoChunkId> MissingChunkIds;
+		TUtf8StringBuilder<128> Filename;
+		TStringBuilder<128> Sb;
+
 		for (const FIoChunkId& ChunkId : Missing)
 		{
 			if (IoDispatcher.DoesChunkExist(ChunkId) == false)
 			{
 				UE_CLOG(MissingChunkIds.Num() == 0, LogIoStoreOnDemand, Warning, TEXT("Failed to resolve the following chunk(s) for content handle '%s':"),
 					*LexToString(InstallRequest.Args.ContentHandle));
-
-				UE_LOG(LogIoStoreOnDemand, Warning, TEXT("\tChunkId='%s' (%s)"), *LexToString(ChunkId), *LexToString(ChunkId.GetChunkType()));
+				
+				Sb.Reset();
+				Filename.Reset();
+
+				Sb << TEXT("PackageName='");
+				TryConvertChunkIdToPackageName(ChunkId, Filename, Sb);
+				Sb << TEXT("' Filename='") << Filename << TEXT("'");
+				Sb << TEXT(" ChunkId='") << ChunkId << TEXT("' (") << LexToString(ChunkId.GetChunkType()) << TEXT(")");
+				UE_LOGFMT(LogIoStoreOnDemand, Warning, "\t{Msg}", Sb);
 				MissingChunkIds.Add(ChunkId);
 			}
 		}
