Fix bug GetIsOnDemand returning true for missing package(s).

#rnx

#Synced-CL 50723010

diff --git a/Engine/Source/Runtime/Experimental/IoStore/OnDemand/Private/OnDemandIoStore.cpp b/Engine/Source/Runtime/Experimental/IoStore/OnDemand/Private/OnDemandIoStore.cpp
--- a/Engine/Source/Runtime/Experimental/IoStore/OnDemand/Private/OnDemandIoStore.cpp
+++ b/Engine/Source/Runtime/Experimental/IoStore/OnDemand/Private/OnDemandIoStore.cpp
@@ -875,7 +875,7 @@ FIoStatus FOnDemandIoStore::GetIsOnDemand(TConstArrayView<FOnDemandIsOnDemandArg
 				const FPackageStoreEntry& PackageStoreEntry,
 				bool bIsOnDemand) 
 			{
-				if (bIsOnDemand)
+				if (bIsOnDemand && (EntryStatus == EPackageStoreEntryStatus::Ok || EntryStatus == EPackageStoreEntryStatus::NotInstalled))
 				{
 					Result.bIsOnDemand = true;
 					return false;
