Dataflow : Fix compiler sometimes not eveluating terminal node in order

#Synced-CL 50764219

diff --git a/Engine/Source/Runtime/Experimental/Dataflow/Engine/Private/Dataflow/DataflowCompiler.cpp b/Engine/Source/Runtime/Experimental/Dataflow/Engine/Private/Dataflow/DataflowCompiler.cpp
--- a/Engine/Source/Runtime/Experimental/Dataflow/Engine/Private/Dataflow/DataflowCompiler.cpp
+++ b/Engine/Source/Runtime/Experimental/Dataflow/Engine/Private/Dataflow/DataflowCompiler.cpp
@@ -180,6 +180,9 @@ namespace UE::Dataflow
 			FConstNodePtr NodeToProcess = StartNode;
 			while (NodeToProcess)
 			{
+				// Important : get the downstream nodes first 
+				FCompiler::FindDownstreamNodes(NodeToProcess.ToSharedRef(), DownstreamNodes);
+
 				// if the current node is a terminal we stop the sequence right away
 				if (NodeToProcess->IsTerminal())
 				{
@@ -187,7 +190,6 @@ namespace UE::Dataflow
 				}
 
 				// check if we have a singular node connected to this one
-				FCompiler::FindDownstreamNodes(NodeToProcess.ToSharedRef(), DownstreamNodes);
 				if (DownstreamNodes.Num() != 1)
 				{
 					// No more nodes or too many, stop building 
