Bumped log verbosity in on-demand content intaller.

#rnx

#Synced-CL 50723010

diff --git a/Engine/Source/Runtime/Experimental/IoStore/OnDemand/Private/OnDemandContentInstaller.cpp b/Engine/Source/Runtime/Experimental/IoStore/OnDemand/Private/OnDemandContentInstaller.cpp
--- a/Engine/Source/Runtime/Experimental/IoStore/OnDemand/Private/OnDemandContentInstaller.cpp
+++ b/Engine/Source/Runtime/Experimental/IoStore/OnDemand/Private/OnDemandContentInstaller.cpp
@@ -507,7 +507,7 @@ void FOnDemandContentInstaller::CancelInstallRequest(FSharedInternalInstallReque
 					}
 #endif // !UE_BUILD_SHIPPING
 
-					UE_LOG(LogIoStoreOnDemand, Log, TEXT("Cancelling install request, ContentHandle=(%s)"),
+					UE_LOG(LogIoStoreOnDemand, Verbose, TEXT("Cancelling install request, ContentHandle=(%s)"),
 						*LexToString(Install.Args.ContentHandle));
 
 					int32 NumCancelled = 0;
@@ -528,7 +528,7 @@ void FOnDemandContentInstaller::CancelInstallRequest(FSharedInternalInstallReque
 					}
 #endif // !UE_BUILD_SHIPPING
 
-					UE_LOG(LogIoStoreOnDemand, Log, TEXT("Cancelling install request, ContentHandle=(%s)"),
+					UE_LOG(LogIoStoreOnDemand, Verbose, TEXT("Cancelling install request, ContentHandle=(%s)"),
 						*LexToString(Install.Args.ContentHandle));
 
 					TUniqueLock Lock(Mutex);
@@ -571,7 +571,7 @@ void FOnDemandContentInstaller::UpdateInstallRequestPriority(FSharedInternalInst
 			FRequest&			Request	= *reinterpret_cast<FRequest*>(InstallRequest->InstallerRequest);
 			FRequest::FInstall& Install	= Request.AsInstall();
 
-			UE_LOG(LogIoStoreOnDemand, Verbose, TEXT("Updating install request priority, SeqNo=%u, Priority=%d, NewPriority=%d, ContentHandle=(%s)"),
+			UE_LOG(LogIoStoreOnDemand, VeryVerbose, TEXT("Updating install request priority, SeqNo=%u, Priority=%d, NewPriority=%d, ContentHandle=(%s)"),
 				Request.SeqNo, Request.Priority, NewPriority, *LexToString(Install.Args.ContentHandle));
 
 			if (Install.bHttpRequestsIssued)
@@ -1457,7 +1457,7 @@ void FOnDemandContentInstaller::OnProcessDownloadedChunkNotifyRequest(
 	static constexpr const TCHAR LogFmt[] = TEXT("Install progress %.2lf/%.2lf MiB, SeqNo=%u, Priority=%d, ContentHandle=(%s), ChunkId='%s', ChunkSize=%.2lf KiB, DownloadResult=%s");
 	if (bChunkCancelled)
 	{
-		UE_LOG(LogIoStoreOnDemand, Verbose, LogFmt,
+		UE_LOG(LogIoStoreOnDemand, VeryVerbose, LogFmt,
 			double(InstallRequest.Progress.CurrentInstallSize) / 1024.0 / 1024.0, double(InstallRequest.Progress.TotalInstallSize) / 1024.0 / 1024.0,
 			Request.SeqNo, Request.Priority, *LexToString(InstallRequest.Args.ContentHandle), *LexToString(ChunkId), double(ChunkEntry.GetDiskSize()) / 1024.0,
 			TEXT("Cancelled"));
