IAX DevMode

* A non-shipping development mode that mounts both installed and on-demand content for a container file, if available.
* Allows unreferenced (not installed) content to be loaded from installed content when possible.
* Logs warning messages when attempting to load on-demand content without installing it first.
* Enabled with -Iax.DevMode.

#rnx

#Synced-CL 50723010

diff --git a/Engine/Source/Runtime/Experimental/IoStore/OnDemand/Private/OnDemandContentInstaller.cpp b/Engine/Source/Runtime/Experimental/IoStore/OnDemand/Private/OnDemandContentInstaller.cpp
--- a/Engine/Source/Runtime/Experimental/IoStore/OnDemand/Private/OnDemandContentInstaller.cpp
+++ b/Engine/Source/Runtime/Experimental/IoStore/OnDemand/Private/OnDemandContentInstaller.cpp
@@ -126,11 +126,23 @@ static EPackageStoreEntryStatus GetPackageStoreEntry(
 		PackageStore,
 		[&EntryStatus, &OutPackageStoreEntry, &bOnDemand, &OnDemandPackageStore, &PackageId](IPackageStoreBackend& Backend)
 		{
-			bOnDemand				= &OnDemandPackageStore == &Backend;
-			EntryStatus				= Backend.GetPackageStoreEntry(PackageId, NAME_None, OutPackageStoreEntry);
-			const bool bContinue	= EntryStatus == EPackageStoreEntryStatus::None || EntryStatus == EPackageStoreEntryStatus::Missing;
+			const bool bOnDemandPackageStore = &OnDemandPackageStore == &Backend;
+			if (bOnDemandPackageStore)
+			{
+				// When running with -iax.devmode, the on-demand package store reports
+				// unreferenced entries as 'Missing' so that lower-priority backends
+				// (e.g., the file backend) can attempt to load the package.
+				// Query the on-demand package store directly to bypass the development mode check.
+				EntryStatus = OnDemandPackageStore.GetPackageStoreEntryInternal(PackageId, OutPackageStoreEntry);
+			}
+			else
+			{
+				EntryStatus = Backend.GetPackageStoreEntry(PackageId, NAME_None, OutPackageStoreEntry);
+			}
 
-			return bContinue;
+			const bool bMissing = EntryStatus == EPackageStoreEntryStatus::None || EntryStatus == EPackageStoreEntryStatus::Missing;
+			bOnDemand = !bMissing && bOnDemandPackageStore;
+			return bMissing;
 		});
 
 	return EntryStatus;
diff --git a/Engine/Source/Runtime/Experimental/IoStore/OnDemand/Private/OnDemandInstallCache.cpp b/Engine/Source/Runtime/Experimental/IoStore/OnDemand/Private/OnDemandInstallCache.cpp
--- a/Engine/Source/Runtime/Experimental/IoStore/OnDemand/Private/OnDemandInstallCache.cpp
+++ b/Engine/Source/Runtime/Experimental/IoStore/OnDemand/Private/OnDemandInstallCache.cpp
@@ -4,6 +4,7 @@
 #include "DiskCacheGovernor.h"
 #include "OnDemandHttpClient.h"
 #include "OnDemandIoStore.h"
+#include "OnDemandMisc.h"
 #include "Statistics.h"
 
 #include "Algo/Accumulate.h"
@@ -78,6 +79,23 @@ namespace CVars
 	);
 }
 
+///////////////////////////////////////////////////////////////////////////////
+static void LogUnreferencedChunk(const FIoChunkId& ChunkId, bool bLogAsWarning)
+{
+#if !UE_BUILD_SHIPPING
+	TUtf8StringBuilder<128> Filename;
+	TStringBuilder<128> Msg;
+
+	Msg << TEXT("Trying to load unreferenced chunk, PackageName='");
+	TryConvertChunkIdToPackageName(ChunkId, Filename, Msg);
+	Msg << TEXT("' Filename='") << Filename << TEXT("'");
+	Msg << TEXT(" ChunkId='") << ChunkId << TEXT("' (") << LexToString(ChunkId.GetChunkType()) << TEXT(")");
+
+	UE_CLOGFMT(bLogAsWarning,  LogIoStoreOnDemand, Warning, "{Msg}", Msg);
+	UE_CLOGFMT(!bLogAsWarning, LogIoStoreOnDemand, VeryVerbose, "{Msg}", Msg);
+#endif
+}
+
 ///////////////////////////////////////////////////////////////////////////////
 double ToKiB(uint64 Value)
 {
@@ -2435,8 +2453,13 @@ bool FOnDemandInstallCache::Resolve(FIoRequestImpl* Request)
 	{
 		if (ErrorCode == EIoErrorCode::NotInstalled)
 		{
-			CompleteRequest(Request, EIoErrorCode::NotInstalled);
-			return true;
+			const bool bLogAsWarning = IsDevModeEnabled();
+			LogUnreferencedChunk(Request->ChunkId, bLogAsWarning);
+			if (!IsDevModeEnabled())
+			{
+				CompleteRequest(Request, EIoErrorCode::NotInstalled);
+				return true;
+			}
 		}
 		return false;
 	}
diff --git a/Engine/Source/Runtime/Experimental/IoStore/OnDemand/Private/OnDemandIoStore.cpp b/Engine/Source/Runtime/Experimental/IoStore/OnDemand/Private/OnDemandIoStore.cpp
--- a/Engine/Source/Runtime/Experimental/IoStore/OnDemand/Private/OnDemandIoStore.cpp
+++ b/Engine/Source/Runtime/Experimental/IoStore/OnDemand/Private/OnDemandIoStore.cpp
@@ -68,6 +68,17 @@ static FAutoConsoleVariableRef CVar_IaxMemoryMappedTocEnabled(
 	TEXT("Whether container TOC memory-mapping is enabled.")
 );
 
+///////////////////////////////////////////////////////////////////////////////
+bool IsDevModeEnabled()
+{
+#if UE_BUILD_SHIPPING
+	return false;
+#else
+	static bool bIsDevModeEnabled = FParse::Param(FCommandLine::Get(), TEXT("iax.DevMode"));
+	return bIsDevModeEnabled;
+#endif
+}
+
 ///////////////////////////////////////////////////////////////////////////////
 bool IsPackageStreamingEnabled()
 {
@@ -1049,7 +1060,7 @@ FIoStatus FOnDemandIoStore::InitializeInstallCache()
 	{
 		int32 BackendPriority = -5; // Lower than file (zero) but higher than streaming backend (-10)
 #if !UE_BUILD_SHIPPING
-		if (FParse::Param(FCommandLine::Get(), TEXT("Iad")))
+		if (IsDevModeEnabled() || FParse::Param(FCommandLine::Get(), TEXT("Iad")))
 		{
 			// Bump the priority to be higher then the file system backend
 			BackendPriority = 5;
@@ -1057,7 +1068,7 @@ FIoStatus FOnDemandIoStore::InitializeInstallCache()
 #endif
 		FIoDispatcher::Get().Mount(InstallCache.ToSharedRef(), BackendPriority);
 		PackageStoreBackend = MakeOnDemandPackageStoreBackend(AsShared().ToWeakPtr());
-		FPackageStore::Get().Mount(PackageStoreBackend.ToSharedRef());
+		FPackageStore::Get().Mount(PackageStoreBackend.ToSharedRef(), BackendPriority);
 	}
 	else
 	{
diff --git a/Engine/Source/Runtime/Experimental/IoStore/OnDemand/Private/OnDemandIoStore.h b/Engine/Source/Runtime/Experimental/IoStore/OnDemand/Private/OnDemandIoStore.h
--- a/Engine/Source/Runtime/Experimental/IoStore/OnDemand/Private/OnDemandIoStore.h
+++ b/Engine/Source/Runtime/Experimental/IoStore/OnDemand/Private/OnDemandIoStore.h
@@ -37,6 +37,10 @@ using FSharedPackageStoreBackend	= TSharedPtr<class IOnDemandPackageStoreBackend
 using FSharedInstallCache			= TSharedPtr<class IOnDemandInstallCache>;
 using FSharedHttpIoBackend			= TSharedPtr<class IOnDemandHttpIoDispatcherBackend>;
 
+///////////////////////////////////////////////////////////////////////////////
+bool IsDevModeEnabled();
+bool IsPackageStreamingEnabled();
+
 ///////////////////////////////////////////////////////////////////////////////
 enum class EOnDemandContainerFlags : uint8
 {
diff --git a/Engine/Source/Runtime/Experimental/IoStore/OnDemand/Private/OnDemandPackageStoreBackend.cpp b/Engine/Source/Runtime/Experimental/IoStore/OnDemand/Private/OnDemandPackageStoreBackend.cpp
--- a/Engine/Source/Runtime/Experimental/IoStore/OnDemand/Private/OnDemandPackageStoreBackend.cpp
+++ b/Engine/Source/Runtime/Experimental/IoStore/OnDemand/Private/OnDemandPackageStoreBackend.cpp
@@ -6,16 +6,46 @@
 #include "Algo/Find.h"
 #include "Async/Mutex.h"
 #include "Async/UniqueLock.h"
+#include "CoreGlobals.h"
+#include "IO/IoChunkId.h"
 #include "IO/IoContainerHeader.h"
 #include "IO/IoStoreOnDemand.h"
 #include "IO/PackageId.h"
 #include "Internationalization/PackageLocalizationManager.h"
+#include "Logging/StructuredLog.h"
 #include "OnDemandIoStore.h"
+#include "OnDemandMisc.h"
 
 namespace UE::IoStore
 {
 
-bool IsPackageStreamingEnabled();
+///////////////////////////////////////////////////////////////////////////////
+static void LogUnreferencedPackage(FName PackageName, const FPackageId& PackageId, bool bLogAsWarning)
+{
+#if !UE_BUILD_SHIPPING
+	// Only log when querying the package store from the loader
+	if (!IsInAsyncLoadingThread() && !IsInParallelLoadingThread())
+	{
+		return;
+	}
+
+	TStringBuilder<128> Msg;
+	Msg << TEXT("Trying to load unreferenced package '");
+	if (PackageName.IsNone())
+	{
+		TUtf8StringBuilder<128> Filename;
+		TryConvertChunkIdToPackageName(CreatePackageDataChunkId(PackageId), Filename, Msg);
+	}
+	else
+	{
+		Msg << PackageName.ToString();
+	}
+	Msg << TEXT("' (") << PackageId << TEXT(")");
+
+	UE_CLOGFMT(bLogAsWarning,  LogIoStoreOnDemand, Warning, "{Msg}", Msg);
+	UE_CLOGFMT(!bLogAsWarning, LogIoStoreOnDemand, VeryVerbose, "{Msg}", Msg);
+#endif
+}
 
 ///////////////////////////////////////////////////////////////////////////////
 class FOnDemandPackageStoreBackend final
@@ -64,10 +94,15 @@ public:
 							FPackageId PackageId,
 							FName PackageName,
 							FPackageStoreEntry& OutPackageStoreEntry) override;
-
+	virtual EPackageStoreEntryStatus GetPackageStoreEntryInternal(
+							FPackageId PackageId,
+							FPackageStoreEntry& OutPackageStoreEntry) override;
 	virtual TConstArrayView<uint32> GetSoftReferences(FPackageId PackageId, TConstArrayView<FPackageId>& OutPackageIds) override;
 
 private:
+	const FOnDemandPackageStoreEntry* FindPackageEntry(
+							const FPackageId& PackageId,
+							FPackageStoreEntry& OutPackageStoreEntry) const;
 	void				UpdateLookupTables(const TArray<FSharedOnDemandContainer>& AllContainers, const TArray<TBitArray<>>& ReferencedChunkEntryIndices);
 	void				UpdateReferencedPackages(const TArray<TBitArray<>>& ReferencedChunkEntryIndices);
 
@@ -155,26 +190,35 @@ EPackageStoreEntryStatus FOnDemandPackageStoreBackend::GetPackageStoreEntry(
 	FName PackageName,
 	FPackageStoreEntry& OutPackageStoreEntry)
 {
-	for (const FContainer& Container : Containers)
+	if (const FOnDemandPackageStoreEntry* PackageEntry = FindPackageEntry(PackageId, OutPackageStoreEntry))
 	{
-		if (const FOnDemandPackageStoreEntry* PackageEntry = Container.EntryMap.Find(PackageId))
+		if (PackageEntry->IsReferenced || IsPackageStreamingEnabled())
 		{
-			const FIoContainerHeader& ContainerHeader			= *Container.Container->Header;
-			const FFilePackageStoreEntry* PackageStoreEntries	= reinterpret_cast<const FFilePackageStoreEntry*>(ContainerHeader.StoreEntries.GetData());
-			const FFilePackageStoreEntry& PackageStoreEntry		= *(PackageStoreEntries + PackageEntry->PackageEntryIndex);
-
-			OutPackageStoreEntry.ImportedPackageIds =
-				MakeArrayView(PackageStoreEntry.ImportedPackages.Data(), PackageStoreEntry.ImportedPackages.Num());
-			OutPackageStoreEntry.ShaderMapHashes =
-				MakeArrayView(PackageStoreEntry.ShaderMapHashes.Data(), PackageStoreEntry.ShaderMapHashes.Num());
+			return EPackageStoreEntryStatus::Ok;
+		}
 
-			if (PackageEntry->IsReferenced || IsPackageStreamingEnabled())
-			{
-				return EPackageStoreEntryStatus::Ok;
-			}
+		const bool bLogAsWarning = IsDevModeEnabled();
+		LogUnreferencedPackage(PackageName, PackageId, bLogAsWarning);
 
-			return EPackageStoreEntryStatus::NotInstalled;
+		if (IsDevModeEnabled())
+		{
+			// Return missing to let lower priority backends resolve the package store entry
+			return EPackageStoreEntryStatus::Missing;
 		}
+
+		return EPackageStoreEntryStatus::NotInstalled;
+	}
+	
+	return EPackageStoreEntryStatus::Missing;
+}
+
+EPackageStoreEntryStatus FOnDemandPackageStoreBackend::GetPackageStoreEntryInternal(
+	FPackageId PackageId,
+	FPackageStoreEntry& OutPackageStoreEntry)
+{
+	if (const FOnDemandPackageStoreEntry* PackageEntry = FindPackageEntry(PackageId, OutPackageStoreEntry))
+	{
+		return EPackageStoreEntryStatus::Ok;
 	}
 
 	return EPackageStoreEntryStatus::Missing;
@@ -226,6 +270,30 @@ TConstArrayView<uint32> FOnDemandPackageStoreBackend::GetSoftReferences(FPackage
 	return TConstArrayView<uint32>();
 }
 
+const FOnDemandPackageStoreBackend::FOnDemandPackageStoreEntry* FOnDemandPackageStoreBackend::FindPackageEntry(
+	const FPackageId& PackageId,
+	FPackageStoreEntry& OutPackageStoreEntry) const
+{
+	for (const FContainer& Container : Containers)
+	{
+		if (const FOnDemandPackageStoreEntry* PackageEntry = Container.EntryMap.Find(PackageId))
+		{
+			const FIoContainerHeader& ContainerHeader			= *Container.Container->Header;
+			const FFilePackageStoreEntry* PackageStoreEntries	= reinterpret_cast<const FFilePackageStoreEntry*>(ContainerHeader.StoreEntries.GetData());
+			const FFilePackageStoreEntry& PackageStoreEntry		= *(PackageStoreEntries + PackageEntry->PackageEntryIndex);
+
+			OutPackageStoreEntry.ImportedPackageIds =
+				MakeArrayView(PackageStoreEntry.ImportedPackages.Data(), PackageStoreEntry.ImportedPackages.Num());
+			OutPackageStoreEntry.ShaderMapHashes =
+				MakeArrayView(PackageStoreEntry.ShaderMapHashes.Data(), PackageStoreEntry.ShaderMapHashes.Num());
+			
+			return PackageEntry;
+		}
+	}
+	
+	return nullptr;
+}
+
 void FOnDemandPackageStoreBackend::UpdateLookupTables(const TArray<FSharedOnDemandContainer>& AllContainers, const TArray<TBitArray<>>& ReferencedChunkEntryIndices)
 {
 	check(Mutex.IsLocked());
diff --git a/Engine/Source/Runtime/Experimental/IoStore/OnDemand/Private/OnDemandPackageStoreBackend.h b/Engine/Source/Runtime/Experimental/IoStore/OnDemand/Private/OnDemandPackageStoreBackend.h
--- a/Engine/Source/Runtime/Experimental/IoStore/OnDemand/Private/OnDemandPackageStoreBackend.h
+++ b/Engine/Source/Runtime/Experimental/IoStore/OnDemand/Private/OnDemandPackageStoreBackend.h
@@ -26,6 +26,7 @@ class IOnDemandPackageStoreBackend
 public:
 	virtual ~IOnDemandPackageStoreBackend() = default;
 
+	virtual EPackageStoreEntryStatus GetPackageStoreEntryInternal(FPackageId PackageId, FPackageStoreEntry& Out) = 0;
 	virtual void NeedsUpdate(EOnDemandPackageStoreUpdateMode Mode = EOnDemandPackageStoreUpdateMode::Full) = 0;
 };
 
