Chaos++: Refactored the internal interface classes into separate files to minimize dependencies and make more clear what the public API is.

#Synced-CL 50785456

diff --git a/Engine/Source/Runtime/Experimental/RigidPhysics/Private/RigidPhysics/Internal/IRigidBody.cpp b/Engine/Source/Runtime/Experimental/RigidPhysics/Private/RigidPhysics/Internal/IRigidBody.cpp
new file mode 100644
--- /dev/null
+++ b/Engine/Source/Runtime/Experimental/RigidPhysics/Private/RigidPhysics/Internal/IRigidBody.cpp
@@ -0,0 +1,13 @@
+// Copyright Epic Games, Inc. All Rights Reserved.
+#include "RigidPhysics/Internal/IRigidBody.h"
+
+#include "RigidPhysics/RigidLog.h"
+
+#if UE_RIGIDPHYSICS_API_ENABLED
+
+namespace UE::Physics
+{
+	UE_RIGIDPHYSICS_RIGIDTYPED_IMPL(IRigidBody, IRigidTyped);
+}
+
+#endif // UE_RIGIDPHYSICS_API_ENABLED
diff --git a/Engine/Source/Runtime/Experimental/RigidPhysics/Private/RigidPhysics/Internal/IRigidBodyContainer.cpp b/Engine/Source/Runtime/Experimental/RigidPhysics/Private/RigidPhysics/Internal/IRigidBodyContainer.cpp
new file mode 100644
--- /dev/null
+++ b/Engine/Source/Runtime/Experimental/RigidPhysics/Private/RigidPhysics/Internal/IRigidBodyContainer.cpp
@@ -0,0 +1,27 @@
+// Copyright Epic Games, Inc. All Rights Reserved.
+#include "RigidPhysics/Internal/IRigidBodyContainer.h"
+
+#if UE_RIGIDPHYSICS_API_ENABLED
+
+namespace UE::Physics
+{
+	UE_RIGIDPHYSICS_RIGIDTYPED_IMPL(IRigidBodyContainer, IRigidTyped);
+
+	FRigidBodyContainer::FRigidBodyContainer()
+	{
+		// TODO_CHAOSAPI: Make slack a parameter
+		BodyRegistry = MakeUnique<FRigidBodyRegistry>(10);
+	}
+
+	IRigidBody* FRigidBodyContainer::GetBody(const FRigidObjectId& InBodyId) const
+	{
+		return BodyRegistry->FindObject(InBodyId);
+	}
+
+	FRigidBodyRegistry* FRigidBodyContainer::GetBodyRegistry() const
+	{
+		return BodyRegistry.Get();
+	}
+} // namespace UE::Physics
+
+#endif // UE_RIGIDPHYSICS_API_ENABLED
diff --git a/Engine/Source/Runtime/Experimental/RigidPhysics/Private/RigidPhysics/Internal/IRigidGeometryCollection.cpp b/Engine/Source/Runtime/Experimental/RigidPhysics/Private/RigidPhysics/Internal/IRigidGeometryCollection.cpp
new file mode 100644
--- /dev/null
+++ b/Engine/Source/Runtime/Experimental/RigidPhysics/Private/RigidPhysics/Internal/IRigidGeometryCollection.cpp
@@ -0,0 +1,11 @@
+// Copyright Epic Games, Inc. All Rights Reserved.
+#include "RigidPhysics/Internal/IRigidGeometryCollection.h"
+
+#if UE_RIGIDPHYSICS_API_ENABLED
+
+namespace UE::Physics
+{
+	UE_RIGIDPHYSICS_RIGIDTYPED_IMPL(IRigidGeometryCollection, IRigidTyped);
+} // namespace UE::Physics
+
+#endif // UE_RIGIDPHYSICS_API_ENABLED
diff --git a/Engine/Source/Runtime/Experimental/RigidPhysics/Private/RigidPhysics/Internal/IRigidScene.cpp b/Engine/Source/Runtime/Experimental/RigidPhysics/Private/RigidPhysics/Internal/IRigidScene.cpp
new file mode 100644
--- /dev/null
+++ b/Engine/Source/Runtime/Experimental/RigidPhysics/Private/RigidPhysics/Internal/IRigidScene.cpp
@@ -0,0 +1,40 @@
+// Copyright Epic Games, Inc. All Rights Reserved.
+#include "RigidPhysics/Internal/IRigidScene.h"
+
+#include "RigidPhysics/RigidBodyContainer.h"
+#include "RigidPhysics/RigidGeometryCollection.h"
+#include "RigidPhysics/RigidLog.h"
+
+#if UE_RIGIDPHYSICS_API_ENABLED
+
+namespace UE::Physics
+{
+	UE_RIGIDPHYSICS_RIGIDTYPED_IMPL(IRigidScene, IRigidLockable);
+
+	FRigidScene::FRigidScene()
+	{
+		// TODO_CHAOSAPI: registry slack should be a parameter
+		BodyContainerRegistry = MakeUnique<FRigidBodyContainerRegistry>(10);
+	}
+
+	IRigidBodyContainer* FRigidScene::GetBodyContainer(const FRigidObjectId& InBodyContainerId) const
+	{
+		return BodyContainerRegistry->FindObject(InBodyContainerId);
+	}
+
+	IRigidGeometryCollection* FRigidScene::GetGeometryCollection(const FRigidObjectId& InGeometryCollectionId) const
+	{
+		if (IRigidBodyContainer* BodyContainer = BodyContainerRegistry->FindObject(InGeometryCollectionId))
+		{
+			return BodyContainer->AsA<IRigidGeometryCollection>();
+		}
+		return nullptr;
+	}
+
+	FRigidBodyContainerRegistry* FRigidScene::GetBodyContainerRegistry() const
+	{
+		return BodyContainerRegistry.Get();
+	}
+} // namespace UE::Physics
+
+#endif // UE_RIGIDPHYSICS_API_ENABLED
diff --git a/Engine/Source/Runtime/Experimental/RigidPhysics/Private/RigidPhysics/RigidBody.cpp b/Engine/Source/Runtime/Experimental/RigidPhysics/Private/RigidPhysics/RigidBody.cpp
deleted file mode 100644
diff --git a/Engine/Source/Runtime/Experimental/RigidPhysics/Private/RigidPhysics/RigidBodyContainer.cpp b/Engine/Source/Runtime/Experimental/RigidPhysics/Private/RigidPhysics/RigidBodyContainer.cpp
deleted file mode 100644
diff --git a/Engine/Source/Runtime/Experimental/RigidPhysics/Private/RigidPhysics/RigidGeometryCollection.cpp b/Engine/Source/Runtime/Experimental/RigidPhysics/Private/RigidPhysics/RigidGeometryCollection.cpp
deleted file mode 100644
diff --git a/Engine/Source/Runtime/Experimental/RigidPhysics/Private/RigidPhysics/RigidScene.cpp b/Engine/Source/Runtime/Experimental/RigidPhysics/Private/RigidPhysics/RigidScene.cpp
deleted file mode 100644
diff --git a/Engine/Source/Runtime/Experimental/RigidPhysics/Private/RigidPhysics/RigidSceneSettings.cpp b/Engine/Source/Runtime/Experimental/RigidPhysics/Private/RigidPhysics/RigidSceneSettings.cpp
new file mode 100644
--- /dev/null
+++ b/Engine/Source/Runtime/Experimental/RigidPhysics/Private/RigidPhysics/RigidSceneSettings.cpp
@@ -0,0 +1,11 @@
+// Copyright Epic Games, Inc. All Rights Reserved.
+#include "RigidPhysics/RigidSceneSettings.h"
+
+#if UE_RIGIDPHYSICS_API_ENABLED
+
+namespace UE::Physics
+{
+	UE_RIGIDPHYSICS_RIGIDTYPED_IMPL(IRigidSceneSettings, IRigidTyped);
+} // namespace UE::Physics
+
+#endif // UE_RIGIDPHYSICS_API_ENABLED
diff --git a/Engine/Source/Runtime/Experimental/RigidPhysics/Public/RigidPhysics/Internal/IRigidBody.h b/Engine/Source/Runtime/Experimental/RigidPhysics/Public/RigidPhysics/Internal/IRigidBody.h
new file mode 100644
--- /dev/null
+++ b/Engine/Source/Runtime/Experimental/RigidPhysics/Public/RigidPhysics/Internal/IRigidBody.h
@@ -0,0 +1,218 @@
+// Copyright Epic Games, Inc. All Rights Reserved.
+#pragma once
+
+#include "RigidPhysics/RigidFwd.h"
+
+#if UE_RIGIDPHYSICS_API_ENABLED
+
+#include "Math/Bounds.h"
+#include "Math/Transform.h"
+#include "RigidPhysics/RigidBodyHandle.h"
+#include "RigidPhysics/RigidCollider.h"
+#include "RigidPhysics/RigidLog.h"
+#include "RigidPhysics/RigidTyped.h"
+
+namespace UE::Physics
+{
+	// The full rigid body API. There will be several implementations of this API. E.g., synchronous
+	// physics, asynchronous physics, immediate physics, remote physics, etc.
+	// NOTE: for some implementations like async physics, there are both Game- and Physics-Context
+	// implementations within the async implementation. Pinning a handle will return the object 
+	// appropriate for the context.
+	class UE_INTERNAL IRigidBody : public IRigidTyped
+	{
+	public:
+		UE_RIGIDPHYSICS_RIGIDTYPED_DECL(RIGIDPHYSICS_API, IRigidBody);
+
+		IRigidBody() = default;
+		virtual ~IRigidBody() = default;
+
+		// Is the body active in its scene?
+		// Inactive bodies do not participate in physics in any way.
+		UE_INTERNAL RIGIDPHYSICS_API virtual bool IsActive() const
+		{
+			RIGIDPHYSICS_API_UNSUPPORTED();
+			return false;
+		}
+
+		UE_INTERNAL RIGIDPHYSICS_API virtual void Activate()
+		{
+			RIGIDPHYSICS_API_UNSUPPORTED();
+		}
+
+		UE_INTERNAL RIGIDPHYSICS_API virtual void Deactivate()
+		{
+			RIGIDPHYSICS_API_UNSUPPORTED();
+		}
+
+
+		// The handle to this body
+		UE_INTERNAL RIGIDPHYSICS_API virtual FRigidBodyHandle GetHandle() const = 0;
+
+		// Is this body Static? Static bodies can be moved but have no velocity.
+		UE_INTERNAL RIGIDPHYSICS_API virtual bool IsStatic() const = 0;
+
+		// Is this body kinematic? Kinematic bodies can be explicitly moved and have velocity, but do
+		// not react to forces, collisions, etc.
+		UE_INTERNAL RIGIDPHYSICS_API virtual bool IsKinematic() const = 0;
+
+		// Is this a dynamic body? Dynamic bodies have momentum and react to forces and collisions.
+		UE_INTERNAL RIGIDPHYSICS_API virtual bool IsDynamic() const = 0;
+
+		// Make the body dynamic (will only affect Kinematic bodies. Static bodies cannot be made dynamic)
+		UE_INTERNAL RIGIDPHYSICS_API virtual void SetIsDynamic(bool bIsDynamic) const
+		{
+			RIGIDPHYSICS_API_UNSUPPORTED();
+		}
+
+		// Is this a sleeping dynamic body? Sleeping bodies are dynamic bodies that are ignored by
+		// the simulation (as an optimization) until the scene changes in a way that might affect them, 
+		// or they are explicitly woken.
+		UE_INTERNAL RIGIDPHYSICS_API virtual bool IsSleeping() const
+		{
+			RIGIDPHYSICS_API_UNSUPPORTED();
+			return false;
+		}
+
+		UE_INTERNAL RIGIDPHYSICS_API virtual int32 GetNumShapes() const = 0;
+		UE_INTERNAL RIGIDPHYSICS_API virtual FConstRigidShapeInstancePtr GetShape(int32 ShapeIndex) const = 0;
+
+		// Add a Collision Shape
+		UE_INTERNAL RIGIDPHYSICS_API virtual void AddShape(const FConstRigidShapeInstancePtr& InShape)
+		{
+			RIGIDPHYSICS_API_UNSUPPORTED();
+		}
+
+		// Set the transform and initialize body state (velocity etc.) at the new location.
+		UE_INTERNAL RIGIDPHYSICS_API virtual void InitTransform(const FTransform3d& InTransform)
+		{
+			RIGIDPHYSICS_API_UNSUPPORTED();
+		}
+
+		// Set the transform without changing any other rigid body state.
+		UE_INTERNAL RIGIDPHYSICS_API virtual void UpdateTransform(const FTransform3d& InTransform)
+		{
+			RIGIDPHYSICS_API_UNSUPPORTED();
+		}
+
+		// Set the kinematic target. The body will attempt to reach the target at the next
+		// tick. Kinematic bodies will always reach the target bscause they do not respond
+		// to forces or collisions. Dynamic bodies will only reach their target if their
+		// velocity in the absence of interactions (forces, collisions, or other constraints).
+		UE_INTERNAL RIGIDPHYSICS_API virtual void SetKinematicTarget(const FTransform3d& InTransform)
+		{
+			RIGIDPHYSICS_API_UNSUPPORTED();
+		}
+
+		// Get the current transform. NOTE: This will immediately reflect any recent calls to
+		// to InitTransform or UpdateTransform, but will not reflect the latest KinematicTarget
+		// until after the next physics tick.
+		UE_INTERNAL RIGIDPHYSICS_API virtual FTransform3d GetTransform() const
+		{
+			RIGIDPHYSICS_API_UNSUPPORTED();
+			return FTransform3d::Identity;
+		}
+
+		// Get the world-space bounds at the current transform
+		UE_INTERNAL RIGIDPHYSICS_API virtual FBounds3d GetBounds() const
+		{
+			RIGIDPHYSICS_API_UNSUPPORTED();
+			return FBounds3d(FVector3d::Zero(), FVector3d::Zero());
+		}
+
+		// TODO: Maybe just Set/GetMassProperties
+		UE_INTERNAL RIGIDPHYSICS_API virtual double GetMass() const
+		{
+			RIGIDPHYSICS_API_UNSUPPORTED();
+			return 0;
+		}
+
+		UE_INTERNAL RIGIDPHYSICS_API virtual void SetMass(double InM)
+		{
+			RIGIDPHYSICS_API_UNSUPPORTED();
+		}
+
+		UE_INTERNAL RIGIDPHYSICS_API virtual FVector3d GetInertia() const
+		{
+			RIGIDPHYSICS_API_UNSUPPORTED();
+			return FVector3d::Zero();
+		}
+
+		UE_INTERNAL RIGIDPHYSICS_API virtual void SetInertia(const FVector3d& InInertia)
+		{
+			RIGIDPHYSICS_API_UNSUPPORTED();
+		}
+
+		UE_INTERNAL RIGIDPHYSICS_API virtual void SetCCDEnabled(bool bInEnabled)
+		{
+			RIGIDPHYSICS_API_UNSUPPORTED();
+		}
+
+		UE_INTERNAL RIGIDPHYSICS_API virtual void SetMACDEnabled(bool bInEnabled)
+		{
+			RIGIDPHYSICS_API_UNSUPPORTED();
+		}
+
+		UE_INTERNAL RIGIDPHYSICS_API virtual void SetSmoothEdgeCollisionsEnabled(bool bInEnabled)
+		{
+			RIGIDPHYSICS_API_UNSUPPORTED();
+		}
+
+		UE_INTERNAL RIGIDPHYSICS_API virtual void SetMaxLinearVelocity(float InMax)
+		{
+			RIGIDPHYSICS_API_UNSUPPORTED();
+		}
+
+		UE_INTERNAL RIGIDPHYSICS_API virtual void SetMaxAngularVelocity(float InMax)
+		{
+			RIGIDPHYSICS_API_UNSUPPORTED();
+		}
+
+		UE_INTERNAL RIGIDPHYSICS_API virtual void SetInertiaConditioningEnabled(bool bInEnabled)
+		{
+			RIGIDPHYSICS_API_UNSUPPORTED();
+		}
+
+		UE_INTERNAL RIGIDPHYSICS_API virtual void SetGyroscopicTorqueEnabled(bool bInEnabled)
+		{
+			RIGIDPHYSICS_API_UNSUPPORTED();
+		}
+
+		UE_INTERNAL RIGIDPHYSICS_API virtual void SetGravityGroup(int32 InId)
+		{
+			RIGIDPHYSICS_API_UNSUPPORTED();
+		}
+
+		UE_INTERNAL RIGIDPHYSICS_API virtual void SetGravityScale(float InScale)
+		{
+			RIGIDPHYSICS_API_UNSUPPORTED();
+		}
+
+		UE_INTERNAL RIGIDPHYSICS_API virtual void SetWakeEventsEnabled(bool bInEnabled)
+		{
+			RIGIDPHYSICS_API_UNSUPPORTED();
+		}
+
+		UE_INTERNAL RIGIDPHYSICS_API virtual void SetNumPositionIterations(int32 InNum)
+		{
+			RIGIDPHYSICS_API_UNSUPPORTED();
+		}
+
+		UE_INTERNAL RIGIDPHYSICS_API virtual void SetNumVelocityIterations(int32 InNum)
+		{
+			RIGIDPHYSICS_API_UNSUPPORTED();
+		}
+
+		UE_INTERNAL RIGIDPHYSICS_API virtual void SetNumProjectionIterations(int32 InNum)
+		{
+			RIGIDPHYSICS_API_UNSUPPORTED();
+		}
+
+		UE_INTERNAL RIGIDPHYSICS_API virtual void SetIsSleeping(bool bInSleeping)
+		{
+			RIGIDPHYSICS_API_UNSUPPORTED();
+		}
+	};
+} // namespace UE::Physics
+
+#endif // UE_RIGIDPHYSICS_API_ENABLED
diff --git a/Engine/Source/Runtime/Experimental/RigidPhysics/Public/RigidPhysics/Internal/IRigidBodyContainer.h b/Engine/Source/Runtime/Experimental/RigidPhysics/Public/RigidPhysics/Internal/IRigidBodyContainer.h
new file mode 100644
--- /dev/null
+++ b/Engine/Source/Runtime/Experimental/RigidPhysics/Public/RigidPhysics/Internal/IRigidBodyContainer.h
@@ -0,0 +1,69 @@
+// Copyright Epic Games, Inc. All Rights Reserved.
+#pragma once
+
+#include "RigidPhysics/RigidFwd.h"
+
+#if UE_RIGIDPHYSICS_API_ENABLED
+
+#include "RigidPhysics/RigidBodyContainerHandle.h"
+#include "RigidPhysics/RigidLog.h"
+#include "RigidPhysics/RigidObjectId.h"
+#include "RigidPhysics/RigidObjectRegistry.h"
+#include "RigidPhysics/RigidTyped.h"
+
+namespace UE::Physics
+{
+	// The interface for rigid body containers. There will be several implementations of this API. E.g., synchronous
+	// physics, asynchronous physics, immediate physics, remote physics, etc. 
+	class UE_INTERNAL IRigidBodyContainer : public IRigidTyped
+	{
+	public:
+		UE_RIGIDPHYSICS_RIGIDTYPED_DECL(RIGIDPHYSICS_API, IRigidBodyContainer);
+
+		IRigidBodyContainer() = default;
+		virtual ~IRigidBodyContainer() = default;
+
+		UE_INTERNAL RIGIDPHYSICS_API virtual FRigidBodyContainerHandle GetHandle() const = 0;
+
+		UE_INTERNAL RIGIDPHYSICS_API virtual int32 GetNumBodies() const = 0;
+
+		UE_INTERNAL RIGIDPHYSICS_API virtual IRigidBody* GetBody(const FRigidObjectId& InBodyId) const = 0;
+
+		//
+		// Below here is GT only. See if we can separate the Game and SimCallback APIs
+		//
+
+		UE_INTERNAL RIGIDPHYSICS_API virtual void SetId(const FRigidObjectId& InBodyId)
+		{
+			RIGIDPHYSICS_API_UNSUPPORTED();
+		}
+
+		UE_INTERNAL RIGIDPHYSICS_API virtual IRigidBody* CreateBody(const FRigidDebugName& InName, ERigidMovementType InMovementType)
+		{
+			RIGIDPHYSICS_API_UNSUPPORTED();
+			return nullptr;
+		}
+
+		UE_INTERNAL RIGIDPHYSICS_API virtual void DestroyBody(IRigidBody* InBody)
+		{
+			RIGIDPHYSICS_API_UNSUPPORTED();
+		}
+	};
+
+	class UE_INTERNAL FRigidBodyContainer : public IRigidBodyContainer
+	{
+	public:
+		UE_INTERNAL RIGIDPHYSICS_API FRigidBodyContainer();
+		virtual ~FRigidBodyContainer() = default;
+
+		UE_INTERNAL RIGIDPHYSICS_API virtual IRigidBody* GetBody(const FRigidObjectId& InBodyId) const override final;
+
+	protected:
+		UE_INTERNAL RIGIDPHYSICS_API FRigidBodyRegistry* GetBodyRegistry() const;
+
+	private:
+		TUniquePtr<FRigidBodyRegistry> BodyRegistry;
+	};
+} // namespace UE::Physics
+
+#endif // UE_RIGIDPHYSICS_API_ENABLED
diff --git a/Engine/Source/Runtime/Experimental/RigidPhysics/Public/RigidPhysics/Internal/IRigidGeometryCollection.h b/Engine/Source/Runtime/Experimental/RigidPhysics/Public/RigidPhysics/Internal/IRigidGeometryCollection.h
new file mode 100644
--- /dev/null
+++ b/Engine/Source/Runtime/Experimental/RigidPhysics/Public/RigidPhysics/Internal/IRigidGeometryCollection.h
@@ -0,0 +1,28 @@
+// Copyright Epic Games, Inc. All Rights Reserved.
+#pragma once
+
+#include "RigidPhysics/RigidFwd.h"
+
+#if UE_RIGIDPHYSICS_API_ENABLED
+
+#include "RigidPhysics/Internal/IRigidBodyContainer.h"
+#include "RigidPhysics/RigidGeometryCollectionHandle.h"
+
+namespace UE::Physics
+{
+	// The interface for rigid body containers. There will be several implementations of this API. E.g., synchronous
+	// physics, asynchronous physics, immediate physics, remote physics, etc. 
+	class UE_INTERNAL IRigidGeometryCollection : public IRigidBodyContainer
+	{
+	public:
+		UE_RIGIDPHYSICS_RIGIDTYPED_DECL(RIGIDPHYSICS_API, IRigidGeometryCollection);
+
+		IRigidGeometryCollection() = default;
+		virtual ~IRigidGeometryCollection() = default;
+
+		// TODO_CHAOSAPI: Figure out the public GeometryCollection API!!
+		UE_INTERNAL RIGIDPHYSICS_API virtual IRigidBody* GetBodyAt(int32 BodyIndex) const = 0;
+	};
+} // namespace UE::Physics
+
+#endif // UE_RIGIDPHYSICS_API_ENABLED
diff --git a/Engine/Source/Runtime/Experimental/RigidPhysics/Public/RigidPhysics/Internal/IRigidScene.h b/Engine/Source/Runtime/Experimental/RigidPhysics/Public/RigidPhysics/Internal/IRigidScene.h
new file mode 100644
--- /dev/null
+++ b/Engine/Source/Runtime/Experimental/RigidPhysics/Public/RigidPhysics/Internal/IRigidScene.h
@@ -0,0 +1,99 @@
+// Copyright Epic Games, Inc. All Rights Reserved.
+#pragma once
+
+#include "RigidPhysics/RigidFwd.h"
+
+#if UE_RIGIDPHYSICS_API_ENABLED
+
+#include "RigidPhysics/RigidLockable.h"
+#include "RigidPhysics/RigidLog.h"
+#include "RigidPhysics/RigidObjectId.h"
+#include "RigidPhysics/RigidSceneHandle.h"
+#include "RigidPhysics/RigidTyped.h"
+
+namespace UE::Physics
+{
+	// The interface to a rigid body scene. The rigid body scene owns the physics simulations and
+	// manages a set of rigid body containers which in turn manage sets of rigid bodies.
+	class UE_INTERNAL IRigidScene : public IRigidLockable
+	{
+	public:
+		UE_RIGIDPHYSICS_RIGIDTYPED_DECL(RIGIDPHYSICS_API, IRigidScene);
+
+		IRigidScene() = default;
+		virtual ~IRigidScene() = default;
+
+		UE_INTERNAL RIGIDPHYSICS_API virtual FRigidSceneHandle GetHandle() const = 0;
+
+		UE_INTERNAL RIGIDPHYSICS_API virtual IRigidBodyContainer* GetBodyContainer(const FRigidObjectId& InBodyContainerId) const = 0;
+
+		UE_INTERNAL RIGIDPHYSICS_API virtual IRigidGeometryCollection* GetGeometryCollection(const FRigidObjectId& InGeometryCollectionId) const = 0;
+
+		UE_INTERNAL RIGIDPHYSICS_API virtual bool Query(const FRigidQueryDesc& InQueryDesc, TArray<FRigidQueryHit>& OutHits)
+		{
+			RIGIDPHYSICS_API_UNSUPPORTED();
+			return false;
+		}
+
+		// Advance the simulation by Dt. If the implementation is asynchronous, this will return immediately
+		UE_INTERNAL RIGIDPHYSICS_API virtual void Tick(float InDt)
+		{
+			RIGIDPHYSICS_API_UNSUPPORTED();
+		}
+
+		// GT Only API - Move this somehow
+		UE_INTERNAL RIGIDPHYSICS_API virtual void RegisterModifier(IRigidSceneModifier* InModifier)
+		{
+			RIGIDPHYSICS_API_UNSUPPORTED();
+		}
+
+		UE_INTERNAL RIGIDPHYSICS_API virtual void UnregisterModifier(IRigidSceneModifier* InModifier)
+		{
+			RIGIDPHYSICS_API_UNSUPPORTED();
+		}
+
+		UE_INTERNAL RIGIDPHYSICS_API virtual IRigidBody* CreateBody(const FRigidDebugName& InName, ERigidMovementType InMovementType)
+		{
+			RIGIDPHYSICS_API_UNSUPPORTED();
+			return nullptr;
+		}
+
+		UE_INTERNAL RIGIDPHYSICS_API virtual void DestroyBody(IRigidBody* InBody)
+		{
+			RIGIDPHYSICS_API_UNSUPPORTED();
+		}
+
+		UE_INTERNAL RIGIDPHYSICS_API virtual IRigidGeometryCollection* CreateGeometryCollection(const FRigidDebugName& InName, const FSimulationParameters& InParameters)
+		{
+			RIGIDPHYSICS_API_UNSUPPORTED();
+			return nullptr;
+		}
+
+		UE_INTERNAL RIGIDPHYSICS_API virtual void DestroyGeometryCollection(IRigidGeometryCollection* InGC)
+		{
+			RIGIDPHYSICS_API_UNSUPPORTED();
+		}
+	};
+
+	//
+	// Base class for the Game Context Scene. This provides the BodyContainerRegistry, partly for convenience
+	// and partly so that we can convert Handles to Objects in the debugger with natvis
+	//
+	class UE_INTERNAL FRigidScene : public IRigidScene
+	{
+	public:
+		RIGIDPHYSICS_API FRigidScene();
+		virtual ~FRigidScene() = default;
+
+		UE_INTERNAL RIGIDPHYSICS_API virtual IRigidBodyContainer* GetBodyContainer(const FRigidObjectId& InBodyContainerId) const override final;
+		UE_INTERNAL RIGIDPHYSICS_API virtual IRigidGeometryCollection* GetGeometryCollection(const FRigidObjectId& InGeometryCollectionId) const override final;
+
+	protected:
+		UE_INTERNAL RIGIDPHYSICS_API FRigidBodyContainerRegistry* GetBodyContainerRegistry() const;
+
+	private:
+		TUniquePtr<FRigidBodyContainerRegistry> BodyContainerRegistry;
+	};
+} // namespace UE::Physics
+
+#endif // UE_RIGIDPHYSICS_API_ENABLED
diff --git a/Engine/Source/Runtime/Experimental/RigidPhysics/Public/RigidPhysics/RigidBody.h b/Engine/Source/Runtime/Experimental/RigidPhysics/Public/RigidPhysics/RigidBody.h
--- a/Engine/Source/Runtime/Experimental/RigidPhysics/Public/RigidPhysics/RigidBody.h
+++ b/Engine/Source/Runtime/Experimental/RigidPhysics/Public/RigidPhysics/RigidBody.h
@@ -7,223 +7,15 @@
 
 #include "Math/Bounds.h"
 #include "Math/Transform.h"
+#include "RigidPhysics/Internal/IRigidBody.h"
 #include "RigidPhysics/RigidBodyContainer.h"
 #include "RigidPhysics/RigidBodyHandle.h"
 #include "RigidPhysics/RigidCollider.h"
-#include "RigidPhysics/RigidLog.h"
 #include "RigidPhysics/RigidObjectId.h"
 #include "RigidPhysics/RigidObjectPtr.h"
-#include "RigidPhysics/RigidScene.h"
-#include "RigidPhysics/RigidTyped.h"
 
 namespace UE::Physics
 {
-	// The full rigid body API. There will be several implementations of this API. E.g., synchronous
-	// physics, asynchronous physics, immediate physics, remote physics, etc.
-	// NOTE: for some implementations like async physics, there are both Game- and Physics-Context
-	// implementations within the async implementation. Pinning a handle will return the object 
-	// appropriate for the context.
-	class UE_INTERNAL IRigidBody: public IRigidTyped
-	{
-	public:
-		UE_RIGIDPHYSICS_RIGIDTYPED_DECL(RIGIDPHYSICS_API, IRigidBody);
-
-		IRigidBody()
-		{
-		}
-
-		virtual ~IRigidBody()
-		{
-		}
-
-		// Is the body active in its scene?
-		// Inactive bodies do not participate in physics in any way.
-		UE_INTERNAL RIGIDPHYSICS_API virtual bool IsActive() const
-		{
-			RIGIDPHYSICS_API_UNSUPPORTED();
-			return false;
-		}
-
-		UE_INTERNAL RIGIDPHYSICS_API virtual void Activate()
-		{
-			RIGIDPHYSICS_API_UNSUPPORTED();
-		}
-
-		UE_INTERNAL RIGIDPHYSICS_API virtual void Deactivate()
-		{
-			RIGIDPHYSICS_API_UNSUPPORTED();
-		}
-
-
-		// The handle to this body
-		UE_INTERNAL RIGIDPHYSICS_API virtual FRigidBodyHandle GetHandle() const = 0;
-
-		// Is this body Static? Static bodies can be moved but have no velocity.
-		UE_INTERNAL RIGIDPHYSICS_API virtual bool IsStatic() const = 0;
-
-		// Is this body kinematic? Kinematic bodies can be explicitly moved and have velocity, but do
-		// not react to forces, collisions, etc.
-		UE_INTERNAL RIGIDPHYSICS_API virtual bool IsKinematic() const = 0;
-
-		// Is this a dynamic body? Dynamic bodies have momentum and react to forces and collisions.
-		UE_INTERNAL RIGIDPHYSICS_API virtual bool IsDynamic() const = 0;
-
-		// Make the body dynamic (will only affect Kinematic bodies. Static bodies cannot be made dynamic)
-		UE_INTERNAL RIGIDPHYSICS_API virtual void SetIsDynamic(bool bIsDynamic) const
-		{
-			RIGIDPHYSICS_API_UNSUPPORTED();
-		}
-
-		// Is this a sleeping dynamic body? Sleeping bodies are dynamic bodies that are ignored by
-		// the simulation (as an optimization) until the scene changes in a way that might affect them, 
-		// or they are explicitly woken.
-		UE_INTERNAL RIGIDPHYSICS_API virtual bool IsSleeping() const
-		{
-			RIGIDPHYSICS_API_UNSUPPORTED();
-			return false;
-		}
-
-		UE_INTERNAL RIGIDPHYSICS_API virtual int32 GetNumShapes() const = 0;
-		UE_INTERNAL RIGIDPHYSICS_API virtual FConstRigidShapeInstancePtr GetShape(int32 ShapeIndex) const = 0;
-
-		// Add a Collision Shape
-		UE_INTERNAL RIGIDPHYSICS_API virtual void AddShape(const FConstRigidShapeInstancePtr& InShape)
-		{
-			RIGIDPHYSICS_API_UNSUPPORTED();
-		}
-
-		// Set the transform and initialize body state (velocity etc.) at the new location.
-		UE_INTERNAL RIGIDPHYSICS_API virtual void InitTransform(const FTransform3d& InTransform)
-		{
-			RIGIDPHYSICS_API_UNSUPPORTED();
-		}
-
-		// Set the transform without changing any other rigid body state.
-		UE_INTERNAL RIGIDPHYSICS_API virtual void UpdateTransform(const FTransform3d& InTransform)
-		{
-			RIGIDPHYSICS_API_UNSUPPORTED();
-		}
-
-		// Set the kinematic target. The body will attempt to reach the target at the next
-		// tick. Kinematic bodies will always reach the target bscause they do not respond
-		// to forces or collisions. Dynamic bodies will only reach their target if their
-		// velocity in the absence of interactions (forces, collisions, or other constraints).
-		UE_INTERNAL RIGIDPHYSICS_API virtual void SetKinematicTarget(const FTransform3d& InTransform)
-		{
-			RIGIDPHYSICS_API_UNSUPPORTED();
-		}
-
-		// Get the current transform. NOTE: This will immediately reflect any recent calls to
-		// to InitTransform or UpdateTransform, but will not reflect the latest KinematicTarget
-		// until after the next physics tick.
-		UE_INTERNAL RIGIDPHYSICS_API virtual FTransform3d GetTransform() const
-		{
-			RIGIDPHYSICS_API_UNSUPPORTED();
-			return FTransform3d::Identity;
-		}
-
-		// Get the world-space bounds at the current transform
-		UE_INTERNAL RIGIDPHYSICS_API virtual FBounds3d GetBounds() const
-		{
-			RIGIDPHYSICS_API_UNSUPPORTED();
-			return FBounds3d(FVector3d::Zero(), FVector3d::Zero());
-		}
-
-		// TODO: Maybe just Set/GetMassProperties
-		UE_INTERNAL RIGIDPHYSICS_API virtual double GetMass() const
-		{
-			RIGIDPHYSICS_API_UNSUPPORTED();
-			return 0;
-		}
-
-		UE_INTERNAL RIGIDPHYSICS_API virtual void SetMass(double InM)
-		{
-			RIGIDPHYSICS_API_UNSUPPORTED();
-		}
-
-		UE_INTERNAL RIGIDPHYSICS_API virtual FVector3d GetInertia() const
-		{
-			RIGIDPHYSICS_API_UNSUPPORTED();
-			return FVector3d::Zero();
-		}
-
-		UE_INTERNAL RIGIDPHYSICS_API virtual void SetInertia(const FVector3d& InInertia)
-		{
-			RIGIDPHYSICS_API_UNSUPPORTED();
-		}
-
-		UE_INTERNAL RIGIDPHYSICS_API virtual void SetCCDEnabled(bool bInEnabled)
-		{
-			RIGIDPHYSICS_API_UNSUPPORTED();
-		}
-
-		UE_INTERNAL RIGIDPHYSICS_API virtual void SetMACDEnabled(bool bInEnabled)
-		{
-			RIGIDPHYSICS_API_UNSUPPORTED();
-		}
-		
-		UE_INTERNAL RIGIDPHYSICS_API virtual void SetSmoothEdgeCollisionsEnabled(bool bInEnabled)
-		{
-			RIGIDPHYSICS_API_UNSUPPORTED();
-		}
-
-		UE_INTERNAL RIGIDPHYSICS_API virtual void SetMaxLinearVelocity(float InMax)
-		{
-			RIGIDPHYSICS_API_UNSUPPORTED();
-		}
-		
-		UE_INTERNAL RIGIDPHYSICS_API virtual void SetMaxAngularVelocity(float InMax)
-		{
-			RIGIDPHYSICS_API_UNSUPPORTED();
-		}
-		
-		UE_INTERNAL RIGIDPHYSICS_API virtual void SetInertiaConditioningEnabled(bool bInEnabled)
-		{
-			RIGIDPHYSICS_API_UNSUPPORTED();
-		}
-		
-		UE_INTERNAL RIGIDPHYSICS_API virtual void SetGyroscopicTorqueEnabled(bool bInEnabled)
-		{
-			RIGIDPHYSICS_API_UNSUPPORTED();
-		}
-
-		UE_INTERNAL RIGIDPHYSICS_API virtual void SetGravityGroup(int32 InId)
-		{
-			RIGIDPHYSICS_API_UNSUPPORTED();
-		}
-		
-		UE_INTERNAL RIGIDPHYSICS_API virtual void SetGravityScale(float InScale)
-		{
-			RIGIDPHYSICS_API_UNSUPPORTED();
-		}
-
-		UE_INTERNAL RIGIDPHYSICS_API virtual void SetWakeEventsEnabled(bool bInEnabled)
-		{
-			RIGIDPHYSICS_API_UNSUPPORTED();
-		}
-		
-		UE_INTERNAL RIGIDPHYSICS_API virtual void SetNumPositionIterations(int32 InNum)
-		{
-			RIGIDPHYSICS_API_UNSUPPORTED();
-		}
-		
-		UE_INTERNAL RIGIDPHYSICS_API virtual void SetNumVelocityIterations(int32 InNum)
-		{
-			RIGIDPHYSICS_API_UNSUPPORTED();
-		}
-		
-		UE_INTERNAL RIGIDPHYSICS_API virtual void SetNumProjectionIterations(int32 InNum)
-		{
-			RIGIDPHYSICS_API_UNSUPPORTED();
-		}
-		
-		UE_INTERNAL RIGIDPHYSICS_API virtual void SetIsSleeping(bool bInSleeping)
-		{
-			RIGIDPHYSICS_API_UNSUPPORTED();
-		}
-	};
-
-
 	// This class only exists so that TRigidBodyPtr can have pointer semantics rather than object semantics,
 	// and so that we can disable functions based on ContextType
 	// See TRigidBodyPtr for comments.
diff --git a/Engine/Source/Runtime/Experimental/RigidPhysics/Public/RigidPhysics/RigidBodyContainer.h b/Engine/Source/Runtime/Experimental/RigidPhysics/Public/RigidPhysics/RigidBodyContainer.h
--- a/Engine/Source/Runtime/Experimental/RigidPhysics/Public/RigidPhysics/RigidBodyContainer.h
+++ b/Engine/Source/Runtime/Experimental/RigidPhysics/Public/RigidPhysics/RigidBodyContainer.h
@@ -5,77 +5,16 @@
 
 #if UE_RIGIDPHYSICS_API_ENABLED
 
+#include "RigidPhysics/Internal/IRigidBodyContainer.h"
 #include "RigidPhysics/RigidBodyContainerHandle.h"
 #include "RigidPhysics/RigidLog.h"
 #include "RigidPhysics/RigidObjectId.h"
 #include "RigidPhysics/RigidObjectPtr.h"
 #include "RigidPhysics/RigidObjectRegistry.h"
-#include "RigidPhysics/RigidScene.h"
 #include "RigidPhysics/RigidTyped.h"
 
 namespace UE::Physics
 {
-	// The interface for rigid body containers. There will be several implementations of this API. E.g., synchronous
-	// physics, asynchronous physics, immediate physics, remote physics, etc. 
-	class UE_INTERNAL IRigidBodyContainer : public IRigidTyped
-	{
-	public:
-		UE_RIGIDPHYSICS_RIGIDTYPED_DECL(RIGIDPHYSICS_API, IRigidBodyContainer);
-
-		UE_INTERNAL IRigidBodyContainer()
-		{
-		}
-
-		UE_INTERNAL virtual ~IRigidBodyContainer()
-		{
-		}
-
-		UE_INTERNAL RIGIDPHYSICS_API virtual FRigidBodyContainerHandle GetHandle() const = 0;
-
-		UE_INTERNAL RIGIDPHYSICS_API virtual int32 GetNumBodies() const = 0;
-
-		UE_INTERNAL RIGIDPHYSICS_API virtual IRigidBody* GetBody(const FRigidObjectId& InBodyId) const = 0;
-
-		//
-		// Below here is GT only. See if we can separate the Game and SimCallback APIs
-		//
-
-		UE_INTERNAL RIGIDPHYSICS_API virtual void SetId(const FRigidObjectId& InBodyId)
-		{
-			RIGIDPHYSICS_API_UNSUPPORTED();
-		}
-
-		UE_INTERNAL RIGIDPHYSICS_API virtual IRigidBody* CreateBody(const FRigidDebugName& InName, ERigidMovementType InMovementType)
-		{
-			RIGIDPHYSICS_API_UNSUPPORTED();
-			return nullptr;
-		}
-
-		UE_INTERNAL RIGIDPHYSICS_API virtual void DestroyBody(IRigidBody* InBody)
-		{
-			RIGIDPHYSICS_API_UNSUPPORTED();
-		}
-
-	};
-
-	class UE_INTERNAL FRigidBodyContainer : public IRigidBodyContainer
-	{
-	public:
-		UE_INTERNAL RIGIDPHYSICS_API FRigidBodyContainer();
-		UE_INTERNAL RIGIDPHYSICS_API virtual ~FRigidBodyContainer();
-
-		UE_INTERNAL RIGIDPHYSICS_API virtual IRigidBody* GetBody(const FRigidObjectId& InBodyId) const override final;
-
-	protected:
-		FRigidBodyRegistry* GetBodyRegistry() const
-		{
-			return BodyRegistry.Get();
-		}
-
-	private:
-		TUniquePtr<FRigidBodyRegistry> BodyRegistry;
-	};
-
 	// Utility wrapper around IRigidBodyContainer that adds Context and allows
 	// TRigidBodyContainerPtr to have pointer semantics without exposing IRigidBodyContainer
 	template<typename ContextType>
diff --git a/Engine/Source/Runtime/Experimental/RigidPhysics/Public/RigidPhysics/RigidGeometryCollection.h b/Engine/Source/Runtime/Experimental/RigidPhysics/Public/RigidPhysics/RigidGeometryCollection.h
--- a/Engine/Source/Runtime/Experimental/RigidPhysics/Public/RigidPhysics/RigidGeometryCollection.h
+++ b/Engine/Source/Runtime/Experimental/RigidPhysics/Public/RigidPhysics/RigidGeometryCollection.h
@@ -5,31 +5,12 @@
 
 #if UE_RIGIDPHYSICS_API_ENABLED
 
-#include "RigidPhysics/RigidBody.h"
-#include "RigidPhysics/RigidBodyContainer.h"
+#include "RigidPhysics/Internal/IRigidGeometryCollection.h"
 #include "RigidPhysics/RigidGeometryCollectionHandle.h"
+#include "RigidPhysics/RigidObjectPtr.h"
 
 namespace UE::Physics
 {
-	// The interface for rigid body containers. There will be several implementations of this API. E.g., synchronous
-	// physics, asynchronous physics, immediate physics, remote physics, etc. 
-	class UE_INTERNAL IRigidGeometryCollection : public IRigidBodyContainer
-	{
-	public:
-		UE_RIGIDPHYSICS_RIGIDTYPED_DECL(RIGIDPHYSICS_API, IRigidGeometryCollection);
-
-		UE_INTERNAL IRigidGeometryCollection()
-		{
-		}
-
-		UE_INTERNAL virtual ~IRigidGeometryCollection()
-		{
-		}
-
-		// TODO_CHAOSAPI: Figure out the public GeometryCollection API!!
-		UE_INTERNAL RIGIDPHYSICS_API virtual IRigidBody* GetBodyAt(int32 BodyIndex) const = 0;
-	};
-
 	// Utility wrapper around IRigidGeometryCollection that adds Context and allows
 	// TRigidGeometryCollectionPtr to have pointer semantics without exposing IRigidGeometryCollection
 	template<typename ContextType>
diff --git a/Engine/Source/Runtime/Experimental/RigidPhysics/Public/RigidPhysics/RigidScene.h b/Engine/Source/Runtime/Experimental/RigidPhysics/Public/RigidPhysics/RigidScene.h
--- a/Engine/Source/Runtime/Experimental/RigidPhysics/Public/RigidPhysics/RigidScene.h
+++ b/Engine/Source/Runtime/Experimental/RigidPhysics/Public/RigidPhysics/RigidScene.h
@@ -5,6 +5,7 @@
 
 #if UE_RIGIDPHYSICS_API_ENABLED
 
+#include "RigidPhysics/Internal/IRigidScene.h"
 #include "RigidPhysics/RigidLockable.h"
 #include "RigidPhysics/RigidLog.h"
 #include "RigidPhysics/RigidObjectId.h"
@@ -13,102 +14,8 @@
 #include "RigidPhysics/RigidTyped.h"
 #include "Templates/SharedPointer.h"
 
-
 namespace UE::Physics
 {
-	// The interface to a rigid body scene. The rigid body scene owns the physics simulations and
-	// manages a set of rigid body containers which in turn manage sets of rigid bodies.
-	class UE_INTERNAL IRigidScene : public IRigidLockable
-	{
-	public:
-		UE_RIGIDPHYSICS_RIGIDTYPED_DECL(RIGIDPHYSICS_API, IRigidScene);
-
-		IRigidScene()
-		{
-		}
-
-		virtual ~IRigidScene()
-		{
-		}
-
-		UE_INTERNAL RIGIDPHYSICS_API virtual FRigidSceneHandle GetHandle() const = 0;
-
-		UE_INTERNAL RIGIDPHYSICS_API virtual IRigidBodyContainer* GetBodyContainer(const FRigidObjectId& InBodyContainerId) const = 0;
-
-		UE_INTERNAL RIGIDPHYSICS_API virtual IRigidGeometryCollection* GetGeometryCollection(const FRigidObjectId& InGeometryCollectionId) const = 0;
-
-		UE_INTERNAL RIGIDPHYSICS_API virtual bool Query(const FRigidQueryDesc& InQueryDesc, TArray<FRigidQueryHit>& OutHits)
-		{
-			RIGIDPHYSICS_API_UNSUPPORTED();
-			return false;
-		}
-
-		// Advance the simulation by Dt. If the implementation is asynchronous, this will return immediately
-		UE_INTERNAL RIGIDPHYSICS_API virtual void Tick(float InDt)
-		{
-			RIGIDPHYSICS_API_UNSUPPORTED();
-		}
-
-		// GT Only API - Move this somehow
-		UE_INTERNAL RIGIDPHYSICS_API virtual void RegisterModifier(IRigidSceneModifier* InModifier)
-		{
-			RIGIDPHYSICS_API_UNSUPPORTED();
-		}
-
-		UE_INTERNAL RIGIDPHYSICS_API virtual void UnregisterModifier(IRigidSceneModifier* InModifier)
-		{
-			RIGIDPHYSICS_API_UNSUPPORTED();
-		}
-
-		UE_INTERNAL RIGIDPHYSICS_API virtual IRigidBody* CreateBody(const FRigidDebugName& InName, ERigidMovementType InMovementType)
-		{
-			RIGIDPHYSICS_API_UNSUPPORTED();
-			return nullptr;
-		}
-
-		UE_INTERNAL RIGIDPHYSICS_API virtual void DestroyBody(IRigidBody* InBody)
-		{
-			RIGIDPHYSICS_API_UNSUPPORTED();
-		}
-
-		UE_INTERNAL RIGIDPHYSICS_API virtual IRigidGeometryCollection* CreateGeometryCollection(const FRigidDebugName& InName, const FSimulationParameters& InParameters)
-		{
-			RIGIDPHYSICS_API_UNSUPPORTED();
-			return nullptr;
-		}
-
-		UE_INTERNAL RIGIDPHYSICS_API virtual void DestroyGeometryCollection(IRigidGeometryCollection* InGC)
-		{
-			RIGIDPHYSICS_API_UNSUPPORTED();
-		}
-
-	};
-
-	//
-	// Base class for the Game Context Scene. This provides the BodyContainerRegistry, partly for convenience
-	// and partly so that we can convert Handles to Objects in the debugger with natvis
-	//
-	class UE_INTERNAL FRigidScene : public IRigidScene
-	{
-	public:
-		RIGIDPHYSICS_API FRigidScene();
-		RIGIDPHYSICS_API virtual ~FRigidScene();
-
-		UE_INTERNAL RIGIDPHYSICS_API virtual IRigidBodyContainer* GetBodyContainer(const FRigidObjectId& InBodyContainerId) const override final;
-		UE_INTERNAL RIGIDPHYSICS_API virtual IRigidGeometryCollection* GetGeometryCollection(const FRigidObjectId& InGeometryCollectionId) const override final;
-
-	protected:
-		UE_INTERNAL RIGIDPHYSICS_API FRigidBodyContainerRegistry* GetBodyContainerRegistry() const
-		{
-			return BodyContainerRegistry.Get();
-		}
-
-	private:
-		TUniquePtr<FRigidBodyContainerRegistry> BodyContainerRegistry;
-
-	};
-
-
 	// A wrapper around IRigidScene.
 	// Exists only to give pointer semantics to TRigidScenePtr.
 	template<typename ContextType>
diff --git a/Engine/Source/Runtime/Experimental/RigidPhysics/Public/RigidPhysics/RigidSceneSettings.h b/Engine/Source/Runtime/Experimental/RigidPhysics/Public/RigidPhysics/RigidSceneSettings.h
--- a/Engine/Source/Runtime/Experimental/RigidPhysics/Public/RigidPhysics/RigidSceneSettings.h
+++ b/Engine/Source/Runtime/Experimental/RigidPhysics/Public/RigidPhysics/RigidSceneSettings.h
@@ -14,9 +14,9 @@ namespace UE::Physics
 	public:
 		UE_RIGIDPHYSICS_RIGIDTYPED_DECL(RIGIDPHYSICS_API, IRigidSceneSettings);
 
-		IRigidSceneSettings() {}
-		virtual ~IRigidSceneSettings() {}
+		IRigidSceneSettings() = default;
+		virtual ~IRigidSceneSettings() = default;
 	};
-}
+} // namespace UE::Physics
 
-#endif
\ No newline at end of file
+#endif // UE_RIGIDPHYSICS_API_ENABLED
