Use global namespace for VectorMaskBits.

#Synced-CL 50790641

diff --git a/Engine/Source/Runtime/Experimental/Chaos/Public/Chaos/AABBVectorized.h b/Engine/Source/Runtime/Experimental/Chaos/Public/Chaos/AABBVectorized.h
--- a/Engine/Source/Runtime/Experimental/Chaos/Public/Chaos/AABBVectorized.h
+++ b/Engine/Source/Runtime/Experimental/Chaos/Public/Chaos/AABBVectorized.h
@@ -47,7 +47,7 @@ namespace Chaos
 			const VectorRegister4Float ZeroGTStarToMax = VectorCompareGT(StartPoint, Max);
 			VectorRegister4Float IsFalse = VectorBitwiseAnd(VectorBitwiseOr(StarToMinGTZero, ZeroGTStarToMax), Parallel);
 
-			if (VectorMaskBits(IsFalse))
+			if (::VectorMaskBits(IsFalse))
 			{
 				return false;	//parallel and outside
 			}
@@ -74,14 +74,14 @@ namespace Chaos
 			//Outside of slab before entering another
 			IsFalse = VectorCompareGT(LatestStartTime, EarliestEndTime);
 
-			return VectorMaskBits(IsFalse) == 0;
+			return ::VectorMaskBits(IsFalse) == 0;
 		}
 
 		FORCEINLINE_DEBUGGABLE bool Intersects(const FAABBVectorized& Other) const
 		{
 			VectorRegister4Float IsFalse = VectorBitwiseOr(VectorCompareGT(Min, Other.GetMax()), VectorCompareGT(Other.GetMin(), Max));
 
-			return !static_cast<bool>(VectorMaskBits(IsFalse));
+			return !static_cast<bool>(::VectorMaskBits(IsFalse));
 		}
 
 		FORCEINLINE_DEBUGGABLE FAABBVectorized& Thicken(const VectorRegister4Float& Thickness)
