Fixed crash in FDataflowInstance::OnOwnerPostEditChangeProperty caused by a delegate binding that was not unbound when the FDataflowInstance was destroyed

#Synced-CL 50776428

diff --git a/Engine/Source/Runtime/Experimental/Dataflow/Engine/Private/Dataflow/DataflowInstance.cpp b/Engine/Source/Runtime/Experimental/Dataflow/Engine/Private/Dataflow/DataflowInstance.cpp
--- a/Engine/Source/Runtime/Experimental/Dataflow/Engine/Private/Dataflow/DataflowInstance.cpp
+++ b/Engine/Source/Runtime/Experimental/Dataflow/Engine/Private/Dataflow/DataflowInstance.cpp
@@ -498,6 +498,25 @@ FDataflowInstance::FDataflowInstance(UObject* InOwner, UDataflow* InDataflowAsse
 #endif
 }
 
+FDataflowInstance& FDataflowInstance::operator=(const FDataflowInstance& Other)
+{
+	DataflowAsset = Other.DataflowAsset;
+	DataflowTerminal = Other.DataflowTerminal;
+	VariableOverrides = Other.VariableOverrides;
+#if WITH_EDITORONLY_DATA
+	PreviewBlueprintVariableOverrides = Other.PreviewBlueprintVariableOverrides;
+#endif
+
+	// Note that Owner is deliberately not copied, since it should be set to the owning object of 
+	// this instance.
+	//
+	// OnOwnerPropertyChangedHandle should also not be copied, as each instance of 
+	// FDataflowInstance needs to have its own unique OnOwnerPropertyChangedHandle to ensure that
+	// the delegate is unregistered properly when the instance is destroyed.
+
+	return *this;
+}
+
 FDataflowInstance::~FDataflowInstance()
 {
 #if WITH_EDITOR
@@ -518,15 +537,19 @@ void FDataflowInstance::BindEditorDelegates()
 	// let listen to property change on the owner so we can listen to our own changes ( because we are a struct )
 	if (Owner)
 	{
-		FDataflowAssetDelegates::OnVariablesChanged.AddWeakLambda(Owner.Get(),
-			[this](const UDataflow* InDataflowAsset, FName InVariableName)
-			{
-				if (InDataflowAsset == DataflowAsset)
+		if (!OnOwnerPropertyChangedHandle.IsValid())
+		{
+			FDataflowAssetDelegates::OnVariablesChanged.AddWeakLambda(Owner.Get(),
+				[this](const UDataflow* InDataflowAsset, FName InVariableName)
 				{
-					VariableOverrides.OnDataflowVariablesChanged(InDataflowAsset, InVariableName);
-				}
-			});
-		OnOwnerPropertyChangedHandle = FCoreUObjectDelegates::OnObjectPropertyChanged.AddRaw(this, &FDataflowInstance::OnOwnerPostEditChangeProperty);
+					if (InDataflowAsset == DataflowAsset)
+					{
+						VariableOverrides.OnDataflowVariablesChanged(InDataflowAsset, InVariableName);
+					}
+				});
+
+			OnOwnerPropertyChangedHandle = FCoreUObjectDelegates::OnObjectPropertyChanged.AddRaw(this, &FDataflowInstance::OnOwnerPostEditChangeProperty);
+		}
 	}
 }
 
@@ -538,6 +561,7 @@ void FDataflowInstance::UnbindEditorDelegates()
 	}
 	// make sure we remove this outside of the Owner check in case of the Owner changed value since construction
 	FCoreUObjectDelegates::OnObjectPropertyChanged.Remove(OnOwnerPropertyChangedHandle);
+	OnOwnerPropertyChangedHandle.Reset();
 }
 
 void FDataflowInstance::OnOwnerPostEditChangeProperty(UObject* InObject, struct FPropertyChangedEvent& InPropertyChangedEvent)
diff --git a/Engine/Source/Runtime/Experimental/Dataflow/Engine/Public/Dataflow/DataflowInstance.h b/Engine/Source/Runtime/Experimental/Dataflow/Engine/Public/Dataflow/DataflowInstance.h
--- a/Engine/Source/Runtime/Experimental/Dataflow/Engine/Public/Dataflow/DataflowInstance.h
+++ b/Engine/Source/Runtime/Experimental/Dataflow/Engine/Public/Dataflow/DataflowInstance.h
@@ -210,6 +210,14 @@ struct FDataflowInstance
 	GENERATED_USTRUCT_BODY();
 
 	DATAFLOWENGINE_API explicit FDataflowInstance(UObject* InOwner = nullptr, UDataflow* InDataflowAsset = nullptr, FName InTerminalNodeName = NAME_None);
+	DATAFLOWENGINE_API FDataflowInstance& operator=(const FDataflowInstance& Other);
+	
+	// Delete these to avoid complications with Bind/UnbindEditorDelegates.
+	//
+	// They can be implemented if needed.
+	FDataflowInstance(const FDataflowInstance& Other) = delete;
+	FDataflowInstance(FDataflowInstance&& Other) = delete;
+	
 	DATAFLOWENGINE_API ~FDataflowInstance();
 
 public:
@@ -267,6 +275,8 @@ public:
 #endif 
 
 private:
+	// IMPORTANT: Any properties added here must be handled in the assignment operator
+
 	/** Dataflow asset to use */
 	UPROPERTY(EditAnywhere, Category = "Dataflow")
 	TObjectPtr<UDataflow> DataflowAsset;
