Constraints: fixed compact set & asan issues
- PerSkeletalMeshEvaluator now stores a unique pointer to the evaluator to ensure that the pointer does not change during insertion/deletion
- FAnimationEvaluator registers a delegate on the SkeletalMeshComponent so it must remain stable

#rnx
#Synced-CL 50820719

diff --git a/Engine/Source/Runtime/Experimental/Animation/Constraints/Private/Transform/AnimationEvaluation.cpp b/Engine/Source/Runtime/Experimental/Animation/Constraints/Private/Transform/AnimationEvaluation.cpp
--- a/Engine/Source/Runtime/Experimental/Animation/Constraints/Private/Transform/AnimationEvaluation.cpp
+++ b/Engine/Source/Runtime/Experimental/Animation/Constraints/Private/Transform/AnimationEvaluation.cpp
@@ -35,17 +35,7 @@ FString GetComponentName(const USceneComponent* InSceneComponent)
 FAnimationEvaluator::FAnimationEvaluator(USkeletalMeshComponent* InSkeletalMeshComponent)
 	: SkeletalMeshComponent(InSkeletalMeshComponent)
 {
-	if (InSkeletalMeshComponent)
-	{
-		const FOnBoneTransformsFinalizedMultiCast::FDelegate OnBoneTransformsFinalizedDelegate =
-			FOnBoneTransformsFinalizedMultiCast::FDelegate::CreateRaw(this, &FAnimationEvaluator::BoneTransformsFinalized);
-		OnBoneTransformsFinalizedHandle = InSkeletalMeshComponent->RegisterOnBoneTransformsFinalizedDelegate(OnBoneTransformsFinalizedDelegate);
-		
-		if (bLogNewEvaluation)
-		{
-			UE_LOG(LogTemp, Warning, TEXT("[%p] FAnimationEvaluator register finalize bone from %s"), this, *GetComponentName(InSkeletalMeshComponent));
-		}
-	}
+	BindDelegates();
 }
 
 static float CachingDelta = 0.f;
@@ -89,20 +79,7 @@ FEvaluationForCachingScope::~FEvaluationForCachingScope()
 	
 FAnimationEvaluator::~FAnimationEvaluator()
 {
-	if (OnBoneTransformsFinalizedHandle.IsValid())
-	{
-		constexpr bool bEvenIfPendingKill = true;
-		if (USkeletalMeshComponent* RawSkeletalMeshComponent = SkeletalMeshComponent.Get(bEvenIfPendingKill))
-		{
-			RawSkeletalMeshComponent->UnregisterOnBoneTransformsFinalizedDelegate(OnBoneTransformsFinalizedHandle);
-			if (bLogNewEvaluation)
-			{
-				UE_LOG(LogTemp, Warning, TEXT("[%p]  FAnimationEvaluator unregister finalize bone from %s"), this, *GetComponentName(RawSkeletalMeshComponent));
-			}
-		}
-		OnBoneTransformsFinalizedHandle.Reset();
-	}
-
+	UnbindDelegates();
 	PostEvaluationTasks.Reset();
 }
 	
@@ -225,6 +202,40 @@ bool FAnimationEvaluator::AddPostEvaluationTask(const FAnimationEvaluationTask&
 	return false;
 }
 	
+void FAnimationEvaluator::BindDelegates()
+{
+	UnbindDelegates();
+
+	if (USkeletalMeshComponent* RawSkeletalMeshComponent = SkeletalMeshComponent.Get())
+	{
+		const FOnBoneTransformsFinalizedMultiCast::FDelegate OnBoneTransformsFinalizedDelegate =
+			FOnBoneTransformsFinalizedMultiCast::FDelegate::CreateRaw(this, &FAnimationEvaluator::BoneTransformsFinalized);
+		OnBoneTransformsFinalizedHandle = RawSkeletalMeshComponent->RegisterOnBoneTransformsFinalizedDelegate(OnBoneTransformsFinalizedDelegate);
+		
+		if (bLogNewEvaluation)
+		{
+			UE_LOG(LogTemp, Warning, TEXT("[%p] FAnimationEvaluator register finalize bone from %s"), this, *GetComponentName(RawSkeletalMeshComponent));
+		}
+	}
+}
+	
+void FAnimationEvaluator::UnbindDelegates()
+{
+	if (OnBoneTransformsFinalizedHandle.IsValid())
+	{
+		constexpr bool bEvenIfPendingKill = true;
+		if (USkeletalMeshComponent* RawSkeletalMeshComponent = SkeletalMeshComponent.Get(bEvenIfPendingKill))
+		{
+			RawSkeletalMeshComponent->UnregisterOnBoneTransformsFinalizedDelegate(OnBoneTransformsFinalizedHandle);
+			if (bLogNewEvaluation)
+			{
+				UE_LOG(LogTemp, Warning, TEXT("[%p]  FAnimationEvaluator unregister finalize bone from %s"), this, *GetComponentName(RawSkeletalMeshComponent));
+			}
+		}
+		OnBoneTransformsFinalizedHandle.Reset();
+	}
+}
+	
 void FAnimationEvaluator::UpdateContext()
 {
 	if (!SkeletalMeshComponent.IsValid())
@@ -447,16 +458,20 @@ void FAnimationEvaluationCache::MarkForEvaluation(const USkeletalMeshComponent*
 		return;
 	}
 
-	if (FAnimationEvaluator* AnimationEvaluator = PerSkeletalMeshEvaluator.Find(InSkeletalMeshComponent))
+	if (TUniquePtr<FAnimationEvaluator>* Found = PerSkeletalMeshEvaluator.Find(InSkeletalMeshComponent))
 	{
-		if (AnimationEvaluator->IsValid())
+		FAnimationEvaluator* AnimationEvaluator = Found->Get();
+		if (ensure(AnimationEvaluator))
 		{
-			if (bLogNewEvaluation)
+			if (AnimationEvaluator->IsValid())
 			{
-				UE_LOG(LogTemp, Warning, TEXT("Marked %s's evaluator for evaluation."), *GetComponentName(InSkeletalMeshComponent));
+				if (bLogNewEvaluation)
+				{
+					UE_LOG(LogTemp, Warning, TEXT("Marked %s's evaluator for evaluation."), *GetComponentName(InSkeletalMeshComponent));
+				}
+				constexpr bool bDoNotRefreshBones = false;
+				AnimationEvaluator->Update(bDoNotRefreshBones);
 			}
-			constexpr bool bDoNotRefreshBones = false;
-			AnimationEvaluator->Update(bDoNotRefreshBones);
 		}
 	}
 }
@@ -471,28 +486,35 @@ const FAnimationEvaluator& FAnimationEvaluationCache::GetEvaluator(USkeletalMesh
 	
 	constexpr bool bRefreshBones = true;
 
-	if (FAnimationEvaluator* Found = PerSkeletalMeshEvaluator.Find(InSkeletalMeshComponent))
+	if (TUniquePtr<FAnimationEvaluator>* Found = PerSkeletalMeshEvaluator.Find(InSkeletalMeshComponent))
 	{
-		if (!Found->IsValid())
+		FAnimationEvaluator* AnimationEvaluator = Found->Get();
+		if (!ensure(AnimationEvaluator))
+		{
+			(*Found) = MakeUnique<FAnimationEvaluator>(InSkeletalMeshComponent);
+			AnimationEvaluator = Found->Get();
+		}
+		
+		if (!AnimationEvaluator->IsValid())
 		{
 			if (bLogNewEvaluation)
 			{
 				UE_LOG(LogTemp, Warning, TEXT("Update %s's evaluator for evaluation."), *GetComponentName(InSkeletalMeshComponent));
 			}
-			Found->Update(bRefreshBones);
+			AnimationEvaluator->Update(bRefreshBones);
 		}
-		return *Found;
+		return *AnimationEvaluator;
 	}
 	
-	FAnimationEvaluator& NewEvaluator = PerSkeletalMeshEvaluator.Emplace(InSkeletalMeshComponent, InSkeletalMeshComponent);
-	NewEvaluator.Update(bRefreshBones);
+	TUniquePtr<FAnimationEvaluator>& NewEvaluator = PerSkeletalMeshEvaluator.Emplace(InSkeletalMeshComponent, new FAnimationEvaluator(InSkeletalMeshComponent));
+	NewEvaluator.Get()->Update(bRefreshBones);
 
 	if (bLogNewEvaluation)
 	{
 		UE_LOG(LogTemp, Warning, TEXT("Create new evaluator for %s."), *GetComponentName(InSkeletalMeshComponent));
 	}
 	
-	return NewEvaluator;
+	return *NewEvaluator.Get();
 }
 
 const FAnimationEvaluator& FAnimationEvaluationCache::GetEvaluator(USkeletalMeshComponent* InSkeletalMeshComponent,
@@ -506,41 +528,51 @@ const FAnimationEvaluator& FAnimationEvaluationCache::GetEvaluator(USkeletalMesh
 	
 	constexpr bool bRefreshBones = true;
 
-	if (FAnimationEvaluator* Found = PerSkeletalMeshEvaluator.Find(InSkeletalMeshComponent))
+	if (TUniquePtr<FAnimationEvaluator>* Found = PerSkeletalMeshEvaluator.Find(InSkeletalMeshComponent))
 	{
-		const bool bAdded = Found->AddPostEvaluationTask(InTask);
-		if (!Found->IsValid())
+		FAnimationEvaluator* AnimationEvaluator = Found->Get();
+		if (!ensure(AnimationEvaluator))
+		{
+			(*Found) = MakeUnique<FAnimationEvaluator>(InSkeletalMeshComponent);
+			AnimationEvaluator = Found->Get();
+		}
+		
+		const bool bAdded = AnimationEvaluator->AddPostEvaluationTask(InTask);
+		if (!AnimationEvaluator->IsValid())
 		{
 			if (bLogNewEvaluation)
 			{
 				UE_LOG(LogTemp, Warning, TEXT("Update %s's evaluator for evaluation."), *GetComponentName(InSkeletalMeshComponent));
 			}
-			Found->Update(bRefreshBones);
+			AnimationEvaluator->Update(bRefreshBones);
 		}
 		else if (bAdded && InTask.IsValid(InSkeletalMeshComponent))
 		{
 			InTask.PostEvaluationFunction();
 		}
-		return *Found;
+		return *AnimationEvaluator;
 	}
 	
-	FAnimationEvaluator& NewEvaluator = PerSkeletalMeshEvaluator.Emplace(InSkeletalMeshComponent, InSkeletalMeshComponent);
-	NewEvaluator.AddPostEvaluationTask(InTask);
-	NewEvaluator.Update(bRefreshBones);
+	TUniquePtr<FAnimationEvaluator>& NewEvaluator = PerSkeletalMeshEvaluator.Emplace(InSkeletalMeshComponent, new FAnimationEvaluator(InSkeletalMeshComponent));
+	NewEvaluator.Get()->AddPostEvaluationTask(InTask);
+	NewEvaluator.Get()->Update(bRefreshBones);
 
 	if (bLogNewEvaluation)
 	{
 		UE_LOG(LogTemp, Warning, TEXT("Create new evaluator for %s."), *GetComponentName(InSkeletalMeshComponent));
 	}
 	
-	return NewEvaluator;
+	return *NewEvaluator.Get();
 }
 
 const FAnimationEvaluator& FAnimationEvaluationCache::FindEvaluator(const USkeletalMeshComponent* InSkeletalMeshComponent) const
 {
-	if (const FAnimationEvaluator* Found = InSkeletalMeshComponent ? PerSkeletalMeshEvaluator.Find(InSkeletalMeshComponent) : nullptr)
+	if (const TUniquePtr<FAnimationEvaluator>* Found = InSkeletalMeshComponent ? PerSkeletalMeshEvaluator.Find(InSkeletalMeshComponent) : nullptr)
 	{
-		return *Found;
+		if (const FAnimationEvaluator* AnimationEvaluator = Found->Get())
+		{
+			return *AnimationEvaluator;	
+		}
 	}
 	
 	static const FAnimationEvaluator Invalid(nullptr);
diff --git a/Engine/Source/Runtime/Experimental/Animation/Constraints/Public/Transform/AnimationEvaluation.h b/Engine/Source/Runtime/Experimental/Animation/Constraints/Public/Transform/AnimationEvaluation.h
--- a/Engine/Source/Runtime/Experimental/Animation/Constraints/Public/Transform/AnimationEvaluation.h
+++ b/Engine/Source/Runtime/Experimental/Animation/Constraints/Public/Transform/AnimationEvaluation.h
@@ -63,6 +63,10 @@ private:
 
 	FAnimationEvaluator() = default;
 
+	/** Bind/Unbind this to the SkeletalMeshComponent's OnBoneTransformsFinalizedMC delegate. */
+	void BindDelegates();
+	void UnbindDelegates();
+	
 	/** Prepares the animation context for evaluation. */
 	void UpdateContext();
 	
@@ -124,7 +128,7 @@ private:
 	void UnregisterNotifications();
 
 	/** List of connected skeletal mesh components and animation evaluators. */
-	TMap<USkeletalMeshComponent*, FAnimationEvaluator> PerSkeletalMeshEvaluator;
+	TMap<USkeletalMeshComponent*, TUniquePtr<FAnimationEvaluator>> PerSkeletalMeshEvaluator;
 
 	FDelegateHandle ConstraintsNotificationHandle;
 };
