Moved Vector16Int8 and Vector32Int8 out of the chaos namespace to fix compiler errors due to unity builds.

#Synced-CL 50785456

diff --git a/Engine/Source/Runtime/Experimental/ChaosCore/Public/Chaos/Vector16Int8FPU.h b/Engine/Source/Runtime/Experimental/ChaosCore/Public/Chaos/Vector16Int8FPU.h
--- a/Engine/Source/Runtime/Experimental/ChaosCore/Public/Chaos/Vector16Int8FPU.h
+++ b/Engine/Source/Runtime/Experimental/ChaosCore/Public/Chaos/Vector16Int8FPU.h
@@ -5,64 +5,60 @@
 
 #include "Math/VectorRegister.h"
 
-UE_INTERNAL
-namespace Chaos
+struct UE_INTERNAL VectorRegister16Int8
 {
-	struct VectorRegister16Int8
-	{
-		int8 Value[16];
-	};
+	int8 Value[16];
+};
 
-	FORCEINLINE VectorRegister16Int8 Vector16Int8Set1(int8 Value)
+UE_INTERNAL FORCEINLINE VectorRegister16Int8 Vector16Int8Set1(int8 Value)
+{
+	VectorRegister16Int8 Result;
+	for (uint32 I = 0; I < 16; ++I)
 	{
-		VectorRegister16Int8 Result;
-		for (uint32 I = 0; I < 16; ++I)
-		{
-			Result.Value[I] = Value;
-		}
-		return Result;
+		Result.Value[I] = Value;
 	}
+	return Result;
+}
 
-	FORCEINLINE VectorRegister16Int8 VectorLoad16Int8(const int8* Ptr)
+UE_INTERNAL FORCEINLINE VectorRegister16Int8 VectorLoad16Int8(const int8* Ptr)
+{
+	VectorRegister16Int8 Result;
+	for (uint32 I = 0; I < 16; ++I)
 	{
-		VectorRegister16Int8 Result;
-		for (uint32 I = 0; I < 16; ++I)
-		{
-			Result.Value[I] = Ptr[I];
-		}
-		return Result;
+		Result.Value[I] = Ptr[I];
 	}
+	return Result;
+}
 
-	FORCEINLINE void VectorStore(VectorRegister16Int8 Vec, int8* Ptr)
+UE_INTERNAL FORCEINLINE void VectorStore(VectorRegister16Int8 Vec, int8* Ptr)
+{
+	for (uint32 I = 0; I < 16; ++I)
 	{
-		for (uint32 I = 0; I < 16; ++I)
-		{
-			Ptr[I] = Vec.Value[I];
-		}
+		Ptr[I] = Vec.Value[I];
 	}
+}
 
-	FORCEINLINE VectorRegister16Int8 VectorCompareEQ(VectorRegister16Int8 Vec1, VectorRegister16Int8 Vec2)
+UE_INTERNAL FORCEINLINE VectorRegister16Int8 VectorCompareEQ(VectorRegister16Int8 Vec1, VectorRegister16Int8 Vec2)
+{
+	VectorRegister16Int8 Result;
+	for (uint32 I = 0; I < 16; ++I)
 	{
-		VectorRegister16Int8 Result;
-		for (uint32 I = 0; I < 16; ++I)
-		{
-			Result.Value[I] = (Vec1.Value[I] == Vec2.Value[I]) ? 0xFF : 0;
-		}
-		return Result;
+		Result.Value[I] = (Vec1.Value[I] == Vec2.Value[I]) ? 0xFF : 0;
 	}
+	return Result;
+}
 
-	FORCEINLINE uint16 VectorMaskBits(VectorRegister16Int8 VecMask)
+UE_INTERNAL FORCEINLINE uint16 VectorMaskBits(VectorRegister16Int8 VecMask)
+{
+	constexpr uint8 SignMask = (uint16)(1 << 7);
+	uint16 Result = 0;
+	for (uint16 I = 0; I < 16; ++I)
 	{
-		constexpr uint8 SignMask = (uint16)(1 << 7);
-		uint16 Result = 0;
-		for (uint16 I = 0; I < 16; ++I)
+		const uint8 Sign = VecMask.Value[I] & SignMask;
+		if (Sign != 0)
 		{
-			const uint8 Sign = VecMask.Value[I] & SignMask;
-			if (Sign != 0)
-			{
-				Result |= 1 << I;
-			}
+			Result |= 1 << I;
 		}
-		return Result;
 	}
-} // namespace Chaos
+	return Result;
+}
diff --git a/Engine/Source/Runtime/Experimental/ChaosCore/Public/Chaos/Vector16Int8Neon.h b/Engine/Source/Runtime/Experimental/ChaosCore/Public/Chaos/Vector16Int8Neon.h
--- a/Engine/Source/Runtime/Experimental/ChaosCore/Public/Chaos/Vector16Int8Neon.h
+++ b/Engine/Source/Runtime/Experimental/ChaosCore/Public/Chaos/Vector16Int8Neon.h
@@ -7,53 +7,49 @@
 
 #include "Math/VectorRegister.h"
 
-UE_INTERNAL
-namespace Chaos
+struct UE_INTERNAL VectorRegister16Int8
 {
-	struct VectorRegister16Int8
-	{
-		int8x16_t Value;
-	};
-
-	FORCEINLINE VectorRegister16Int8 Vector16Int8Set1(int8 Value)
-	{
-		VectorRegister16Int8 Result;
-		Result.Value = vdupq_n_s8(Value);
-		return Result;
-	}
-
-	FORCEINLINE VectorRegister16Int8 VectorLoad16Int8(const int8* Ptr)
-	{
-		VectorRegister16Int8 Result;
-		Result.Value = vld1q_s8(Ptr);
-		return Result;
-	}
-
-	FORCEINLINE void VectorStore(VectorRegister16Int8 Vec, int8* Ptr)
-	{
-		vst1q_s8(Ptr, Vec.Value);
-	}
-
-	FORCEINLINE VectorRegister16Int8 VectorCompareEQ(VectorRegister16Int8 Vec1, VectorRegister16Int8 Vec2)
-	{
-		VectorRegister16Int8 Result;
-		Result.Value = vceqq_s8(Vec1.Value, Vec2.Value);
-		return Result;
-	}
-
-	FORCEINLINE uint16 VectorMaskBits(VectorRegister16Int8 VecMask)
-	{
-		const uint8 Weights[16]{ 1, 2, 4, 8, 16, 32, 64, 128, 1, 2, 4, 8, 16, 32, 64, 128 };
-
-		int32x4_t Signs = vshrq_n_s8(VecMask.Value, 7); // sign bit of each lane replicated 32x
-		uint8x16_t Weighted = vandq_u8(Signs, vld1q_u8(Weights)); // pick bit for lane position
-
-		// Split into low and high and sum across the vectors
-		const uint16 Low = vaddv_u8(vget_low_u8(Weighted));
-		const uint16 High = vaddv_u8(vget_high_u8(Weighted));
-		const uint16 Result = static_cast<uint16>(High << 8u) | Low;
-		return Result;
-	}
-} // namespace Chaos
+	int8x16_t Value;
+};
+
+UE_INTERNAL FORCEINLINE VectorRegister16Int8 Vector16Int8Set1(int8 Value)
+{
+	VectorRegister16Int8 Result;
+	Result.Value = vdupq_n_s8(Value);
+	return Result;
+}
+
+UE_INTERNAL FORCEINLINE VectorRegister16Int8 VectorLoad16Int8(const int8* Ptr)
+{
+	VectorRegister16Int8 Result;
+	Result.Value = vld1q_s8(Ptr);
+	return Result;
+}
+
+UE_INTERNAL FORCEINLINE void VectorStore(VectorRegister16Int8 Vec, int8* Ptr)
+{
+	vst1q_s8(Ptr, Vec.Value);
+}
+
+UE_INTERNAL FORCEINLINE VectorRegister16Int8 VectorCompareEQ(VectorRegister16Int8 Vec1, VectorRegister16Int8 Vec2)
+{
+	VectorRegister16Int8 Result;
+	Result.Value = vceqq_s8(Vec1.Value, Vec2.Value);
+	return Result;
+}
+
+UE_INTERNAL FORCEINLINE uint16 VectorMaskBits(VectorRegister16Int8 VecMask)
+{
+	const uint8 Weights[16]{ 1, 2, 4, 8, 16, 32, 64, 128, 1, 2, 4, 8, 16, 32, 64, 128 };
+
+	int32x4_t Signs = vshrq_n_s8(VecMask.Value, 7); // sign bit of each lane replicated 32x
+	uint8x16_t Weighted = vandq_u8(Signs, vld1q_u8(Weights)); // pick bit for lane position
+
+	// Split into low and high and sum across the vectors
+	const uint16 Low = vaddv_u8(vget_low_u8(Weighted));
+	const uint16 High = vaddv_u8(vget_high_u8(Weighted));
+	const uint16 Result = static_cast<uint16>(High << 8u) | Low;
+	return Result;
+}
 
 #endif // PLATFORM_ENABLE_VECTORINTRINSICS_NEON
diff --git a/Engine/Source/Runtime/Experimental/ChaosCore/Public/Chaos/Vector16Int8SSE.h b/Engine/Source/Runtime/Experimental/ChaosCore/Public/Chaos/Vector16Int8SSE.h
--- a/Engine/Source/Runtime/Experimental/ChaosCore/Public/Chaos/Vector16Int8SSE.h
+++ b/Engine/Source/Runtime/Experimental/ChaosCore/Public/Chaos/Vector16Int8SSE.h
@@ -7,45 +7,41 @@
 
 #include "Math/VectorRegister.h"
 
-UE_INTERNAL
-namespace Chaos
+struct UE_INTERNAL VectorRegister16Int8
 {
-	struct VectorRegister16Int8
-	{
-		VectorRegister4Int Value;
-	};
-
-	FORCEINLINE VectorRegister16Int8 Vector16Int8Set1(int8 Value)
-	{
-		VectorRegister16Int8 Result;
-		Result.Value = _mm_set1_epi8(Value);
-		return Result;
-	}
-
-	FORCEINLINE VectorRegister16Int8 VectorLoad16Int8(const int8* Ptr)
-	{
-		VectorRegister16Int8 Result;
-		Result.Value = _mm_loadu_si128((const __m128i*)(Ptr));
-		return Result;
-	}
-
-	FORCEINLINE void VectorStore(VectorRegister16Int8 Vec, int8* Ptr)
-	{
-		_mm_storeu_si128((__m128i*)(Ptr), Vec.Value);
-	}
-
-	FORCEINLINE VectorRegister16Int8 VectorCompareEQ(VectorRegister16Int8 Vec1, VectorRegister16Int8 Vec2)
-	{
-		VectorRegister16Int8 Result;
-		Result.Value = _mm_cmpeq_epi8(Vec1.Value, Vec2.Value);
-		return Result;
-	}
-
-	FORCEINLINE uint16 VectorMaskBits(VectorRegister16Int8 VecMask)
-	{
-		uint16 Result = (uint16)_mm_movemask_epi8(VecMask.Value);
-		return Result;
-	}
-} // namespace Chaos
+	VectorRegister4Int Value;
+};
+
+UE_INTERNAL FORCEINLINE VectorRegister16Int8 Vector16Int8Set1(int8 Value)
+{
+	VectorRegister16Int8 Result;
+	Result.Value = _mm_set1_epi8(Value);
+	return Result;
+}
+
+UE_INTERNAL FORCEINLINE VectorRegister16Int8 VectorLoad16Int8(const int8* Ptr)
+{
+	VectorRegister16Int8 Result;
+	Result.Value = _mm_loadu_si128((const __m128i*)(Ptr));
+	return Result;
+}
+
+UE_INTERNAL FORCEINLINE void VectorStore(VectorRegister16Int8 Vec, int8* Ptr)
+{
+	_mm_storeu_si128((__m128i*)(Ptr), Vec.Value);
+}
+
+UE_INTERNAL FORCEINLINE VectorRegister16Int8 VectorCompareEQ(VectorRegister16Int8 Vec1, VectorRegister16Int8 Vec2)
+{
+	VectorRegister16Int8 Result;
+	Result.Value = _mm_cmpeq_epi8(Vec1.Value, Vec2.Value);
+	return Result;
+}
+
+UE_INTERNAL FORCEINLINE uint16 VectorMaskBits(VectorRegister16Int8 VecMask)
+{
+	uint16 Result = (uint16)_mm_movemask_epi8(VecMask.Value);
+	return Result;
+}
 
 #endif // PLATFORM_ENABLE_VECTORINTRINSICS
diff --git a/Engine/Source/Runtime/Experimental/ChaosCore/Public/Chaos/Vector32Int8AVX.h b/Engine/Source/Runtime/Experimental/ChaosCore/Public/Chaos/Vector32Int8AVX.h
--- a/Engine/Source/Runtime/Experimental/ChaosCore/Public/Chaos/Vector32Int8AVX.h
+++ b/Engine/Source/Runtime/Experimental/ChaosCore/Public/Chaos/Vector32Int8AVX.h
@@ -7,67 +7,63 @@
 
 #include "Chaos/Vector16Int8.h"
 
-UE_INTERNAL
-namespace Chaos
+struct UE_INTERNAL VectorRegister32Int8
 {
-	struct VectorRegister32Int8
-	{
-		__m256i V;
-	};
+	__m256i V;
+};
 
-	FORCEINLINE VectorRegister32Int8 Vector32Int8Set1(int8 Value)
-	{
-		VectorRegister32Int8 Result;
-		Result.V = _mm256_set1_epi8(Value);
-		return Result;
-	}
+UE_INTERNAL FORCEINLINE VectorRegister32Int8 Vector32Int8Set1(int8 Value)
+{
+	VectorRegister32Int8 Result;
+	Result.V = _mm256_set1_epi8(Value);
+	return Result;
+}
 
-	FORCEINLINE VectorRegister32Int8 VectorLoad32Int8(const int8* Ptr)
-	{
-		VectorRegister32Int8 Result;
-		Result.V = _mm256_loadu_si256((const __m256i*)Ptr);
-		return Result;
-	}
+UE_INTERNAL FORCEINLINE VectorRegister32Int8 VectorLoad32Int8(const int8* Ptr)
+{
+	VectorRegister32Int8 Result;
+	Result.V = _mm256_loadu_si256((const __m256i*)Ptr);
+	return Result;
+}
 
-	FORCEINLINE void VectorStore(VectorRegister32Int8 Vec, int8* Ptr)
-	{
-		_mm256_storeu_si256((__m256i*)Ptr, Vec.V);
-	}
+UE_INTERNAL FORCEINLINE void VectorStore(VectorRegister32Int8 Vec, int8* Ptr)
+{
+	_mm256_storeu_si256((__m256i*)Ptr, Vec.V);
+}
 
-	FORCEINLINE VectorRegister32Int8 VectorCompareEQ(VectorRegister32Int8 Vec1, VectorRegister32Int8 Vec2)
-	{
-		VectorRegister32Int8 Result;
+UE_INTERNAL FORCEINLINE VectorRegister32Int8 VectorCompareEQ(VectorRegister32Int8 Vec1, VectorRegister32Int8 Vec2)
+{
+	VectorRegister32Int8 Result;
 #if UE_PLATFORM_MATH_USE_AVX_2
-		Result.V = _mm256_cmpeq_epi8(Vec1.V, Vec2.V);
+	Result.V = _mm256_cmpeq_epi8(Vec1.V, Vec2.V);
 #else
-		VectorRegister16Int8 V1Low = { _mm256_extractf128_si256(Vec1.V, 0) };
-		VectorRegister16Int8 V1High = { _mm256_extractf128_si256(Vec1.V, 1) };
-		VectorRegister16Int8 V2Low = { _mm256_extractf128_si256(Vec2.V, 0) };
-		VectorRegister16Int8 V2High = { _mm256_extractf128_si256(Vec2.V, 1) };
+	VectorRegister16Int8 V1Low = { _mm256_extractf128_si256(Vec1.V, 0) };
+	VectorRegister16Int8 V1High = { _mm256_extractf128_si256(Vec1.V, 1) };
+	VectorRegister16Int8 V2Low = { _mm256_extractf128_si256(Vec2.V, 0) };
+	VectorRegister16Int8 V2High = { _mm256_extractf128_si256(Vec2.V, 1) };
 
-		VectorRegister16Int8 CmpLow = VectorCompareEQ(V1Low, V2Low);
-		VectorRegister16Int8 CmpHigh = VectorCompareEQ(V1High, V2High);
+	VectorRegister16Int8 CmpLow = VectorCompareEQ(V1Low, V2Low);
+	VectorRegister16Int8 CmpHigh = VectorCompareEQ(V1High, V2High);
 
-		Result.V = _mm256_castsi128_si256(CmpLow.Value);
-		Result.V = _mm256_insertf128_si256(Result.V, CmpHigh.Value, 1);
+	Result.V = _mm256_castsi128_si256(CmpLow.Value);
+	Result.V = _mm256_insertf128_si256(Result.V, CmpHigh.Value, 1);
 #endif
-		return Result;
-	}
+	return Result;
+}
 
-	FORCEINLINE uint32 VectorMaskBits(VectorRegister32Int8 VecMask)
-	{
-		uint32 Result = 0;
+UE_INTERNAL FORCEINLINE uint32 VectorMaskBits(VectorRegister32Int8 VecMask)
+{
+	uint32 Result = 0;
 #if UE_PLATFORM_MATH_USE_AVX_2
-		Result = _mm256_movemask_epi8(VecMask.V);
+	Result = _mm256_movemask_epi8(VecMask.V);
 #else
-		VectorRegister16Int8 VecMaskLow = { _mm256_extractf128_si256(VecMask.V, 0) };
-		VectorRegister16Int8 VecMaskHigh = { _mm256_extractf128_si256(VecMask.V, 1) };
-		uint32 Low = (uint32)VectorMaskBits(VecMaskLow);
-		uint32 High = (uint32)VectorMaskBits(VecMaskHigh);
-		Result = (High << 16) | Low;
+	VectorRegister16Int8 VecMaskLow = { _mm256_extractf128_si256(VecMask.V, 0) };
+	VectorRegister16Int8 VecMaskHigh = { _mm256_extractf128_si256(VecMask.V, 1) };
+	uint32 Low = (uint32)VectorMaskBits(VecMaskLow);
+	uint32 High = (uint32)VectorMaskBits(VecMaskHigh);
+	Result = (High << 16) | Low;
 #endif
-		return Result;
-	}
-} // namespace Chaos
+	return Result;
+}
 
 #endif
diff --git a/Engine/Source/Runtime/Experimental/ChaosCore/Public/Chaos/Vector32Int8Common.h b/Engine/Source/Runtime/Experimental/ChaosCore/Public/Chaos/Vector32Int8Common.h
--- a/Engine/Source/Runtime/Experimental/ChaosCore/Public/Chaos/Vector32Int8Common.h
+++ b/Engine/Source/Runtime/Experimental/ChaosCore/Public/Chaos/Vector32Int8Common.h
@@ -5,49 +5,45 @@
 
 #include "Chaos/Vector16Int8.h"
 
-UE_INTERNAL
-namespace Chaos
+struct UE_INTERNAL VectorRegister32Int8
 {
-	struct VectorRegister32Int8
-	{
-		VectorRegister16Int8 V[2];
-	};
-
-	FORCEINLINE VectorRegister32Int8 Vector32Int8Set1(int8 Value)
-	{
-		VectorRegister32Int8 Result;
-		Result.V[1] = Result.V[0] = Vector16Int8Set1(Value);
-		return Result;
-	}
-
-	FORCEINLINE VectorRegister32Int8 VectorLoad32Int8(const int8* Ptr)
-	{
-		VectorRegister32Int8 Result;
-		Result.V[0] = VectorLoad16Int8(Ptr);
-		Result.V[1] = VectorLoad16Int8(Ptr + 16);
-		return Result;
-	}
-
-	FORCEINLINE void VectorStore(VectorRegister32Int8 Vec, int8* Ptr)
-	{
-		VectorStore(Vec.V[0], Ptr);
-		VectorStore(Vec.V[1], Ptr + 16);
-	}
-
-	FORCEINLINE VectorRegister32Int8 VectorCompareEQ(VectorRegister32Int8 Vec1, VectorRegister32Int8 Vec2)
-	{
-		VectorRegister32Int8 Result;
-		Result.V[0] = VectorCompareEQ(Vec1.V[0], Vec2.V[0]);
-		Result.V[1] = VectorCompareEQ(Vec1.V[1], Vec2.V[1]);
-		return Result;
-	}
-
-	FORCEINLINE uint32 VectorMaskBits(VectorRegister32Int8 VecMask)
-	{
-		uint32 Result = 0;
-		uint32 Low = (uint32)VectorMaskBits(VecMask.V[0]);
-		uint32 High = (uint32)VectorMaskBits(VecMask.V[1]);
-		Result = (High << 16) | Low;
-		return Result;
-	}
-} // namespace Chaos
+	VectorRegister16Int8 V[2];
+};
+
+UE_INTERNAL FORCEINLINE VectorRegister32Int8 Vector32Int8Set1(int8 Value)
+{
+	VectorRegister32Int8 Result;
+	Result.V[1] = Result.V[0] = Vector16Int8Set1(Value);
+	return Result;
+}
+
+UE_INTERNAL FORCEINLINE VectorRegister32Int8 VectorLoad32Int8(const int8* Ptr)
+{
+	VectorRegister32Int8 Result;
+	Result.V[0] = VectorLoad16Int8(Ptr);
+	Result.V[1] = VectorLoad16Int8(Ptr + 16);
+	return Result;
+}
+
+UE_INTERNAL FORCEINLINE void VectorStore(VectorRegister32Int8 Vec, int8* Ptr)
+{
+	VectorStore(Vec.V[0], Ptr);
+	VectorStore(Vec.V[1], Ptr + 16);
+}
+
+UE_INTERNAL FORCEINLINE VectorRegister32Int8 VectorCompareEQ(VectorRegister32Int8 Vec1, VectorRegister32Int8 Vec2)
+{
+	VectorRegister32Int8 Result;
+	Result.V[0] = VectorCompareEQ(Vec1.V[0], Vec2.V[0]);
+	Result.V[1] = VectorCompareEQ(Vec1.V[1], Vec2.V[1]);
+	return Result;
+}
+
+UE_INTERNAL FORCEINLINE uint32 VectorMaskBits(VectorRegister32Int8 VecMask)
+{
+	uint32 Result = 0;
+	uint32 Low = (uint32)VectorMaskBits(VecMask.V[0]);
+	uint32 High = (uint32)VectorMaskBits(VecMask.V[1]);
+	Result = (High << 16) | Low;
+	return Result;
+}
