[Chaos] Fixed broken prerequisite chain in the task-based parallel island group solver

#Synced-CL 50656299

diff --git a/Engine/Source/Runtime/Experimental/Chaos/Private/Chaos/Island/IslandGroupManager.cpp b/Engine/Source/Runtime/Experimental/Chaos/Private/Chaos/Island/IslandGroupManager.cpp
--- a/Engine/Source/Runtime/Experimental/Chaos/Private/Chaos/Island/IslandGroupManager.cpp
+++ b/Engine/Source/Runtime/Experimental/Chaos/Private/Chaos/Island/IslandGroupManager.cpp
@@ -535,8 +535,17 @@ namespace Chaos
 						}
 
 						// Completion event
-						FGraphEventRef ScatterCompletionEvent = TGraphTask<FNullGraphTask>::CreateTask(&ScatterBodyEvents).ConstructAndDispatchWhenReady(TStatId(), ENamedThreads::AnyThread);
-						MyCompletionGraphEvent->DontCompleteUntil(ScatterCompletionEvent);
+						if (ScatterBodyEvents.IsEmpty())
+						{
+							// NOTE: If no solver bodies present, ScatterBodyEvents will be empty.
+							// Our completion event will need to depend on the previous task to wait until all constraint scatter tasks have finished.
+							MyCompletionGraphEvent->DontCompleteUntil(BodyScatterDependency);
+						}
+						else
+						{
+							FGraphEventRef ScatterCompletionEvent = TGraphTask<FNullGraphTask>::CreateTask(&ScatterBodyEvents).ConstructAndDispatchWhenReady(TStatId(), ENamedThreads::AnyThread);
+							MyCompletionGraphEvent->DontCompleteUntil(ScatterCompletionEvent);
+						}
 					});
 
 				GroupSolveEvents.Add(GroupSolveEvent);
