Updated HTTP I/O dispatcher to use TryConvertChunkIdToPackageName.

#rnx

#Synced-CL 50811283

diff --git a/Engine/Source/Runtime/Experimental/IoStore/OnDemand/Private/OnDemandHttpIoDispacher.cpp b/Engine/Source/Runtime/Experimental/IoStore/OnDemand/Private/OnDemandHttpIoDispacher.cpp
--- a/Engine/Source/Runtime/Experimental/IoStore/OnDemand/Private/OnDemandHttpIoDispacher.cpp
+++ b/Engine/Source/Runtime/Experimental/IoStore/OnDemand/Private/OnDemandHttpIoDispacher.cpp
@@ -23,6 +23,7 @@
 #include "Misc/Fork.h"
 #include "Misc/SingleThreadRunnable.h"
 #include "OnDemandHttpClient.h"
+#include "OnDemandMisc.h"
 #include "ProfilingDebugging/AssetMetadataTrace.h"
 #include "ProfilingDebugging/HttpTrace.h"
 #include "Statistics.h"
@@ -631,12 +632,12 @@ private:
 
 		void Reset()
 		{
-			PackageName.Empty(Slack);
+			PackageName.Reset();
 			Filename.Reset();
 		}
 
-		FString						PackageName;
 		TUtf8StringBuilder<Slack>	Filename;
+		TStringBuilder<Slack>		PackageName;
 	};
 
 	uint64				Dispatcher = 0;
@@ -695,45 +696,17 @@ void FTraceContext::TracePackageName(const FIoChunkId& ChunkId)
 		return;
 	}
 
-	bool bPackageDataChunk = false;
-	switch (ChunkId.GetChunkType())
-	{
-		case EIoChunkType::ExportBundleData:
-		case EIoChunkType::BulkData:
-		case EIoChunkType::OptionalBulkData:
-		case EIoChunkType::MemoryMappedBulkData:
-			bPackageDataChunk = true;
-		default:
-			break;
-	};
-
-	if (bPackageDataChunk == false)
-	{
-		return;
-	}
-
-	FPackageId PackageId;
-	FMemory::Memcpy(&PackageId, &ChunkId, sizeof(FPackageId));
-
-	if (PackageId.IsValid() == false)
-	{
-		return;
-	}
-
 	PackageNameCache.Reset();
-	FUtf8StringView _;
-	FUtf8StringView Filename = FIoDispatcherInternal::GetFilename(ChunkId, PackageNameCache.Filename, _);
-
-	if (Filename.IsEmpty())
-	{
-		return;
-	}
-
-	FString Tmp(Filename);
-	if (FPackageName::TryConvertFilenameToLongPackageName(Tmp, PackageNameCache.PackageName))
+	if (FStringView PackageName = UE::IoStore::TryConvertChunkIdToPackageName(
+		ChunkId,
+		PackageNameCache.Filename,
+		PackageNameCache.PackageName); !PackageName.IsEmpty())
 	{
+		FPackageId PackageId;
+		FMemory::Memcpy(&PackageId, &ChunkId, sizeof(FPackageId));
 		UE_TRACE_PACKAGE_NAME(PackageId, FName(PackageNameCache.PackageName));
 	}
+
 #endif // UE_TRACE_HTTP_METADATA_ENABLED
 }
 
@@ -2129,4 +2102,3 @@ TSharedPtr<IOnDemandHttpIoDispatcher> MakeOnDemanHttpIoDispatcher(TUniquePtr<cla
 #undef TRACE_INIT
 #undef TRACE_REQUEST_STARTED
 #undef TRACE_REQUEST_COMPLETED
-
diff --git a/Engine/Source/Runtime/Experimental/IoStore/OnDemand/Private/OnDemandMisc.cpp b/Engine/Source/Runtime/Experimental/IoStore/OnDemand/Private/OnDemandMisc.cpp
--- a/Engine/Source/Runtime/Experimental/IoStore/OnDemand/Private/OnDemandMisc.cpp
+++ b/Engine/Source/Runtime/Experimental/IoStore/OnDemand/Private/OnDemandMisc.cpp
@@ -44,11 +44,11 @@ FStringView TryConvertChunkIdToPackageName(
 			break;
 	};
 
-	if (bPackageDataChunk)
+	if (bPackageDataChunk == false)
 	{
 		if (OutError != nullptr)
 		{
-			*OutError << TEXT("Non package chunk ID '") << LexToString(ChunkId) << TEXT("'");
+			*OutError << TEXT("Non package chunk ID '") << ChunkId << TEXT("'");
 		}
 		return FStringView();
 	}
@@ -88,4 +88,4 @@ FString TryConvertChunkIdToPackageName(
 
 } // namespace UE::IoStore
 
-#undef UE_API 
+#undef UE_API
