Fix IAD cache losing snapshot state when making a new snapshot.
- Add new snapshot version, and invalidate old version which will reset the cache for anyone with an old snapshot.
- Move snapshot creation from shutdown to startup so it is less likely to be interrupted by a crash.
- Detect a partially created, snapshot and reset cache if found.
- Add unit test for appending journal to an existing snapshot.

#rnx

diff --git a/Engine/Source/Runtime/Experimental/IoStore/OnDemand/Private/OnDemandInstallCache.cpp b/Engine/Source/Runtime/Experimental/IoStore/OnDemand/Private/OnDemandInstallCache.cpp
--- a/Engine/Source/Runtime/Experimental/IoStore/OnDemand/Private/OnDemandInstallCache.cpp
+++ b/Engine/Source/Runtime/Experimental/IoStore/OnDemand/Private/OnDemandInstallCache.cpp
@@ -1534,6 +1534,7 @@ struct FCasSnapshot
 
 		AddModificationTime,
 		FixSwitchBuild,
+		FixLostSnapshots,
 
 		LatestPlusOne,
 		Latest = LatestPlusOne - 1
@@ -1587,10 +1588,35 @@ struct FCasSnapshot
 
 	using FChunkLocation = TPair<FCasAddr, FCasLocation>;
 
-	static TResult<FCasSnapshot>		FromJournal(const FString& JournalFile);
+	static TResult<FCasSnapshot>		FromJournal(const FString& JournalFile, FCasSnapshot&& BaseSnapshot);
 	static TResult<FCasSnapshot>		Load(const FString& SnapshotFile, int64* OutFileSize = nullptr);
 	static TResult<int64>				Save(const FCasSnapshot& Snapshot, const FString& SnapshotFile);
-	static TResult<int64>				TryCreateAndResetJournal(const FString& SnapshotFile, const FString& JournalFile);
+	/** 
+	 * Append the journal file to the snapshot and reset the journal file.
+	 *
+	 * @param SnapshotFile	The snapshot file to appened the journal to. May not exist yet.
+	 * @param JournalFile	The journal file to append to the snapshot, and then reset.
+	 *
+	 * Note: If this operation fails, it will leave behind temporary files, call TempFilesExist() check for them.
+	 * There is no provided method to cleanup temporary files. If they are discovered, the
+	 * cache is in an inconsistent state and should be reset.
+	 */
+	static TResult<int64>				AppendAndResetJournal(
+											const FString& SnapshotFile, 
+											const FString& JournalFile);
+
+	/**
+	 * There is no provided method to cleanup temporary files. If they are discovered, the
+	 * cacche is in an inconsistent state and should be reset.
+	 */
+	static bool TempFilesExist(const FString& SnapshotFile, const FString& JournalFile);
+
+	static FString MakeTempFileName(const FString& FileName)
+	{
+		const bool bIncludeDot = true;
+		const FString Extension = FPaths::GetExtension(FileName, bIncludeDot) + TEXT("tmp");
+		return FPaths::ChangeExtension(FileName, Extension);
+	}
 
 	TArray<FBlock>						Blocks;
 	TArray<FChunkLocation>				ChunkLocations;
@@ -1612,6 +1638,12 @@ bool FCasSnapshot::FHeader::IsValid() const
 		return false;
 	}
 
+	// Drop all snapshots from before we fixed snapshots not being propogated
+	if (static_cast<uint32>(Version) < static_cast<uint32>(EVersion::FixLostSnapshots))
+	{
+		return false;
+	}
+
 	// Versions that can no longer be loaded. The cache must be dropped instead.
 	const EVersion BadVersions[] = { UE_ONDEMANDINSTALLCACHE_BAD_SNAPSHOT_VERSIONS };
 	for (EVersion BadVersion : BadVersions)
@@ -1630,11 +1662,23 @@ bool FCasSnapshot::FFooter::IsValid() const
 	return FMemory::Memcmp(Magic, FFooter::MagicSequence, sizeof(FFooter::MagicSequence)) == 0;
 }
 
-TResult<FCasSnapshot> FCasSnapshot::FromJournal(const FString& JournalFile)
+TResult<FCasSnapshot> FCasSnapshot::FromJournal(const FString& JournalFile, FCasSnapshot&& BaseSnapshot)
 {
-	FCas::FLookup				CasLookup;
-	TMap<FCasBlockId, FBlock>	BlockMap;
-	FCasBlockId					CurrentBlockId;
+	FCas::FLookup CasLookup;
+	CasLookup.Reserve(BaseSnapshot.ChunkLocations.Num());
+	for (TPair<FCasAddr, FCasLocation>& Kv : BaseSnapshot.ChunkLocations)
+	{
+		CasLookup.Add(MoveTemp(Kv));
+	}
+
+	TMap<FCasBlockId, FBlock> BlockMap;
+	BlockMap.Reserve(BaseSnapshot.Blocks.Num());
+	for (const FCasSnapshot::FBlock& Block : BaseSnapshot.Blocks)
+	{
+		BlockMap.Add(Block.BlockId, Block);
+	}
+
+	FCasBlockId CurrentBlockId = BaseSnapshot.CurrentBlockId;
 
 	FResult ReplayResult = FCasJournal::Replay(
 		JournalFile,
@@ -1711,8 +1755,6 @@ TResult<FCasSnapshot> FCasSnapshot::FromJournal(const FString& JournalFile)
 		return MakeError(ReplayResult.StealError());
 	}
 
-	const int64 TimestampForMissingLastAccess = FCas::GetTimestampForMissingLastAccess();
-
 	FCasSnapshot Snapshot;
 	BlockMap.GenerateValueArray(Snapshot.Blocks);
 
@@ -1728,9 +1770,7 @@ TResult<int64> FCasSnapshot::Save(const FCasSnapshot& Snapshot, const FString& S
 
 	IFileManager& Ifm = IFileManager::Get();
 
-	const FString TmpSnapshotFile = FPaths::ChangeExtension(SnapshotFile, TEXT(".snptmp"));
-
-	TUniquePtr<FArchive> Ar(Ifm.CreateFileWriter(*TmpSnapshotFile));
+	TUniquePtr<FArchive> Ar(Ifm.CreateFileWriter(*SnapshotFile));
 	if (Ar.IsValid() == false)
 	{
 		return MakeSnapshotError<int64>(ECasErrorCode::SaveSnapshotFailed, EIoErrorCode::FileOpenFailed, FString::Printf(TEXT("Failed to open file '%s' for writing"), *SnapshotFile));
@@ -1747,7 +1787,6 @@ TResult<int64> FCasSnapshot::Save(const FCasSnapshot& Snapshot, const FString& S
 	{
 		const uint32 LastError = FPlatformMisc::GetLastError();
 		Ar.Reset();
-		Ifm.Delete(*TmpSnapshotFile);
 		return MakeSnapshotError<int64>(ECasErrorCode::SaveSnapshotFailed, EIoErrorCode::WriteError, FString::Printf(TEXT("Failed to write snapshot header to '%s'"), *SnapshotFile), LastError);
 	}
 
@@ -1760,7 +1799,6 @@ TResult<int64> FCasSnapshot::Save(const FCasSnapshot& Snapshot, const FString& S
 	{
 		const uint32 LastError = FPlatformMisc::GetLastError();
 		Ar.Reset();
-		Ifm.Delete(*TmpSnapshotFile);
 		return MakeSnapshotError<int64>(ECasErrorCode::SaveSnapshotFailed, EIoErrorCode::WriteError, FString::Printf(TEXT("Failed to write snapshot to '%s'"), *SnapshotFile), LastError);
 	}
 
@@ -1771,7 +1809,6 @@ TResult<int64> FCasSnapshot::Save(const FCasSnapshot& Snapshot, const FString& S
 	{
 		const uint32 LastError = FPlatformMisc::GetLastError();
 		Ar.Reset();
-		Ifm.Delete(*TmpSnapshotFile);
 		return MakeSnapshotError<int64>(ECasErrorCode::SaveSnapshotFailed, EIoErrorCode::WriteError, FString::Printf(TEXT("Failed to write snapshot footer to '%s'"), *SnapshotFile), LastError);
 	}
 
@@ -1780,20 +1817,9 @@ TResult<int64> FCasSnapshot::Save(const FCasSnapshot& Snapshot, const FString& S
 	{
 		const uint32 LastError = FPlatformMisc::GetLastError();
 		Ar.Reset();
-		Ifm.Delete(*TmpSnapshotFile);
 		return MakeSnapshotError<int64>(ECasErrorCode::SaveSnapshotFailed, EIoErrorCode::WriteError, FString::Printf(TEXT("Failed to close snapshot footer to '%s'"), *SnapshotFile), LastError);
 	}
 
-	if (Ifm.Move(*SnapshotFile, *TmpSnapshotFile) == false)
-	{
-		const uint32 LastError = FPlatformMisc::GetLastError();
-		Ifm.Delete(*TmpSnapshotFile);
-		return MakeSnapshotError<int64>(
-			ECasErrorCode::SaveSnapshotFailed,
-			EIoErrorCode::FileMoveFailed,
-			FString::Printf(TEXT("Failed to move snapshot '%s' -> '%s'"), *TmpSnapshotFile, *SnapshotFile), LastError);
-	}
-
 	return MakeValue(FileSize);
 }
 
@@ -1843,67 +1869,111 @@ TResult<FCasSnapshot> FCasSnapshot::Load(const FString& SnapshotFile, int64* Out
 	return MakeValue(MoveTemp(Snapshot));
 }
 
-TResult<int64> FCasSnapshot::TryCreateAndResetJournal(const FString& SnapshotFile, const FString& JournalFile)
+TResult<int64> FCasSnapshot::AppendAndResetJournal(const FString& SnapshotFile, const FString& JournalFile)
 {
 	using namespace UE::UnifiedError::IoStoreOnDemand;
 
 	IFileManager& Ifm = IFileManager::Get();
 
-	const int64 JournalFileSize = Ifm.FileSize(*JournalFile);
-	if (JournalFileSize < 0)
+	if (Ifm.FileExists(*JournalFile) == false)
 	{
-		return MakeSnapshotError<int64>(ECasErrorCode::CreateSnapshotFailed, EIoErrorCode::NotFound, FString::Printf(TEXT("Failed to find snapshot '%s'"), *SnapshotFile));
+		return MakeSnapshotError<int64>(
+			ECasErrorCode::CreateSnapshotFailed,
+			EIoErrorCode::NotFound,
+			FString::Printf(TEXT("Cannot create snapshot, failed to find journal file '%s'"), *JournalFile));
+	}
+
+	// Load the previous snapshot if it exists
+	TResult<FCasSnapshot> SnapshotResult = FCasSnapshot::Load(SnapshotFile);
+	if (SnapshotResult.HasError())
+	{
+		const FCasErrorContext* Ctx = SnapshotResult.GetError().GetErrorContext<FCasErrorContext>();
+		if (Ctx == nullptr || Ctx->IoErrorCode != EIoErrorCode::NotFound)
+		{
+			return MakeError(SnapshotResult.StealError());
+		}
+	}
+
+	if (SnapshotResult.HasValue() == false)
+	{
+		SnapshotResult = MakeValue(FCasSnapshot());
 	}
 
 	// Load the snapshot from the journal
-	TResult<FCasSnapshot> SnapshotResult = FCasSnapshot::FromJournal(JournalFile);
+	FCasSnapshot BaseSnapshot = SnapshotResult.StealValue();
+	SnapshotResult = FCasSnapshot::FromJournal(JournalFile, MoveTemp(BaseSnapshot));
 	if (SnapshotResult.HasError())
 	{
 		return MakeError(SnapshotResult.StealError());
 	}
 
+	const FString TmpSnapFile		= FCasSnapshot::MakeTempFileName(SnapshotFile);
+	const FString TmpJournalFile	= FCasSnapshot::MakeTempFileName(JournalFile);
+	
 	// Save the snapshot
 	int64 SnapshotSize			= -1;
 	FCasSnapshot Snapshot		= SnapshotResult.StealValue();
-	TResult<int64> SaveResult	= FCasSnapshot::Save(Snapshot, SnapshotFile);
-	if (SaveResult.HasValue())
+	TResult<int64> SnapshotSaveResult	= FCasSnapshot::Save(Snapshot, TmpSnapFile);
+	if (SnapshotSaveResult.HasValue())
 	{
-		SnapshotSize = SaveResult.GetValue();
+		SnapshotSize = SnapshotSaveResult.GetValue();
 	}
 	else
 	{
-		return SaveResult;
+		return SnapshotSaveResult;
 	}
 
 	// Try create a new empty journal 
-	const FString TmpJournalFile = FPaths::ChangeExtension(JournalFile, TEXT(".jrntmp"));
 	if (FResult Result = FCasJournal::Create(TmpJournalFile); Result.HasError())
 	{
-		if (Ifm.Delete(*SnapshotFile) == false)
-		{
-			UE_LOG(LogIoStoreOnDemand, Error, TEXT("Failed to delete CAS snapshot '%s'"), *SnapshotFile); 
-		}
-
 		return MakeError(Result.StealError());
 	}
 
-	if (Ifm.Move(*JournalFile , *TmpJournalFile) == false)
+	if (Ifm.Move(*SnapshotFile, *TmpSnapFile) == false)
 	{
 		const uint32 LastError = FPlatformMisc::GetLastError();
-		if (Ifm.Delete(*SnapshotFile) == false)
-		{
-			UE_LOG(LogIoStoreOnDemand, Error, TEXT("Failed to delete CAS snapshot '%s'"), *SnapshotFile); 
-		}
 
 		return MakeSnapshotError<int64>(
 			ECasErrorCode::CreateSnapshotFailed,
 			EIoErrorCode::FileMoveFailed,
-			FString::Printf(TEXT("Failed to move tmp journal file '%s' -> '%s'"), *TmpJournalFile, *JournalFile));
+			FString::Printf(TEXT("Failed to move tmp snapshot file '%s' -> '%s'"), *TmpSnapFile, *SnapshotFile),
+			LastError);
+	}
+
+	if (Ifm.Move(*JournalFile, *TmpJournalFile) == false)
+	{
+		const uint32 LastError = FPlatformMisc::GetLastError();
+
+		return MakeSnapshotError<int64>(
+			ECasErrorCode::CreateSnapshotFailed,
+			EIoErrorCode::FileMoveFailed,
+			FString::Printf(TEXT("Failed to move tmp journal file '%s' -> '%s'"), *TmpJournalFile, *JournalFile),
+			LastError);
 	}
 
 	return MakeValue(SnapshotSize);
 }
 
+bool FCasSnapshot::TempFilesExist(const FString& SnapshotFile, const FString& JournalFile)
+{
+	IFileManager& Ifm = IFileManager::Get();
+
+	const FString TmpSnapFile = MakeTempFileName(SnapshotFile);
+	const FString TmpJournalFile = MakeTempFileName(JournalFile);
+
+	if (Ifm.FileExists(*TmpSnapFile))
+	{
+		return true;
+	}
+
+	if (Ifm.FileExists(*TmpJournalFile))
+	{
+		return true;
+	}
+
+	return false;
+}
+
 ///////////////////////////////////////////////////////////////////////////////
 void FCas::LoadSnapshot(FCasSnapshot&& Snapshot)
 {
@@ -2203,16 +2273,51 @@ FOnDemandInstallCache::FOnDemandInstallCache(const FOnDemandInstallCacheConfig&
 		return;
 	}
 
+	const FString SnapshotFilename = GetSnapshotFilename();
+	const FString JournalFilename = GetJournalFilename();
+
+	// Check for any temp files that indicate we are in an incosistent state
+	if (InitResult.HasValue())
+	{
+		if (FCasSnapshot::TempFilesExist(SnapshotFilename, JournalFilename))
+		{
+			InitResult = MakeSnapshotError<void>(
+				ECasErrorCode::InitializeFailed, EIoErrorCode::FileCorrupt, 
+				FString::Printf(TEXT("Last snapshot update is incomplete, snapshot is in an inconsistent state")));
+		}
+	}
+
+	// Check if the snapshot should be updated
+	if (InitResult.HasValue())
+	{
+		IFileManager& Ifm = IFileManager::Get();
+		if (Ifm.FileSize(*JournalFilename) > int64(MaxJournalSize))
+		{
+			TResult<int64> SnapshotResult = FCasSnapshot::AppendAndResetJournal(SnapshotFilename, JournalFilename);
+			if (SnapshotResult.HasValue())
+			{
+				UE_LOG(LogIoStoreOnDemand, Log, TEXT("Saved CAS snapshot '%s' %.2lf KiB"), *SnapshotFilename, ToKiB(SnapshotResult.GetValue()));
+			}
+			else
+			{
+				UE_LOG(LogIoStoreOnDemand, Error, TEXT("Failed to update CAS snapshot from journal '%s', error '%s'"),
+					*JournalFilename, *SnapshotResult.GetError().GetErrorMessage(true).ToString());
+
+				InitResult = MakeError(SnapshotResult.StealError());
+			}
+		}
+	}
+
 	// Try read the journal snapshot
+	if (InitResult.HasValue())
 	{
-		const FString SnapshotFile = GetSnapshotFilename();
 		int64 SnapshotSize = -1;
-		TResult<FCasSnapshot> SnapshotResult = FCasSnapshot::Load(SnapshotFile, &SnapshotSize);
+		TResult<FCasSnapshot> SnapshotResult = FCasSnapshot::Load(SnapshotFilename, &SnapshotSize);
 		if (SnapshotResult.HasValue())
 		{
 			FCasSnapshot Snapshot = SnapshotResult.StealValue();
 			UE_LOG(LogIoStoreOnDemand, Log, TEXT("Loaded CAS snapshot '%s' %.2lf KiB with %d blocks and %d chunk locations"),
-					*SnapshotFile, ToKiB(SnapshotSize), Snapshot.Blocks.Num(), Snapshot.ChunkLocations.Num());
+					*SnapshotFilename, ToKiB(SnapshotSize), Snapshot.Blocks.Num(), Snapshot.ChunkLocations.Num());
 
 			Cas.LoadSnapshot(MoveTemp(Snapshot));
 		}
@@ -2221,16 +2326,15 @@ FOnDemandInstallCache::FOnDemandInstallCache(const FOnDemandInstallCacheConfig&
 			// It's ok for the snapshot to not exist, other errors are significant
 			if (Ctx->IoErrorCode != EIoErrorCode::NotFound)
 			{
-				InitResult = MakeError(SnapshotResult.GetError());
+				InitResult = MakeError(SnapshotResult.StealError());
 			}
 		}
 	}
 
-	const FString JournalFile = GetJournalFilename();
 	if (InitResult.HasValue())
 	{
 		// Replay the journal 
-		InitResult = FCasJournal::Replay(JournalFile, [this](const FCasJournal::FEntry& JournalEntry)
+		InitResult = FCasJournal::Replay(JournalFilename, [this](const FCasJournal::FEntry& JournalEntry)
 		{
 			FResult Result = Cas.ApplyJournalEntry(JournalEntry);
 			if (Result.HasError())
@@ -2264,11 +2368,12 @@ FOnDemandInstallCache::FOnDemandInstallCache(const FOnDemandInstallCacheConfig&
 	{
 		if (const FCasErrorContext* Ctx = InitResult.GetError().GetErrorContext<FCasErrorContext>())
 		{
-			if (Ctx->IoErrorCode == EIoErrorCode::NotFound)
+			if (Ctx->ErrorCode == ECasErrorCode::ReplayJournalFailed &&
+				Ctx->IoErrorCode == EIoErrorCode::NotFound)
 			{
-				if (InitResult = FCasJournal::Create(JournalFile); InitResult.HasValue())
+				if (InitResult = FCasJournal::Create(JournalFilename); InitResult.HasValue())
 				{
-					UE_LOG(LogIoStoreOnDemand, Log, TEXT("Created CAS journal '%s'"), *JournalFile);
+					UE_LOG(LogIoStoreOnDemand, Log, TEXT("Created CAS journal '%s'"), *JournalFilename);
 
 					// Make sure that there are no existing blocks when starting from an empty cache
 					const bool bDeleteExisting = true;
@@ -2333,24 +2438,6 @@ void FOnDemandInstallCache::Shutdown()
 			*JournalFile, *LexToString(Result.GetError()));
 	}
 
-	// TODO: journal snapshot is also only made on shutdown? Feels like this would be better done on startup as well
-	IFileManager& Ifm = IFileManager::Get();
-	const FString JournalFilename = GetJournalFilename();
-	if (Ifm.FileSize(*JournalFile) > int64(MaxJournalSize))
-	{
-		const FString SnapshotFilename = GetSnapshotFilename();
-		TResult<int64> SnapshotResult = FCasSnapshot::TryCreateAndResetJournal(SnapshotFilename, JournalFilename);
-		if (SnapshotResult.HasValue())
-		{
-			UE_LOG(LogIoStoreOnDemand, Log, TEXT("Saved CAS snapshot '%s' %.2lf KiB"), *SnapshotFilename, ToKiB(SnapshotResult.GetValue()));
-		}
-		else
-		{
-			UE_LOG(LogIoStoreOnDemand, Error, TEXT("Failed to create CAS snapshot from journal '%s', error '%s'"),
-				*JournalFile, *SnapshotResult.GetError().GetErrorMessage(true).ToString());
-		}
-	}
-
 #if UE_IAD_DEBUG_CONSOLE_CMDS
 	for (IConsoleCommand* Cmd : ConsoleCommands)
 	{
@@ -4559,7 +4646,7 @@ TEST_CASE("IoStore::OnDemand::InstallCache::Snapshot", "[IoStoreOnDemand][Instal
 		CHECK(Result.HasValue());
 
 		// Act
-		const FCasSnapshot Snapshot = FCasSnapshot::FromJournal(JournalFile).StealValue();
+		const FCasSnapshot Snapshot = FCasSnapshot::FromJournal(JournalFile, FCasSnapshot()).StealValue();
 
 		// Assert
 		CHECK(Snapshot.CurrentBlockId == ExpectedCurrentBlockId);
@@ -4582,6 +4669,111 @@ TEST_CASE("IoStore::OnDemand::InstallCache::Snapshot", "[IoStoreOnDemand][Instal
 			}
 		}
 	}
+
+	SECTION("AppendJournal")
+	{
+		// Arrange
+		FCasSnapshot ExpectedSnapshot;
+
+		const int64 Now = FDateTime::UtcNow().GetTicks();
+
+		for (uint32 Id = 1; Id <= 10; ++Id)
+		{
+			ExpectedSnapshot.Blocks.Add(FCasSnapshot::FBlock
+			{
+				.BlockId = FCasBlockId(Id),
+				.LastAccess = Now + Id,
+				.ModTime = Now + 10 + Id
+			});
+
+			for (uint32 Idx = 1; Idx <= 10; ++Idx)
+			{
+				FCasAddr CasAddr = CreateCasTestAddr(uint64(Id * 10 + Idx));
+				FCasLocation Loc = FCasLocation
+				{
+					.BlockId = FCasBlockId(Id),
+					.BlockOffset = Idx * 256
+				};
+				ExpectedSnapshot.ChunkLocations.Emplace(CasAddr, Loc);
+			}
+		}
+		ExpectedSnapshot.CurrentBlockId = FCasBlockId(1);
+
+		FTmpDirectoryScope _(TestBaseDir);
+		const FString SnapshotFile = TestBaseDir / TEXT("test.snp");
+		const FString JournalFile = TestBaseDir / TEXT("test.jrn");
+
+		TResult<int64> SaveSnapResult = FCasSnapshot::Save(ExpectedSnapshot, SnapshotFile);
+		CHECK(SaveSnapResult.HasValue());
+
+		FResult Result = FCasJournal::Create(JournalFile);
+		CHECK(Result.HasValue());
+
+		FCasJournal::FTransaction Tx = FCasJournal::Begin(JournalFile);
+
+		// Add a block and some chunk locations
+		FCasBlockId Block11(11);
+		Tx.BlockCreated(Block11);
+		Tx.BlockModification(Block11, Now + 100);
+		Tx.BlockAccess(Block11, Now + 101);
+		for (uint32 Idx = 1; Idx <= 10; ++Idx)
+		{
+			Tx.ChunkLocation(FCasLocation
+			{
+				.BlockId = Block11,
+				.BlockOffset = Idx * 256
+			},
+			CreateCasTestAddr(uint64(110 + Idx)));
+		}
+
+		Result = FCasJournal::Commit(MoveTemp(Tx));
+		CHECK(Result.HasValue());
+
+		// Act
+		TResult<int64> AppendResult = FCasSnapshot::AppendAndResetJournal(SnapshotFile, JournalFile);
+		CHECK(AppendResult.HasValue());
+
+		TResult<FCasSnapshot> LoadSnapResult = FCasSnapshot::Load(SnapshotFile);
+		CHECK(LoadSnapResult.HasValue());
+		FCasSnapshot Snapshot = LoadSnapResult.StealValue();
+
+		// Assert
+		CHECK(Snapshot.Blocks.Num() == ExpectedSnapshot.Blocks.Num() + 1);
+		for (int32 Idx = 0; Idx < ExpectedSnapshot.Blocks.Num(); ++Idx)
+		{
+			CHECK(Snapshot.Blocks[Idx].BlockId == ExpectedSnapshot.Blocks[Idx].BlockId);
+			CHECK(Snapshot.Blocks[Idx].LastAccess == ExpectedSnapshot.Blocks[Idx].LastAccess);
+			CHECK(Snapshot.Blocks[Idx].ModTime == ExpectedSnapshot.Blocks[Idx].ModTime);
+		}
+		CHECK(Snapshot.ChunkLocations.Num() == ExpectedSnapshot.ChunkLocations.Num() + 10);
+		for (int32 Idx = 0; Idx < ExpectedSnapshot.ChunkLocations.Num(); ++Idx)
+		{
+			CHECK(Snapshot.ChunkLocations[Idx].Get<0>() == ExpectedSnapshot.ChunkLocations[Idx].Get<0>());
+			CHECK(Snapshot.ChunkLocations[Idx].Get<1>() == ExpectedSnapshot.ChunkLocations[Idx].Get<1>());
+		}
+
+		CHECK(Snapshot.CurrentBlockId == Block11);
+
+		CHECK(Snapshot.Blocks[10].BlockId == Block11);
+		CHECK(Snapshot.Blocks[10].ModTime == Now + 100);
+		CHECK(Snapshot.Blocks[10].LastAccess == Now + 101);
+
+		for (int32 Idx = 0; Idx < 10; ++Idx)
+		{
+			const int32 ExpectIdx = ExpectedSnapshot.ChunkLocations.Num() + Idx;
+
+			FCasAddr ExpectAddr = CreateCasTestAddr(111 + Idx);
+
+			FCasLocation ExpectLoc
+			{
+				.BlockId = FCasBlockId(11),
+				.BlockOffset = (Idx + 1) * 256u
+			};
+
+			CHECK(Snapshot.ChunkLocations[ExpectIdx].Get<0>() == ExpectAddr);
+			CHECK(Snapshot.ChunkLocations[ExpectIdx].Get<1>() == ExpectLoc);
+		}
+	}
 }
 
 TEST_CASE("IoStore::OnDemand::InstallCache::ErrorHandling", "[IoStoreOnDemand][InstallCache]")
