Refined IAX CSV stat counters.

* renamed IAS cache counters to start with IAX/StreamingCache
* removed a bunch of CSV stats that are not interesting
* added CSV stats for IAX install cache

#rnx

#Synced-CL 50811283

diff --git a/Engine/Source/Runtime/Experimental/IoStore/OnDemand/Private/OnDemandHttpIoDispatcherBackend.cpp b/Engine/Source/Runtime/Experimental/IoStore/OnDemand/Private/OnDemandHttpIoDispatcherBackend.cpp
--- a/Engine/Source/Runtime/Experimental/IoStore/OnDemand/Private/OnDemandHttpIoDispatcherBackend.cpp
+++ b/Engine/Source/Runtime/Experimental/IoStore/OnDemand/Private/OnDemandHttpIoDispatcherBackend.cpp
@@ -184,7 +184,6 @@ void FHttpIoDispatcherBackend::ResolveIoRequests(FIoRequestList Requests, FIoReq
 		{
 			// Mark the request as failed (FIoRequestImpl::SetFailed) but with a specific error code.
 			Request->SetLastBackendError(EIoErrorCode::InvalidParameter);
-			Stats.OnIoRequestError();
 
 			{
 				TUniqueLock _(CompletedRequestMutex);
@@ -376,17 +375,10 @@ void FHttpIoDispatcherBackend::CompleteRequest(FResolvedRequest& ResolvedRequest
 			FIoStatus Status = FHttpIoDispatcher::CacheResponse(HttpResponse);
 			UE_CLOG(!Status.IsOk(), LogHttpIoDispatcher, Error, TEXT("Put HTTP cache failed, reason '%s'"), *Status.ToString());
 		}
-
-		Stats.OnIoRequestComplete(DispatcherRequest.GetBuffer().GetSize(), DurationMs);
-	}
-	else if (bCancelled)
-	{
-		Stats.OnIoRequestCancel();
 	}
-	else
+	else if (!bCancelled)
 	{
 		DispatcherRequest.SetFailed();
-		Stats.OnIoRequestError();
 	}
 	
 	LogRequest(ResolvedRequest, DispatcherRequest, bDecoded, HttpResponse.IsCached(), bCancelled, DurationMs);
diff --git a/Engine/Source/Runtime/Experimental/IoStore/OnDemand/Private/Statistics.cpp b/Engine/Source/Runtime/Experimental/IoStore/OnDemand/Private/Statistics.cpp
--- a/Engine/Source/Runtime/Experimental/IoStore/OnDemand/Private/Statistics.cpp
+++ b/Engine/Source/Runtime/Experimental/IoStore/OnDemand/Private/Statistics.cpp
@@ -395,31 +395,26 @@ private:
 #define UE_IAX_COUNTER(Name, Type) Type G##Name[(uint8)EHttpRequestType::NUM_SOURCES] = {{TEXT(UE_STRINGIZE(Ias/Name)), TraceCounterDisplayHint_None}, {TEXT(UE_STRINGIZE(Iad/Name)), TraceCounterDisplayHint_None}};
 #define UE_IAX_MEMORY_COUNTER(Name, Type) Type G##Name[(uint8)EHttpRequestType::NUM_SOURCES] = {{TEXT(UE_STRINGIZE(Ias/Name)), TraceCounterDisplayHint_None}, {TEXT(UE_STRINGIZE(Iad/Name)), TraceCounterDisplayHint_Memory}};
 
-// iorequest stats
-FCounterInt				GIoRequestCount(TEXT("Ias/IoRequestCount"), TraceCounterDisplayHint_None);
-FCounterAtomicInt		GIoRequestReadCount(TEXT("Ias/IoRequestReadCount"), TraceCounterDisplayHint_None);
-FCounterAtomicInt		GIoRequestReadBytes(TEXT("Ias/IoRequestReadBytes"), TraceCounterDisplayHint_Memory);
-FCounterInt				GIoRequestCancelCount(TEXT("Ias/IoRequestCancelCount"), TraceCounterDisplayHint_None);
-FCounterAtomicInt		GIoRequestErrorCount(TEXT("Ias/IoRequestErrorCount"), TraceCounterDisplayHint_None);
 // cache stats
 static uint32			GCacheBootMs = 0;
-FCounterAtomicInt		GCacheErrorCount(TEXT("Ias/CacheErrorCount"), TraceCounterDisplayHint_None);
-FCounterAtomicInt		GCacheDecodeErrorCount(TEXT("Ias/CacheDecodeErrorCount"), TraceCounterDisplayHint_None);
-FCounterAtomicInt		GCacheGetCount(TEXT("Ias/CacheGetCount"), TraceCounterDisplayHint_None);
-FCounterAtomicInt		GCachePutCount(TEXT("Ias/CachePutCount"), TraceCounterDisplayHint_None);
-FCounterAtomicInt		GCachePutExistingCount(TEXT("Ias/CachePutExistingCount"), TraceCounterDisplayHint_None);
-FCounterAtomicInt		GCachePutRejectCount(TEXT("Ias/CachePutRejectCount"), TraceCounterDisplayHint_None);
-FCounterAtomicInt		GCacheCachedBytes(TEXT("Ias/CacheCachedBytes"), TraceCounterDisplayHint_Memory);
-FCounterAtomicInt		GCacheCachedEntries(TEXT("Ias/CacheCachedEntries"), TraceCounterDisplayHint_Memory);
-FCounterAtomicInt		GCacheWrittenBytes(TEXT("Ias/CacheWrittenBytes"), TraceCounterDisplayHint_Memory);
-FCounterAtomicInt		GCacheEvictedBytes(TEXT("Ias/CacheEvictedBytes"), TraceCounterDisplayHint_Memory);
-FCounterFloat			GCacheSuspendedSeconds(TEXT("Ias/CacheSuspendedSeconds"), TraceCounterDisplayHint_None);
+FCounterAtomicInt		GCacheErrorCount(TEXT("Iax/StreamingCache/CacheErrorCount"), TraceCounterDisplayHint_None);
+FCounterAtomicInt		GCacheDecodeErrorCount(TEXT("Iax/StreamingCache/DecodeErrorCount"), TraceCounterDisplayHint_None);
+FCounterAtomicInt		GCacheGetCount(TEXT("Iax/StreamingCache/GetCount"), TraceCounterDisplayHint_None);
+FCounterAtomicInt		GCachePutCount(TEXT("Iax/StreamingCache/PutCount"), TraceCounterDisplayHint_None);
+FCounterAtomicInt		GCachePutExistingCount(TEXT("Iax/StreamingCache/PutExistingCount"), TraceCounterDisplayHint_None);
+FCounterAtomicInt		GCachePutRejectCount(TEXT("Iax/StreamingCache/PutRejectCount"), TraceCounterDisplayHint_None);
+FCounterAtomicInt		GCacheCachedBytes(TEXT("Iax/StreamingCache/CachedBytes"), TraceCounterDisplayHint_Memory);
+FCounterAtomicInt		GCacheCachedEntries(TEXT("Iax/StreamingCache/CachedEntries"), TraceCounterDisplayHint_Memory);
+FCounterAtomicInt		GCacheWriteCount(TEXT("Iax/StreamingCache/WriteCount"), TraceCounterDisplayHint_Memory);
+FCounterAtomicInt		GCacheWrittenBytes(TEXT("Iax/StreamingCache/WriteBytes"), TraceCounterDisplayHint_Memory);
+FCounterAtomicInt		GCacheEvictedBytes(TEXT("Iax/StreamingCache/EvictedBytes"), TraceCounterDisplayHint_Memory);
+FCounterFloat			GCacheSuspendedSeconds(TEXT("Iax/StreamingCache/SuspendedSeconds"), TraceCounterDisplayHint_None);
 int64					GCacheMaxBytes = 0;
-FCounterAtomicInt		GCachePendingBytes(TEXT("Ias/CachePendingBytes"), TraceCounterDisplayHint_Memory);
-FCounterAtomicInt		GCacheReadBytes(TEXT("Ias/CacheReadBytes"), TraceCounterDisplayHint_Memory);
-FCounterAtomicInt		GCacheReadMemory(TEXT("Ias/CacheReadMemoryBytes"), TraceCounterDisplayHint_Memory);
-FCounterAtomicInt		GCacheReadDisk(TEXT("Ias/CacheReadDiskBytes"), TraceCounterDisplayHint_Memory);
-FCounterAtomicInt		GCacheRejectBytes(TEXT("Ias/CachePutRejectBytes"), TraceCounterDisplayHint_Memory);
+FCounterAtomicInt		GCachePendingBytes(TEXT("Iax/StreamingCache/PendingBytes"), TraceCounterDisplayHint_Memory);
+FCounterAtomicInt		GCacheReadBytes(TEXT("Iax/StreamingCache/ReadBytes"), TraceCounterDisplayHint_Memory);
+FCounterAtomicInt		GCacheReadMemory(TEXT("Iax/StreamingCache/ReadMemoryBytes"), TraceCounterDisplayHint_Memory);
+FCounterAtomicInt		GCacheReadDisk(TEXT("Iax/StreamingCache/ReadDiskBytes"), TraceCounterDisplayHint_Memory);
+FCounterAtomicInt		GCacheRejectBytes(TEXT("Iax/StreamingCache/PutRejectBytes"), TraceCounterDisplayHint_Memory);
 // general stats
 double					GAnalyticsStartTiming = 0.0;
 // http stats
@@ -499,23 +494,6 @@ static uint32			GHttpCdnCacheUnknown = 0;
 
 CSV_DEFINE_CATEGORY(Ias, true);
 CSV_DEFINE_CATEGORY(Iad, true);
-// iorequest per frame stats
-CSV_DEFINE_STAT(Ias, FrameIoRequestCount);
-CSV_DEFINE_STAT(Ias, FrameIoRequestReadCount);
-CSV_DEFINE_STAT(Ias, FrameIoRequestReadMB);
-CSV_DEFINE_STAT(Ias, FrameIoRequestCancelCount);
-CSV_DEFINE_STAT(Ias, FrameIoRequestErrorCount);
-// cache stat totals
-CSV_DEFINE_STAT(Ias, CacheGetCount);
-CSV_DEFINE_STAT(Ias, CacheErrorCount);
-CSV_DEFINE_STAT(Ias, CachePutCount);
-CSV_DEFINE_STAT(Ias, CachePutExistingCount);
-CSV_DEFINE_STAT(Ias, CachePutRejectCount);
-CSV_DEFINE_STAT(Ias, CacheCachedMB);
-CSV_DEFINE_STAT(Ias, CacheWrittenMB);
-CSV_DEFINE_STAT(Ias, CacheEvictedMB);
-CSV_DEFINE_STAT(Ias, CacheReadMB);
-CSV_DEFINE_STAT(Ias, CacheRejectedMB);
 // http stat totals
 UE_IAX_CSV_DEFINE_STAT(HttpGetCount);
 UE_IAX_CSV_DEFINE_STAT(HttpRetryCount);
@@ -528,34 +506,45 @@ UE_IAX_CSV_DEFINE_STAT(HttpDurationMsAvg);
 UE_IAX_CSV_DEFINE_STAT(HttpDurationMsMax);
 
 ////////////////////////////////////////////////////////////////////////////////
-#if COUNTERSTRACE_ENABLED
-struct FInstallerTraceCounters
-{
-	using FCounterFloat = FCountersTrace::FCounterFloat;
+// IAX CSV Counters
+CSV_DEFINE_CATEGORY(Iax, true);
 
-	void Lock()		{ Mutex.Lock(); }
-	void Unlock()	{ Mutex.Unlock(); }
-	static FMutex Mutex; 
+CSV_DEFINE_STAT(Iax, InstallerInstallCount);
+CSV_DEFINE_STAT(Iax, InstallerInflightCount);
+CSV_DEFINE_STAT(Iax, InstallerDownloadedMiB);
+CSV_DEFINE_STAT(Iax, InstallerDurationAvgMs);
+
+CSV_DEFINE_STAT(Iax, InstallCachePurgeCount);
+CSV_DEFINE_STAT(Iax, InstallCacheDefragCount);
+CSV_DEFINE_STAT(Iax, InstallCacheWriteCount);
+CSV_DEFINE_STAT(Iax, InstallCacheWriteMiB);
+CSV_DEFINE_STAT(Iax, InstallCacheDecodeErrorCount);
+CSV_DEFINE_STAT(Iax, InstallCacheAvgCacheHitPct);
 
-	FInstallerTraceCounters()
-		: InstallCount(TEXT("Iad/InstallCount"), TraceCounterDisplayHint_None)
-		, InflightInstallCount(TEXT("Iad/InflightInstallCount"), TraceCounterDisplayHint_None)
-		, DownloadedBytes(TEXT("Iad/DownloadedBytes"), TraceCounterDisplayHint_Memory)
-		, AvgInstallDurationMs(TEXT("Iad/AvgInstallDurationMs"), TraceCounterDisplayHint_None)
-		, AvgCacheHitRatio(TEXT("Iad/AvgCacheHitRatio"), TraceCounterDisplayHint_None)
+CSV_DEFINE_STAT(Iax, StreamingCacheWriteCount);
+CSV_DEFINE_STAT(Iax, StreamingCacheWriteMiB);
+CSV_DEFINE_STAT(Iax, StreamingCacheDecodeErrorCount);
+
+////////////////////////////////////////////////////////////////////////////////
+struct FInstallerCounters
+{
+	FInstallerCounters()
+		: InstallCount(TEXT("Iax/Installer/InstallCount"), TraceCounterDisplayHint_None)
+		, InflightCount(TEXT("Iax/Installer/InflightCount"), TraceCounterDisplayHint_None)
+		, DownloadedBytes(TEXT("Iax/Installer/DownloadedBytes"), TraceCounterDisplayHint_Memory)
+		, DurationAvgMs(TEXT("Iax/Installer/DurationAvgMs"), TraceCounterDisplayHint_None)
+		, CacheHitAvgPct(TEXT("Iax/Installer/CacheHitAvgPct."), TraceCounterDisplayHint_None)
 	{ }
 
-	FCounterInt				InstallCount;
-	FCounterInt				InflightInstallCount;
-	FCounterInt				DownloadedBytes; 
-	FCounterFloat			AvgInstallDurationMs;
-	FCounterFloat			AvgCacheHitRatio;
-	FIncrementalVariance	InstallDurationMs;
-	FIncrementalVariance	CacheHitRatio;
+	FCounterAtomicInt 		InstallCount;
+	FCounterAtomicInt 		InflightCount;
+	FCounterAtomicInt 		DownloadedBytes;
+	std::atomic_uint64_t	DurationTotalMs{0};
+	FCounterAtomicInt 		DurationAvgMs;
+	FCounterAtomicInt 		CacheHitAvgPct;
+	std::atomic_uint64_t	CacheHitTotalPct{0};
 };
-FMutex FInstallerTraceCounters::Mutex;
-FInstallerTraceCounters InstallerTraceCounters;
-#endif // COUNTERSTRACE_ENABLED
+FInstallerCounters InstallerCounters;
 
 ////////////////////////////////////////////////////////////////////////////////
 struct FInstallerAnalytics
@@ -614,14 +603,18 @@ static FInstallCacheAnalytics InstallCacheAnalytics;
 struct FInstallCacheCounters
 {
 	FInstallCacheCounters()
-		: MaxCacheSize(TEXT("Iad/MaxCacheSize"), TraceCounterDisplayHint_Memory)
-		, CacheSize(TEXT("Iad/CacheSize"), TraceCounterDisplayHint_Memory)
-		, FlushedBytes(TEXT("Iad/FlushedBytes"), TraceCounterDisplayHint_Memory)
-		, FlushCount(TEXT("Iad/FlushCount"), TraceCounterDisplayHint_None)
-		, InMemoryBytes(TEXT("Iad/InMemoryBytes"), TraceCounterDisplayHint_Memory)
-		, InMemoryCount(TEXT("Iad/InMemoryCount"), TraceCounterDisplayHint_None)
-		, ReadBytes(TEXT("Iad/ReadBytes"), TraceCounterDisplayHint_Memory)
-		, ReadCount(TEXT("Iad/ReadCount"), TraceCounterDisplayHint_None)
+		: MaxCacheSize(TEXT("Iax/InstallCache/MaxCacheSize"), TraceCounterDisplayHint_Memory)
+		, CacheSize(TEXT("Iax/InstallCache/CacheSize"), TraceCounterDisplayHint_Memory)
+		, FlushedBytes(TEXT("Iax/IntallCache/WriteBytes"), TraceCounterDisplayHint_Memory)
+		, FlushCount(TEXT("Iax/InstallCache/WriteCount"), TraceCounterDisplayHint_None)
+		, InMemoryBytes(TEXT("Iax/InstallCache/InMemoryBytes"), TraceCounterDisplayHint_Memory)
+		, InMemoryCount(TEXT("Iax/InstallCache/InMemoryCount"), TraceCounterDisplayHint_None)
+		, ReadBytes(TEXT("Iax/InstallCache/ReadBytes"), TraceCounterDisplayHint_Memory)
+		, ReadCount(TEXT("Iax/InstallCache/ReadCount"), TraceCounterDisplayHint_None)
+		, PurgeCount(TEXT("Iax/InstallCache/PurgeCount"), TraceCounterDisplayHint_None)
+		, DefragCount(TEXT("Iax/InstallCache/DefragCount"), TraceCounterDisplayHint_None)
+		, ReadErrorCount(TEXT("Iax/InstallCache/ReadErrorCount"), TraceCounterDisplayHint_None)
+		, ReadDecodeErrorCount(TEXT("Iax/InstallCache/ReadDecodeErrorCount"), TraceCounterDisplayHint_None)
 	{ }
 
 	FCounterAtomicInt MaxCacheSize;
@@ -632,6 +625,10 @@ struct FInstallCacheCounters
 	FCounterAtomicInt InMemoryCount;
 	FCounterAtomicInt ReadBytes;
 	FCounterAtomicInt ReadCount;
+	FCounterAtomicInt PurgeCount;
+	FCounterAtomicInt DefragCount;
+	FCounterAtomicInt ReadErrorCount;
+	FCounterAtomicInt ReadDecodeErrorCount;
 };
 FInstallCacheCounters InstallCacheCounters;
 
@@ -699,27 +696,37 @@ void FOnDemandIoBackendStats::UpdateCSVValues()
 {
 	static constexpr float OneOver1024 = 1.0f / 1024.0f;
 
-	// cache stat totals
-	int32 CGetCount = int32(GCacheGetCount.Get());
-	int32 CErrorCount = int32(GCacheErrorCount.Get());
-	int32 CPutCount = int32(GCachePutCount.Get());
-	int32 CPutExistingCount = int32(GCachePutExistingCount.Get());
-	int32 CPutRejectCount = int32(GCachePutRejectCount.Get());
-	float CCachedKiB = (float)GCacheCachedBytes.Get() * OneOver1024;
-	float CWrittenKiB = (float)GCacheWrittenBytes.Get() * OneOver1024;
-	float CEvictedKiB = (float)GCacheEvictedBytes.Get() * OneOver1024;
-	float CReadKiB = (float)GCacheReadBytes.Get() * OneOver1024;
-	float CRejectedKiB = (float)GCacheRejectBytes.Get() * OneOver1024;
-	CSV_CUSTOM_STAT_DEFINED(CacheGetCount, CGetCount, ECsvCustomStatOp::Set);
-	CSV_CUSTOM_STAT_DEFINED(CacheErrorCount, CErrorCount, ECsvCustomStatOp::Set);
-	CSV_CUSTOM_STAT_DEFINED(CachePutCount, CPutCount, ECsvCustomStatOp::Set);
-	CSV_CUSTOM_STAT_DEFINED(CachePutExistingCount, CPutExistingCount, ECsvCustomStatOp::Set);
-	CSV_CUSTOM_STAT_DEFINED(CachePutRejectCount, CPutRejectCount, ECsvCustomStatOp::Set);
-	CSV_CUSTOM_STAT_DEFINED(CacheCachedMB, CCachedKiB * OneOver1024, ECsvCustomStatOp::Set);
-	CSV_CUSTOM_STAT_DEFINED(CacheWrittenMB, CWrittenKiB * OneOver1024, ECsvCustomStatOp::Set);
-	CSV_CUSTOM_STAT_DEFINED(CacheEvictedMB, CEvictedKiB * OneOver1024, ECsvCustomStatOp::Set);
-	CSV_CUSTOM_STAT_DEFINED(CacheReadMB, CReadKiB * OneOver1024, ECsvCustomStatOp::Set);
-	CSV_CUSTOM_STAT_DEFINED(CacheRejectedMB, CRejectedKiB * OneOver1024, ECsvCustomStatOp::Set);
+	// Streaming cache stats
+	{
+		CSV_CUSTOM_STAT_DEFINED(StreamingCacheWriteCount, int32(GCacheWriteCount.Get()), ECsvCustomStatOp::Set);
+		CSV_CUSTOM_STAT_DEFINED(StreamingCacheWriteMiB, BytesToApproxMB(GCacheWrittenBytes.Get()), ECsvCustomStatOp::Set);
+		CSV_CUSTOM_STAT_DEFINED(StreamingCacheDecodeErrorCount, int32(GCacheDecodeErrorCount.Get()), ECsvCustomStatOp::Set);
+	}
+
+	// Install cache stats
+	{
+		CSV_CUSTOM_STAT_DEFINED(InstallCachePurgeCount, int32(InstallCacheCounters.PurgeCount.Get()), ECsvCustomStatOp::Set);
+		CSV_CUSTOM_STAT_DEFINED(InstallCacheDefragCount, int32(InstallCacheCounters.DefragCount.Get()), ECsvCustomStatOp::Set);
+		CSV_CUSTOM_STAT_DEFINED(InstallCacheWriteCount, int32(InstallCacheCounters.FlushCount.Get()), ECsvCustomStatOp::Set);
+		CSV_CUSTOM_STAT_DEFINED(InstallCacheWriteMiB, BytesToApproxMB(InstallCacheCounters.FlushedBytes.Get()), ECsvCustomStatOp::Set);
+		CSV_CUSTOM_STAT_DEFINED(InstallCacheDecodeErrorCount, int32(InstallCacheCounters.ReadDecodeErrorCount.Get()), ECsvCustomStatOp::Set);
+	}
+
+	// Installer stats
+	{
+		const int64 InstallCount		= InstallerCounters.InstallCount.Get();
+		const int64 InflightCount		= InstallerCounters.InflightCount.Get();
+		const int64 DownloadedBytes		= InstallerCounters.DownloadedBytes.Get();
+		const uint64 DurationTotalMs	= InstallerCounters.DurationTotalMs;
+		const double DurationAvgMs		= InstallCount > 0 ? double(DurationTotalMs) / double(InstallCount) : 0;
+		const int64 CacheHitTotalPct 	= InstallerCounters.CacheHitTotalPct;
+		const double CacheHitAvgPct		= InstallCount > 0 ? double(CacheHitTotalPct) / double(InstallCount) : 0;
+		CSV_CUSTOM_STAT_DEFINED(InstallerInstallCount, int32(InstallCount), ECsvCustomStatOp::Set);
+		CSV_CUSTOM_STAT_DEFINED(InstallerInflightCount, int32(InflightCount), ECsvCustomStatOp::Set);
+		CSV_CUSTOM_STAT_DEFINED(InstallerDownloadedMiB, BytesToApproxMB(DownloadedBytes), ECsvCustomStatOp::Set);
+		CSV_CUSTOM_STAT_DEFINED(InstallerDurationAvgMs, int32(FMath::RoundToInt64(DurationAvgMs)), ECsvCustomStatOp::Set);
+		CSV_CUSTOM_STAT_DEFINED(InstallCacheAvgCacheHitPct, int32(FMath::RoundToInt64(CacheHitAvgPct)), ECsvCustomStatOp::Set);
+	}
 
 	// http stat totals
 	auto HttpCsv = [](EHttpRequestType Type)
@@ -748,6 +755,9 @@ void FOnDemandIoBackendStats::PrintPeriodicLogging()
 	static double LastLogTime = 0.0;
 	if (double Time = FPlatformTime::Seconds(); Time - LastLogTime > (double)GIasStatisticsLogInterval)
 	{
+		UE_LOG(LogIoStoreOnDemand, Log, TEXT("Stats:"));
+		UE_LOG(LogIoStoreOnDemand, Log, TEXT("==============="));
+
 		static constexpr float OneOver1024 = 1.0f / 1024.0f;
 
 		int32 CGetCount = int32(GCacheGetCount.Get());
@@ -764,18 +774,18 @@ void FOnDemandIoBackendStats::PrintPeriodicLogging()
 
 		if (BackendStatus.IsCacheEnabled())
 		{
-			UE_LOG(LogIas, Log, TEXT("IASCache: Count=%d, CachedKiB=%d, WrittenKiB=%d, ReadKiB=%d, EvictedKiB=%d, RejectedKiB=%d, Get=%d, Error=%d, Put=%d, PutReject=%d, PutExisting=%d"),
+			UE_LOG(LogIoStoreOnDemand, Log, TEXT("StreamingCache: Count=%d, CachedKiB=%d, WrittenKiB=%d, ReadKiB=%d, EvictedKiB=%d, RejectedKiB=%d, Get=%d, Error=%d, Put=%d, PutReject=%d, PutExisting=%d"),
 				(int32)GCacheCachedEntries.Get(), (int32)CCachedKiB, (int32)CWrittenKiB, (int32)CReadKiB,
 				(int32)CEvictedKiB, (int32)CRejectedKiB, CGetCount, CErrorCount, CPutCount, CPutRejectCount, CPutExistingCount);
 		}
 		else
 		{
-			UE_LOG(LogIas, Log, TEXT("CacheStats: Disabled"));
+			UE_LOG(LogIoStoreOnDemand, Log, TEXT("StreamingCache: Disabled"));
 		}
 
 		auto HttpLog = [](const TCHAR* Title, EHttpRequestType Type)
 			{
-				UE_LOG(LogIas, Log, TEXT("%sHttp : DownloadedKiB=%d, Get=%d, Retry=%d, Cancel=%d, Error=%d, CurPending=%d, CurDurationMsAvg=%d, CurDurationMsMax=%d"),
+				UE_LOG(LogIoStoreOnDemand, Log, TEXT("Http: Category: %s, DownloadedKiB=%d, Get=%d, Retry=%d, Cancel=%d, Error=%d, CurPending=%d, CurDurationMsAvg=%d, CurDurationMsMax=%d"),
 					Title,
 					(int32)((float)GHttpDownloadedBytes[(uint8)Type].Get() * OneOver1024),
 					(int32)GHttpGetCount[(uint8)Type].Get(),
@@ -787,19 +797,26 @@ void FOnDemandIoBackendStats::PrintPeriodicLogging()
 					GHttpDurationMsMax[(uint8)Type]);
 			};
 
-		HttpLog(TEXT("IAS"), EHttpRequestType::Streaming);
-		HttpLog(TEXT("IAD"), EHttpRequestType::Installed);
+		HttpLog(TEXT("Streaming"), EHttpRequestType::Streaming);
+		HttpLog(TEXT("Install"), EHttpRequestType::Installed);
 
-#if COUNTERSTRACE_ENABLED
 		{
-			TUniqueLock Lock(InstallerTraceCounters);
-			UE_LOG(LogIoStoreOnDemand, Log, TEXT("IadStats: InstallCount=%llu, Downloaded=%d KiB, AvgInstallDuration=%dms, AvgCacheHitRatio=%d%%"),
-				InstallerTraceCounters.InstallCount.Get(),
-				(int32(InstallerTraceCounters.DownloadedBytes.Get()) / 1024),
-				int32(InstallerTraceCounters.AvgInstallDurationMs.Get()),
-				int32(InstallerTraceCounters.CacheHitRatio.GetMean() * 100.0));
+			const int64 InstallCount		= InstallerCounters.InstallCount.Get();
+			const uint64 DurationTotalMs	= InstallerCounters.DurationTotalMs;
+			const double DurationAvgMs		= InstallCount > 0 ? double(DurationTotalMs) / double(InstallCount) : 0;
+			const int64 CacheHitTotalPct 	= InstallerCounters.CacheHitTotalPct;
+			const double CacheHitAvgPct		= InstallCount > 0 ? double(CacheHitTotalPct) / double(InstallCount) : 0;
+
+			UE_LOG(LogIoStoreOnDemand, Log, TEXT("Installer: InstallCount=%lld, AvgDuration: %d, AvgCacheHitPct=%.2lf"),
+				InstallCount, int32(DurationAvgMs), CacheHitAvgPct);
+		}
+
+		{
+			const int64 WriteCount = InstallCacheCounters.FlushCount.Get();
+			const int64 WriteBytes = InstallCacheCounters.FlushedBytes.Get();
+			UE_LOG(LogIoStoreOnDemand, Log, TEXT("InstallCache: WriteCount=%lld, WriteBytes=%.2lf MiB"),
+				WriteCount, double(WriteBytes) / 1024.0 / 1024.0);
 		}
-#endif // COUNTERSTRACE_ENABLED
 
 		LastLogTime = Time;
 	}
@@ -1343,33 +1360,6 @@ void FOnDemandIoBackendStats::GetIasCacheStats(uint64& OutUsed, uint64& OutMaxSi
 	OutMaxSize = GCacheMaxBytes;
 }
 
-void FOnDemandIoBackendStats::OnIoRequestEnqueue()
-{
-	GIoRequestCount.Increment();
-	CSV_CUSTOM_STAT_DEFINED(FrameIoRequestCount, int32(GIoRequestCount.Get()), ECsvCustomStatOp::Set);
-}
-
-void FOnDemandIoBackendStats::OnIoRequestComplete(uint64 Size, uint64 Duration)
-{
-	GIoRequestReadCount.Increment();
-	GIoRequestReadBytes.Add(Size);
-
-	CSV_CUSTOM_STAT_DEFINED(FrameIoRequestReadCount, int32(GIoRequestReadCount.Get()), ECsvCustomStatOp::Set);
-	CSV_CUSTOM_STAT_DEFINED(FrameIoRequestReadMB, BytesToApproxMB(GIoRequestReadBytes.Get()), ECsvCustomStatOp::Set);
-}
-
-void FOnDemandIoBackendStats::OnIoRequestCancel()
-{
-	GIoRequestCancelCount.Increment();
-	CSV_CUSTOM_STAT_DEFINED(FrameIoRequestCancelCount, int32(GIoRequestCancelCount.Get()), ECsvCustomStatOp::Set);
-}
-
-void FOnDemandIoBackendStats::OnIoRequestError()
-{
-	GIoRequestErrorCount.Increment();
-	CSV_CUSTOM_STAT_DEFINED(FrameIoRequestErrorCount, int32(GIoRequestErrorCount.Get()), ECsvCustomStatOp::Set);
-}
-
 void FOnDemandIoBackendStats::OnCacheBootMs(uint64 TimeMs)
 {
 	GCacheBootMs = uint32(TimeMs);
@@ -1432,6 +1422,7 @@ void FOnDemandIoBackendStats::OnCachePersistedBytes(uint64 TotalSize, int32 NumE
 
 void FOnDemandIoBackendStats::OnCacheWriteBytes(uint64 WriteSize)
 {
+	GCacheWriteCount.Increment();
 	GCacheWrittenBytes.Add(WriteSize);
 }
 
@@ -1556,10 +1547,7 @@ void FOnDemandIoBackendStats::OnApplicationResume()
 ////////////////////////////////////////////////////////////////////////////////
 void FOnDemandContentInstallerStats::OnRequestEnqueued()
 {
-#if COUNTERSTRACE_ENABLED
-	TUniqueLock Lock(InstallerTraceCounters);
-	InstallerTraceCounters.InflightInstallCount.Increment();
-#endif
+	InstallerCounters.InflightCount.Increment();
 }
 
 void FOnDemandContentInstallerStats::OnRequestCompleted(
@@ -1571,29 +1559,25 @@ void FOnDemandContentInstallerStats::OnRequestCompleted(
 		double CacheHitRatio,
 		uint64 DurationCycles) 
 {
-#if COUNTERSTRACE_ENABLED
-	{
-		TUniqueLock Lock(InstallerTraceCounters);
-		InstallerTraceCounters.InflightInstallCount.Decrement();
-	}
-#endif // COUNTERSTRACE_ENABLED
+	InstallerCounters.InflightCount.Decrement();
 
 	if (RequestedChunkCount == 0)
 	{
 		return;
 	}
 
-#if COUNTERSTRACE_ENABLED
 	{
-		TUniqueLock Lock(InstallerTraceCounters);
-		InstallerTraceCounters.InstallCount.Increment();
-		InstallerTraceCounters.DownloadedBytes.Add(DownloadedBytes);
-		InstallerTraceCounters.InstallDurationMs.Increment(FPlatformTime::ToMilliseconds64(DurationCycles));
-		InstallerTraceCounters.CacheHitRatio.Increment(CacheHitRatio);
-		InstallerTraceCounters.AvgInstallDurationMs.Set(InstallerTraceCounters.InstallDurationMs.GetMean());
-		InstallerTraceCounters.AvgCacheHitRatio.Set(InstallerTraceCounters.CacheHitRatio.GetMean());
+		InstallerCounters.InstallCount.Increment();
+		const int64 CurrentInstallCount = InstallerCounters.InstallCount.Get();
+		InstallerCounters.DownloadedBytes.Add(DownloadedBytes);
+		InstallerCounters.DurationTotalMs.fetch_add(uint64(FPlatformTime::ToMilliseconds64(DurationCycles)));
+		InstallerCounters.DurationAvgMs.Set(
+			int64(double(InstallerCounters.DurationTotalMs) / double(CurrentInstallCount)));
+		const int64 CacheHitPct = FMath::RoundToInt64(100.0 * CacheHitRatio);
+		const int64 TotalCacheHitPct = CacheHitPct + InstallerCounters.CacheHitTotalPct.fetch_add(CacheHitPct);
+		const double AvgCacheHitPec = double(TotalCacheHitPct) / double(CurrentInstallCount);
+		InstallerCounters.CacheHitAvgPct.Set(FMath::Clamp(FMath::RoundToInt64(AvgCacheHitPec), 0, 100));
 	}
-#endif // COUNTERSTRACE_ENABLED
 
 	{
 		TUniqueLock Lock(InstallerAnalytics);
@@ -1809,6 +1793,8 @@ void FOnDemandInstallCacheStats::OnPurge(
 		return;
 	}
 
+	InstallCacheCounters.PurgeCount.Increment();
+
 	{
 		UE::TUniqueLock Lock(InstallCacheAnalytics);
 		if (Result.HasError())
@@ -1840,6 +1826,8 @@ void FOnDemandInstallCacheStats::OnDefrag(
 		return;
 	}
 
+	InstallCacheCounters.DefragCount.Increment();
+
 	const double TotalFragmentionRatio = TotalCachedBytes > 0
 		? double(TotalFragmentedBytes) / double(TotalCachedBytes)
 		: 0.0;
@@ -1945,7 +1933,13 @@ void FOnDemandInstallCacheStats::OnReadCompleted(EIoErrorCode ErrorCode, uint64
 	UE::TUniqueLock Lock(InstallCacheAnalytics);
 	if (ErrorCode != EIoErrorCode::Ok)
 	{
+		InstallCacheCounters.ReadErrorCount.Increment();
 		++InstallCacheAnalytics.ReadErrorCount;
+		
+		if (ErrorCode == EIoErrorCode::CompressionError)
+		{
+			InstallCacheCounters.ReadDecodeErrorCount.Increment();
+		}
 	}
 	++InstallCacheAnalytics.ReadCount;
 }
diff --git a/Engine/Source/Runtime/Experimental/IoStore/OnDemand/Private/Statistics.h b/Engine/Source/Runtime/Experimental/IoStore/OnDemand/Private/Statistics.h
--- a/Engine/Source/Runtime/Experimental/IoStore/OnDemand/Private/Statistics.h
+++ b/Engine/Source/Runtime/Experimental/IoStore/OnDemand/Private/Statistics.h
@@ -79,11 +79,6 @@ public:
 	TUniquePtr<IAnalyticsRecording> StartAnalyticsRecording() const IAS_STATISTICS_IMPL(TUniquePtr<IAnalyticsRecording>())
 	void GetIasCacheStats(uint64& OutUsed, uint64& OutMaxSize) const IAS_STATISTICS_IMPL(void(OutUsed = OutMaxSize = 0);) 
 
-	void OnIoRequestEnqueue() IAS_STATISTICS_IMPL()
-	void OnIoRequestComplete(uint64 Size, uint64 DurationMs) IAS_STATISTICS_IMPL()
-	void OnIoRequestCancel() IAS_STATISTICS_IMPL()
-	void OnIoRequestError() IAS_STATISTICS_IMPL()
-
 	void OnCacheBootMs(uint64 TimeMs) IAS_STATISTICS_IMPL()
 	void OnCacheError() IAS_STATISTICS_IMPL()
 	void OnCacheDecodeError() IAS_STATISTICS_IMPL()
