IAD timestamp hardening
- Fix some cases  that could have added a default constructed FCasBlockState with 0 timestamps.
-- In order to do that properly, FCas::OpenAsyncRead needs to return a TResult, so I changed FCasErrorContext to include FFileSystemError.
-- Also added FFileSystemError to the error context in all the other places where it could be missing.
- Add some additional timestamp verification and logging to VerifyTimestamp during CAS startup.

#Synced-CL 50847432

diff --git a/Engine/Source/Runtime/Experimental/IoStore/OnDemand/Private/OnDemandInstallCache.cpp b/Engine/Source/Runtime/Experimental/IoStore/OnDemand/Private/OnDemandInstallCache.cpp
--- a/Engine/Source/Runtime/Experimental/IoStore/OnDemand/Private/OnDemandInstallCache.cpp
+++ b/Engine/Source/Runtime/Experimental/IoStore/OnDemand/Private/OnDemandInstallCache.cpp
@@ -114,7 +114,6 @@ using FSharedFileHandle				= TSharedPtr<IFileHandle>;
 
 using FSharedAsyncFileHandle		= TSharedPtr<IAsyncReadFileHandle>;
 using FWeakAsyncFileHandle			= TWeakPtr<IAsyncReadFileHandle>;
-using FSharedFileOpenAsyncResult	= TValueOrError<FSharedAsyncFileHandle, FFileSystemError>;
 
 using FCasAddr						= FHash96;
 static const FCasAddr& AsCasAddr(const FIoHash& IoHash)
@@ -404,12 +403,14 @@ FResult FCasJournal::Replay(const FString& JournalFile, FEntryHandler&& Handler)
 		return MakeJournalError(ECasErrorCode::ReplayJournalFailed, EIoErrorCode::NotFound, FString::Printf(TEXT("Failed to find '%s'"), *JournalFile));
 	}
 
-	TUniquePtr<IFileHandle> FileHandle(Ipf.OpenRead(*JournalFile));
-	if (FileHandle.IsValid() == false)
+	FFileOpenResult OpenReadResult = Ipf.OpenRead(*JournalFile, IPlatformFile::EOpenReadFlags::None);
+	if (OpenReadResult.HasError())
 	{
-		return MakeJournalError(ECasErrorCode::ReplayJournalFailed, EIoErrorCode::FileOpenFailed, FString::Printf(TEXT("Failed to open '%s'"), *JournalFile));
+		return MakeJournalError(ECasErrorCode::ReplayJournalFailed, EIoErrorCode::FileOpenFailed, FString::Printf(TEXT("Failed to open '%s'"), *JournalFile), OpenReadResult.StealError());
 	}
 
+	TUniquePtr<IFileHandle> FileHandle = OpenReadResult.StealValue();
+
 	FHeader Header;
 	if ((FileHandle->Read(reinterpret_cast<uint8*>(&Header), FHeader::Size()) == false) || (Header.IsValid() == false))
 	{
@@ -493,12 +494,17 @@ FResult FCasJournal::Create(const FString& JournalFile)
 	IPlatformFile& Ipf = FPlatformFileManager::Get().GetPlatformFile();
 	Ipf.DeleteFile(*JournalFile);
 
-	TUniquePtr<IFileHandle> FileHandle(Ipf.OpenWrite(*JournalFile));
-	if (FileHandle.IsValid() == false)
+	FFileOpenResult OpenWriteResult(Ipf.OpenWrite(*JournalFile, IPhysicalPlatformFile::EOpenWriteFlags::None));
+	if (OpenWriteResult.HasError())
 	{
-		return MakeJournalError(ECasErrorCode::CreateJournalFailed, EIoErrorCode::FileOpenFailed, FString::Printf(TEXT("Failed to create journal '%s'"), *JournalFile));
+		return MakeJournalError(
+			ECasErrorCode::CreateJournalFailed, EIoErrorCode::FileOpenFailed, 
+			FString::Printf(TEXT("Failed to create journal '%s'"), *JournalFile),
+			OpenWriteResult.StealError());
 	}
 
+	TUniquePtr<IFileHandle> FileHandle = OpenWriteResult.StealValue();
+
 	FHeader Header;
 	FMemory::Memcpy(&Header.Magic, &FHeader::MagicSequence, sizeof(FHeader::MagicSequence));
 	Header.Version = EVersion::Latest;
@@ -542,15 +548,19 @@ FResult FCasJournal::Commit(FTransaction&& Transaction, uint64& OutByteCount, ui
 
 	// Validate header and footer
 	{
-		TUniquePtr<IFileHandle> FileHandle(Ipf.OpenRead(*Transaction.JournalFile));
-		if (FileHandle.IsValid() == false)
+		FFileOpenResult OpenReadResult = Ipf.OpenRead(*Transaction.JournalFile, IPlatformFile::EOpenReadFlags::None);
+		if (OpenReadResult.HasError())
 		{
 			FResult Result = MakeJournalError(
-				ECasErrorCode::CommitJournalFailed, EIoErrorCode::FileOpenFailed, FString::Printf(TEXT("Failed to open journal '%s'"), *Transaction.JournalFile));
+				ECasErrorCode::CommitJournalFailed, EIoErrorCode::FileOpenFailed, 
+				FString::Printf(TEXT("Failed to open journal '%s'"), *Transaction.JournalFile),
+				OpenReadResult.StealError());
 			FOnDemandInstallCacheStats::OnJournalCommit(Result, 0);
 			return Result;
 		}
 
+		TUniquePtr<IFileHandle> FileHandle = OpenReadResult.StealValue();
+
 		const int64 FileSize = FileHandle->Size();
 		if (FileSize < FHeader::Size())
 		{
@@ -590,15 +600,17 @@ FResult FCasJournal::Commit(FTransaction&& Transaction, uint64& OutByteCount, ui
 
 	// Append entries
 	{
-		const bool				bAppend = true;
-		TUniquePtr<IFileHandle> FileHandle(Ipf.OpenWrite(*Transaction.JournalFile, bAppend));
+		FFileOpenResult			OpenWriteResult(Ipf.OpenWrite(*Transaction.JournalFile, IPhysicalPlatformFile::EOpenWriteFlags::Append));
+		TUniquePtr<IFileHandle>	FileHandle = OpenWriteResult.HasValue() ? OpenWriteResult.StealValue() : nullptr;
 		const int64				FileSize = FileHandle.IsValid() ? FileHandle->Size() : -1;
 		const int64				EntriesPos = FileSize > 0 ? FileSize - FFooter::Size() : -1;
 
 		if (EntriesPos < 0)
 		{
 			FResult Result = MakeJournalError(
-				ECasErrorCode::CommitJournalFailed, EIoErrorCode::FileOpenFailed, FString::Printf(TEXT("Failed to open journal '%s'"), *Transaction.JournalFile));
+				ECasErrorCode::CommitJournalFailed, EIoErrorCode::FileOpenFailed, 
+				FString::Printf(TEXT("Failed to open journal '%s'"), *Transaction.JournalFile),
+				OpenWriteResult.HasError() ? OpenWriteResult.StealError() : FFileSystemError(FString()));
 			FOnDemandInstallCacheStats::OnJournalCommit(Result, 0);
 			return Result;
 		}
@@ -753,7 +765,7 @@ struct FCas
 	FResult						DeleteBlock(FCasBlockId BlockId, TFunctionRef<void(const FCasAddr&)> OutAddrs);
 	FString						GetBlockFilename(FCasBlockId BlockId) const;
 	TResult<FSharedFileHandle>	OpenRead(FCasBlockId BlockId);
-	FSharedFileOpenAsyncResult	OpenAsyncRead(FCasBlockId BlockId);
+	TResult<FSharedAsyncFileHandle>	OpenAsyncRead(FCasBlockId BlockId);
 	void						OnFileHandleDeleted(FCasBlockId BlockId);
 	TResult<FUniqueFileHandle>	OpenWrite(FCasBlockId BlockId, bool bAppend) const;
 	bool						TrackAccessIf(ECasTrackAccessType Type, FCasBlockId BlockId, int64 UtcTicks, bool bDirty);
@@ -788,6 +800,9 @@ struct FCas
 	// Returns the place holder timestamp to use if a timestamp is not found in the LastAccess table
 	static int64				GetTimestampForMissingLastAccess();
 
+	// Return the minimum value for a valid timestamp
+	static int64				GetMinimumTimestamp();
+
 	uint32						GetMinBlockSize()
 	{
 		return MinBlockSize;
@@ -1020,10 +1035,17 @@ TResult<FSharedFileHandle> FCas::OpenRead(FCasBlockId BlockId)
 
 	UE::TUniqueLock Lock(Mutex);
 
+	FCasBlockState* BlockState = BlockIds.Find(BlockId);
+	if (BlockState == nullptr)
+	{
+		return MakeCasError<FSharedFileHandle>(
+			ECasErrorCode::ReadBlockFailed, EIoErrorCode::NotFound, FString::Printf(TEXT("Could not find block state for '%s'"), *Filename));
+	}
+
 	FFileOpenResult Result = Ipf.OpenRead(*Filename, IPlatformFile::EOpenReadFlags::AllowWrite);
 	if (Result.HasValue())
 	{
-		BlockIds.FindOrAdd(BlockId).ReadHandleCount++;
+		BlockState->ReadHandleCount++;
 
 		FSharedFileHandle NewHandle(
 			Result.GetValue().Release(),
@@ -1038,10 +1060,12 @@ TResult<FSharedFileHandle> FCas::OpenRead(FCasBlockId BlockId)
 	}
 
 	return MakeCasError<FSharedFileHandle>(
-		ECasErrorCode::ReadBlockFailed, EIoErrorCode::FileOpenFailed, FString::Printf(TEXT("Failed to open CAS block '%s'"), *Filename));
+		ECasErrorCode::ReadBlockFailed, EIoErrorCode::FileOpenFailed, 
+		FString::Printf(TEXT("Failed to open CAS block '%s'"), *Filename),
+		Result.StealError());
 }
 
-FSharedFileOpenAsyncResult FCas::OpenAsyncRead(FCasBlockId BlockId)
+TResult<FSharedAsyncFileHandle> FCas::OpenAsyncRead(FCasBlockId BlockId)
 {
 	UE::TUniqueLock Lock(Mutex);
 
@@ -1053,13 +1077,23 @@ FSharedFileOpenAsyncResult FCas::OpenAsyncRead(FCasBlockId BlockId)
 		}
 	}
 
-	IPlatformFile&			Ipf = FPlatformFileManager::Get().GetPlatformFile();
+	using namespace UE::UnifiedError::IoStoreOnDemand;
+
 	const FString			Filename = GetBlockFilename(BlockId);
+
+	FCasBlockState* BlockState = BlockIds.Find(BlockId);
+	if (BlockState == nullptr)
+	{
+		return MakeCasError<FSharedAsyncFileHandle>(
+			ECasErrorCode::ReadBlockFailed, EIoErrorCode::NotFound, FString::Printf(TEXT("Could not find block state for '%s'"), *Filename));
+	}
+
+	IPlatformFile&			Ipf = FPlatformFileManager::Get().GetPlatformFile();
 	FFileOpenAsyncResult	HandleResult(Ipf.OpenAsyncRead(*Filename, IPlatformFile::EOpenReadFlags::AllowWrite));
 
 	if (HandleResult.HasValue())
 	{
-		BlockIds.FindOrAdd(BlockId).ReadHandleCount++;
+		BlockState->ReadHandleCount++;
 
 		FSharedAsyncFileHandle NewHandle(
 			HandleResult.GetValue().Release(),
@@ -1074,7 +1108,10 @@ FSharedFileOpenAsyncResult FCas::OpenAsyncRead(FCasBlockId BlockId)
 		return MakeValue(MoveTemp(NewHandle));
 	}
 
-	return MakeError(HandleResult.StealError());
+	return MakeCasError<FSharedAsyncFileHandle>(
+		ECasErrorCode::ReadBlockFailed, EIoErrorCode::FileOpenFailed,
+		FString::Printf(TEXT("Failed to open async CAS block '%s'"), *Filename),
+		HandleResult.StealError());
 }
 
 void FCas::OnFileHandleDeleted(FCasBlockId BlockId)
@@ -1094,16 +1131,23 @@ TResult<FUniqueFileHandle> FCas::OpenWrite(FCasBlockId BlockId, const bool bAppe
 
 	IPlatformFile&	Ipf = FPlatformFileManager::Get().GetPlatformFile();
 	const FString	Filename = GetBlockFilename(BlockId);
-	const bool		bAllowRead = true;
 
-	FUniqueFileHandle FileHandle(Ipf.OpenWrite(*Filename, bAppend, bAllowRead));
-	if (FileHandle.IsValid())
+	IPhysicalPlatformFile::EOpenWriteFlags OpenWriteFlags = IPhysicalPlatformFile::EOpenWriteFlags::AllowRead;
+	if (bAppend)
 	{
-		return MakeValue(MoveTemp(FileHandle));
+		OpenWriteFlags |= IPhysicalPlatformFile::EOpenWriteFlags::Append;
+	}
+
+	FFileOpenResult OpenWriteResult(Ipf.OpenWrite(*Filename, OpenWriteFlags));
+	if (OpenWriteResult.HasValue())
+	{
+		return MakeValue(OpenWriteResult.StealValue());
 	}
 
 	return MakeCasError<FUniqueFileHandle>(
-		ECasErrorCode::WriteBlockFailed, EIoErrorCode::FileOpenFailed, FString::Printf(TEXT("Failed to open CAS block '%s' for writing"), *Filename));
+		ECasErrorCode::WriteBlockFailed, EIoErrorCode::FileOpenFailed, 
+		FString::Printf(TEXT("Failed to open CAS block '%s' for writing"), *Filename),
+		OpenWriteResult.StealError());
 }
 
 bool FCas::TrackAccessIf(const ECasTrackAccessType Type, const FCasBlockId BlockId, const int64 UtcTicks, const bool bDirty)
@@ -1356,49 +1400,73 @@ FResult FCas::Verify(TArray<FCasAddr>& OutRemovedChunksAddrs, TMap<FCasBlockId,
 	FResult					Result = MakeValue();
 
 	const int64 Now = FDateTime::UtcNow().GetTicks();
-	const int64 TimestampForMissingLastAccess = FCas::GetTimestampForMissingLastAccess();
 
 	IPlatformFile& Ipf = FPlatformFileManager::Get().GetPlatformFile();
 
-	auto VerifyTimestamp = [this, &Ipf, Now, TimestampValueIfMissing = TimestampForMissingLastAccess](
+	auto VerifyTimestamp = [this, &Ipf, Now](
 		int64& InOutUtcTicks, const FCasBlockId BlockId, const TCHAR* TimestampName)
 	{
-		if (InOutUtcTicks)
+		const int64 TimestampValueIfMissing = FCas::GetTimestampForMissingLastAccess();
+		const int64 MinimumTimestamp = FCas::GetMinimumTimestamp();
+
+		const FString Filename = FCas::GetBlockFilename(BlockId);
+
+		bool bTimestampOk = false;
+		if (InOutUtcTicks == 0)
 		{
-			if (InOutUtcTicks > Now)
-			{
-				const FString Filename = FCas::GetBlockFilename(BlockId);
-				UE_LOG(LogIoStoreOnDemand, Warning, TEXT("Found future %s time for CAS block '%s'"), TimestampName, *Filename);
-				InOutUtcTicks = Now;
-			}
+			UE_LOG(LogIoStoreOnDemand, Warning, TEXT("Found missing %s time for CAS block '%s'"), TimestampName, *Filename);
+		}
+		else if (InOutUtcTicks < 0)
+		{
+			UE_LOG(LogIoStoreOnDemand, Warning, TEXT("Found negative %s time for CAS block '%s'"), TimestampName, *Filename);
+		}
+		else if (InOutUtcTicks < MinimumTimestamp)
+		{
+			UE_LOG(LogIoStoreOnDemand, Warning, TEXT("Found invalid %s time for CAS block '%s'"), TimestampName, *Filename);
+		}
+		else if (InOutUtcTicks > Now)
+		{
+			UE_LOG(LogIoStoreOnDemand, Warning, TEXT("Found future %s time for CAS block '%s'"), TimestampName, *Filename);
+			InOutUtcTicks = Now;
+			bTimestampOk = true;
 		}
 		else
 		{
-			const FString Filename = FCas::GetBlockFilename(BlockId);
+			bTimestampOk = true;
+		}
+
+		if (bTimestampOk)
+		{
+			return;
+		}
 
 	#if UE_ONDEMANDINSTALLCACHE_USE_MODTIME
-			FDateTime ModTime = Ipf.GetTimeStamp(*Filename);
-			if (ensure(ModTime > FDateTime::MinValue()))
+		FDateTime ModTime = Ipf.GetTimeStamp(*Filename);
+		if (ensure(ModTime > FDateTime::MinValue()))
+		{
+			if (ModTime < MinimumTimestamp)
 			{
-				if (ModTime > Now)
-				{
-					UE_LOG(LogIoStoreOnDemand, Warning, TEXT("Missing %s time and found future mod time so using now time for CAS block '%s'"), TimestampName, *Filename);
-					InOutUtcTicks = Now;
-				}
-				else
-				{
-					UE_LOG(LogIoStoreOnDemand, Warning, TEXT("Missing %s time so using file mod time for CAS block '%s'"), TimestampName, *Filename);
-					InOutUtcTicks = ModTime.GetTicks();
-				}
+				UE_LOG(LogIoStoreOnDemand, Warning, TEXT("Unusable %s time and found invalid file mod time for CAS block '%s'"), TimestampName, *Filename);
+				InOutUtcTicks = TimestampValueIfMissing;
+			}
+			else if (ModTime > Now)
+			{
+				UE_LOG(LogIoStoreOnDemand, Warning, TEXT("Unusable %s time and found future file mod time so using now time for CAS block '%s'"), TimestampName, *Filename);
+				InOutUtcTicks = Now;
 			}
 			else
-	#endif // UE_ONDEMANDINSTALLCACHE_USE_MODTIME
 			{
-				// Failed to get mod time
-				UE_LOG(LogIoStoreOnDemand, Warning, TEXT("Missing %s time for CAS block '%s'"), TimestampName, *Filename);
-				InOutUtcTicks = TimestampValueIfMissing;
+				UE_LOG(LogIoStoreOnDemand, Warning, TEXT("Unusable %s time so using file mod time %s for CAS block '%s'"), TimestampName, *ModTime.ToString(), *Filename);
+				InOutUtcTicks = ModTime.GetTicks();
 			}
 		}
+		else
+	#endif // UE_ONDEMANDINSTALLCACHE_USE_MODTIME
+		{
+			// Failed to get mod time
+			UE_LOG(LogIoStoreOnDemand, Warning, TEXT("Unusable %s time for CAS block '%s'"), TimestampName, *Filename);
+			InOutUtcTicks = TimestampValueIfMissing;
+		}
 	};
 
 	UE::TUniqueLock Lock(Mutex);
@@ -1511,6 +1579,13 @@ int64 FCas::GetTimestampForMissingLastAccess()
 	return FourWeeksAgo.GetTicks();
 }
 
+int64 FCas::GetMinimumTimestamp()
+{
+	// FCas didn't exist prior to 2024.4.24 so no date before then is valid.
+	FDateTime MinTime = FDateTime(2024, 4, 24);
+	return MinTime.GetTicks();
+}
+
 ///////////////////////////////////////////////////////////////////////////////
 static void JournalLastAccess(FCasJournal::FTransaction& Transaction, const FCas::FLastAccess& LastAccess)
 {
@@ -2670,7 +2745,7 @@ bool FOnDemandInstallCache::Resolve(FIoRequestImpl* Request)
 			}
 			else
 			{
-				UE_LOGFMT(LogIoStoreOnDemand, Error, "Failed to open CAS block '{Filename}' for reading ({Error})", *Filename, FileOpenResult.GetError());
+				UE_LOGFMT(LogIoStoreOnDemand, Error, "Failed to open CAS block '{Filename}' for reading ({Error})", *Filename, FileOpenResult.GetError().GetErrorMessage(true).ToString()));
 			}
 
 			UE::Tasks::Launch(UE_SOURCE_LOCATION, [this, Request, Status]
@@ -2683,11 +2758,11 @@ bool FOnDemandInstallCache::Resolve(FIoRequestImpl* Request)
 	}
 #endif // UE_ONDEMANDINSTALLCACHE_EXCLUSIVE_WRITE
 
-	FSharedFileOpenAsyncResult FileOpenResult = Cas.OpenAsyncRead(CasLoc.BlockId);
+	TResult<FSharedAsyncFileHandle> FileOpenResult = Cas.OpenAsyncRead(CasLoc.BlockId);
 	if (FileOpenResult.HasError())
 	{
 		const FString Filename = Cas.GetBlockFilename(CasLoc.BlockId);
-		UE_LOG(LogIoStoreOnDemand, Error, TEXT("Failed to open CAS block '%s' for async reading, reason '%s'"), *Filename, *FileOpenResult.GetError().GetMessage());
+		UE_LOG(LogIoStoreOnDemand, Error, TEXT("Failed to open CAS block '%s' for async reading, reason '%s'"), *Filename, *FileOpenResult.GetError().GetErrorMessage(true).ToString());
 		TUniquePtr<FChunkRequest> Detached = FChunkRequest::Detach(*Request);
 		CompleteRequest(Request, EIoErrorCode::FileOpenFailed);
 		return true;
diff --git a/Engine/Source/Runtime/Experimental/IoStore/OnDemandCore/Public/IO/OnDemandError.h b/Engine/Source/Runtime/Experimental/IoStore/OnDemandCore/Public/IO/OnDemandError.h
--- a/Engine/Source/Runtime/Experimental/IoStore/OnDemandCore/Public/IO/OnDemandError.h
+++ b/Engine/Source/Runtime/Experimental/IoStore/OnDemandCore/Public/IO/OnDemandError.h
@@ -4,6 +4,7 @@
 
 #include "Experimental/UnifiedError/UnifiedError.h"
 #include "Experimental/UnifiedError/CoreErrorTypes.h"
+#include "GenericPlatform/GenericPlatformFile.h"
 #include "IO/IoChunkId.h"
 #include "IO/IoStatus.h"
 #include "Templates/ValueOrError.h"
@@ -94,10 +95,11 @@ inline const TCHAR* ToString(ECasErrorCode Code)
 
 struct FCasErrorContext
 {
-	ECasErrorCode	ErrorCode = ECasErrorCode::None;
-	EIoErrorCode	IoErrorCode = EIoErrorCode::Unknown;
-	uint32			SystemErrorCode = 0;
-	FString			ErrorMessage;
+	ECasErrorCode		ErrorCode = ECasErrorCode::None;
+	EIoErrorCode		IoErrorCode = EIoErrorCode::Unknown;
+	FString				ErrorMessage;
+	FFileSystemError	FileSystemError;
+	uint32				SystemErrorCode = 0;
 };
 
 template <typename ResultType>
@@ -105,14 +107,16 @@ inline UE::IoStore::TResult<ResultType> MakeCasError(
 	ECasErrorCode ErrorCode,
 	EIoErrorCode IoErrorCode,
 	FString&& ErrorMessage,
+	FFileSystemError&& FileSystemError = FFileSystemError(FString()),
 	uint32 SystemErrorCode = FPlatformMisc::GetLastError())
 {
 	return MakeError(CasError::MakeError(FCasErrorContext
 	{
 		.ErrorCode			= ErrorCode,
 		.IoErrorCode		= IoErrorCode,
-		.SystemErrorCode	= SystemErrorCode,
-		.ErrorMessage		= MoveTemp(ErrorMessage)
+		.ErrorMessage		= MoveTemp(ErrorMessage),
+		.FileSystemError	= MoveTemp(FileSystemError),
+		.SystemErrorCode	= SystemErrorCode
 	}, UE::UnifiedError::EDetailFilter::All));
 }
 
@@ -120,14 +124,16 @@ inline UE::IoStore::FResult MakeJournalError(
 	ECasErrorCode ErrorCode,
 	EIoErrorCode IoErrorCode,
 	FString&& ErrorMessage,
+	FFileSystemError&& FileSystemError = FFileSystemError(FString()),
 	uint32 SystemErrorCode = FPlatformMisc::GetLastError())
 {
 	return MakeError(CasJournalError::MakeError(FCasErrorContext
 	{
 		.ErrorCode			= ErrorCode,
 		.IoErrorCode		= IoErrorCode,
-		.SystemErrorCode	= SystemErrorCode,
-		.ErrorMessage		= MoveTemp(ErrorMessage)
+		.ErrorMessage		= MoveTemp(ErrorMessage),
+		.FileSystemError	= MoveTemp(FileSystemError),
+		.SystemErrorCode	= SystemErrorCode
 	}, UE::UnifiedError::EDetailFilter::All));
 }
 
@@ -136,6 +142,7 @@ inline UE::IoStore::TResult<ResultType> MakeSnapshotError(
 	ECasErrorCode ErrorCode,
 	EIoErrorCode IoErrorCode,
 	FString&& ErrorMessage,
+	FFileSystemError&& FileSystemError = FFileSystemError(FString()),
 	uint32 SystemErrorCode = FPlatformMisc::GetLastError())
 {
 	using namespace UE::UnifiedError;
@@ -145,8 +152,29 @@ inline UE::IoStore::TResult<ResultType> MakeSnapshotError(
 	{
 		.ErrorCode			= ErrorCode,
 		.IoErrorCode		= IoErrorCode,
+		.ErrorMessage		= MoveTemp(ErrorMessage),
+		.FileSystemError	= MoveTemp(FileSystemError),
+		.SystemErrorCode	= SystemErrorCode,
+	}, UE::UnifiedError::EDetailFilter::All));
+}
+
+template <typename ResultType>
+inline UE::IoStore::TResult<ResultType> MakeSnapshotError(
+	ECasErrorCode ErrorCode,
+	EIoErrorCode IoErrorCode,
+	FString&& ErrorMessage,
+	uint32 SystemErrorCode)
+{
+	using namespace UE::UnifiedError;
+	using namespace UE::UnifiedError::IoStoreOnDemand;
+
+	return MakeError(CasSnapshotError::MakeError(FCasErrorContext
+	{
+		.ErrorCode			= ErrorCode,
+		.IoErrorCode		= IoErrorCode,
+		.ErrorMessage		= MoveTemp(ErrorMessage),
+		.FileSystemError	= FFileSystemError(FString()),
 		.SystemErrorCode	= SystemErrorCode,
-		.ErrorMessage		= MoveTemp(ErrorMessage)
 	}, UE::UnifiedError::EDetailFilter::All));
 }
 
@@ -189,14 +217,19 @@ UE_DECLARE_ERRORSTRUCT_FEATURES(IoStoreOnDemand, FCasErrorContext);
 inline void SerializeForLog(FCbWriter& Writer, const UE::UnifiedError::IoStoreOnDemand::FCasErrorContext& Ctx)
 {
 	Writer.BeginObject();
-	Writer.AddString(ANSITEXTVIEW("$type"), TErrorStructFeatures<UE::UnifiedError::IoStoreOnDemand::FCasErrorContext>::GetErrorContextTypeNameAsString());
-	Writer.AddString(ANSITEXTVIEW("$format"), TEXT("({ErrorCodeText}: {ErrorMessage})"));
-	Writer.AddInteger(ANSITEXTVIEW("ErrorCode"), uint32(Ctx.ErrorCode));
-	Writer.AddString(ANSITEXTVIEW("ErrorCodeText"), ToString(Ctx.ErrorCode));
-	Writer.AddInteger(ANSITEXTVIEW("IoErrorCode"), uint32(Ctx.IoErrorCode));
-	Writer.AddInteger(ANSITEXTVIEW("SystemErrorCode"), Ctx.SystemErrorCode);
-	Writer.AddString(ANSITEXTVIEW("IoErrorCodeText"), GetIoErrorText(Ctx.IoErrorCode));
-	Writer.AddString(ANSITEXTVIEW("ErrorMessage"), Ctx.ErrorMessage); 
+	Writer.AddString(UTF8TEXTVIEW("$type"), TErrorStructFeatures<UE::UnifiedError::IoStoreOnDemand::FCasErrorContext>::GetErrorContextTypeNameAsString());
+	Writer.AddString(UTF8TEXTVIEW("$format"), TEXT("({ErrorCodeText}: {ErrorMessage} [FileSystemError: {FileSystemError}])"));
+	Writer.AddInteger(UTF8TEXTVIEW("ErrorCode"), uint32(Ctx.ErrorCode));
+	Writer.AddString(UTF8TEXTVIEW("ErrorCodeText"), ToString(Ctx.ErrorCode));
+	Writer.AddInteger(UTF8TEXTVIEW("IoErrorCode"), uint32(Ctx.IoErrorCode));
+	Writer.AddInteger(UTF8TEXTVIEW("SystemErrorCode"), Ctx.SystemErrorCode);
+	Writer.AddString(UTF8TEXTVIEW("IoErrorCodeText"), GetIoErrorText(Ctx.IoErrorCode));
+	Writer.AddString(UTF8TEXTVIEW("ErrorMessage"), Ctx.ErrorMessage);
+
+	TUtf8StringBuilder<512> StringBuilder;
+	StringBuilder << Ctx.FileSystemError;
+	Writer.AddString(UTF8TEXTVIEW("FileSystemError"), StringBuilder);
+
 	Writer.EndObject();
 }
 
@@ -212,11 +245,11 @@ UE_DECLARE_ERRORSTRUCT_FEATURES(IoStoreOnDemand, FChunkHashMismatchErrorContext)
 inline void SerializeForLog(FCbWriter& Writer, const UE::UnifiedError::IoStoreOnDemand::FChunkHashMismatchErrorContext& Ctx)
 {
 	Writer.BeginObject();
-	Writer.AddString(ANSITEXTVIEW("$type"), TErrorStructFeatures<UE::UnifiedError::IoStoreOnDemand::FChunkHashMismatchErrorContext>::GetErrorContextTypeNameAsString());
-	Writer.AddString(ANSITEXTVIEW("$format"), TEXT("(ChunkId: {ChunkId}, ExpectedHash: {ExpectedHash}, ActualHash: {ActualHash})"));
-	Writer.AddString(ANSITEXTVIEW("ChunkId"), LexToString(Ctx.ChunkId));
-	Writer.AddString(ANSITEXTVIEW("ExpectedHash"), LexToString(Ctx.ExpectedHash));
-	Writer.AddString(ANSITEXTVIEW("ActualHash"), LexToString(Ctx.ActualHash));
+	Writer.AddString(UTF8TEXTVIEW("$type"), TErrorStructFeatures<UE::UnifiedError::IoStoreOnDemand::FChunkHashMismatchErrorContext>::GetErrorContextTypeNameAsString());
+	Writer.AddString(UTF8TEXTVIEW("$format"), TEXT("(ChunkId: {ChunkId}, ExpectedHash: {ExpectedHash}, ActualHash: {ActualHash})"));
+	Writer.AddString(UTF8TEXTVIEW("ChunkId"), LexToString(Ctx.ChunkId));
+	Writer.AddString(UTF8TEXTVIEW("ExpectedHash"), LexToString(Ctx.ExpectedHash));
+	Writer.AddString(UTF8TEXTVIEW("ActualHash"), LexToString(Ctx.ActualHash));
 	Writer.EndObject();
 }
 ////////////////////////////////////////////////////////////////////////////////
@@ -224,17 +257,17 @@ UE_DECLARE_ERRORSTRUCT_FEATURES(IoStoreOnDemand, FInstallCacheErrorContext);
 inline void SerializeForLog(FCbWriter& Writer, const UE::UnifiedError::IoStoreOnDemand::FInstallCacheErrorContext& Ctx)
 {
 	Writer.BeginObject();
-	Writer.AddString(ANSITEXTVIEW("$type"), TErrorStructFeatures<UE::UnifiedError::IoStoreOnDemand::FInstallCacheErrorContext>::GetErrorContextTypeNameAsString());
-	Writer.AddString(ANSITEXTVIEW("$format"), TEXT("({ErrorMessage})"));
-	Writer.AddInteger(ANSITEXTVIEW("ErrorCode"), int32(Ctx.IoError.GetErrorCode()));
-	Writer.AddString(ANSITEXTVIEW("ErrorMessage"), Ctx.IoError.ToString());
-	Writer.AddInteger(ANSITEXTVIEW("SystemErrorCode"), int32(Ctx.IoError.GetSystemErrorCode()));
-	Writer.AddInteger(ANSITEXTVIEW("MaxCacheSize"), Ctx.MaxCacheSize);
-	Writer.AddInteger(ANSITEXTVIEW("CacheSize"), Ctx.CacheSize);
-	Writer.AddInteger(ANSITEXTVIEW("RequestedSize"), Ctx.RequestedSize);
-	Writer.AddInteger(ANSITEXTVIEW("DiskTotalBytes"), Ctx.DiskTotalBytes);
-	Writer.AddInteger(ANSITEXTVIEW("DiskFreeBytes"), Ctx.DiskFreeBytes);
-	Writer.AddInteger(ANSITEXTVIEW("LineNo"), uint32(Ctx.LineNo));
+	Writer.AddString(UTF8TEXTVIEW("$type"), TErrorStructFeatures<UE::UnifiedError::IoStoreOnDemand::FInstallCacheErrorContext>::GetErrorContextTypeNameAsString());
+	Writer.AddString(UTF8TEXTVIEW("$format"), TEXT("({ErrorMessage})"));
+	Writer.AddInteger(UTF8TEXTVIEW("ErrorCode"), int32(Ctx.IoError.GetErrorCode()));
+	Writer.AddString(UTF8TEXTVIEW("ErrorMessage"), Ctx.IoError.ToString());
+	Writer.AddInteger(UTF8TEXTVIEW("SystemErrorCode"), int32(Ctx.IoError.GetSystemErrorCode()));
+	Writer.AddInteger(UTF8TEXTVIEW("MaxCacheSize"), Ctx.MaxCacheSize);
+	Writer.AddInteger(UTF8TEXTVIEW("CacheSize"), Ctx.CacheSize);
+	Writer.AddInteger(UTF8TEXTVIEW("RequestedSize"), Ctx.RequestedSize);
+	Writer.AddInteger(UTF8TEXTVIEW("DiskTotalBytes"), Ctx.DiskTotalBytes);
+	Writer.AddInteger(UTF8TEXTVIEW("DiskFreeBytes"), Ctx.DiskFreeBytes);
+	Writer.AddInteger(UTF8TEXTVIEW("LineNo"), uint32(Ctx.LineNo));
 	Writer.EndObject();
 }
 ////////////////////////////////////////////////////////////////////////////////
@@ -242,11 +275,11 @@ UE_DECLARE_ERRORSTRUCT_FEATURES(IoStoreOnDemand, FVerificationErrorContext);
 inline void SerializeForLog(FCbWriter& Writer, const UE::UnifiedError::IoStoreOnDemand::FVerificationErrorContext& Ctx)
 {
 	Writer.BeginObject();
-	Writer.AddString(ANSITEXTVIEW("$type"), TErrorStructFeatures<UE::UnifiedError::IoStoreOnDemand::FVerificationErrorContext>::GetErrorContextTypeNameAsString());
-	Writer.AddString(ANSITEXTVIEW("$format"), TEXT("(CorruptChunkCount: {CorruptChunkCount}, MissingChunkCount: {MissingChunkCount}, ReadErrorCount: {ReadErrorCount})"));
-	Writer.AddInteger(ANSITEXTVIEW("CorruptChunkCount"), Ctx.CorruptChunkCount);
-	Writer.AddInteger(ANSITEXTVIEW("MissingChunkCount"), Ctx.MissingChunkCount);
-	Writer.AddInteger(ANSITEXTVIEW("ReadErrorCount"), Ctx.ReadErrorCount);
+	Writer.AddString(UTF8TEXTVIEW("$type"), TErrorStructFeatures<UE::UnifiedError::IoStoreOnDemand::FVerificationErrorContext>::GetErrorContextTypeNameAsString());
+	Writer.AddString(UTF8TEXTVIEW("$format"), TEXT("(CorruptChunkCount: {CorruptChunkCount}, MissingChunkCount: {MissingChunkCount}, ReadErrorCount: {ReadErrorCount})"));
+	Writer.AddInteger(UTF8TEXTVIEW("CorruptChunkCount"), Ctx.CorruptChunkCount);
+	Writer.AddInteger(UTF8TEXTVIEW("MissingChunkCount"), Ctx.MissingChunkCount);
+	Writer.AddInteger(UTF8TEXTVIEW("ReadErrorCount"), Ctx.ReadErrorCount);
 	Writer.EndObject();
 }
 
