Fix [redacted] build

#Synced-CL 50847432

diff --git a/Engine/Source/Runtime/Experimental/IoStore/OnDemand/Private/OnDemandInstallCache.cpp b/Engine/Source/Runtime/Experimental/IoStore/OnDemand/Private/OnDemandInstallCache.cpp
--- a/Engine/Source/Runtime/Experimental/IoStore/OnDemand/Private/OnDemandInstallCache.cpp
+++ b/Engine/Source/Runtime/Experimental/IoStore/OnDemand/Private/OnDemandInstallCache.cpp
@@ -2745,7 +2745,7 @@ bool FOnDemandInstallCache::Resolve(FIoRequestImpl* Request)
 			}
 			else
 			{
-				UE_LOGFMT(LogIoStoreOnDemand, Error, "Failed to open CAS block '{Filename}' for reading ({Error})", *Filename, FileOpenResult.GetError().GetErrorMessage(true).ToString()));
+				UE_LOGFMT(LogIoStoreOnDemand, Error, "Failed to open CAS block '{Filename}' for reading ({Error})", *Filename, FileOpenResult.GetError().GetErrorMessage(true).ToString());
 			}
 
 			UE::Tasks::Launch(UE_SOURCE_LOCATION, [this, Request, Status]
