Preping for IAD install cache refactor
- Remove chunk pinning options from FOnDemandGetInstallSizeArgs.  This can be better accomplished by just making an install request with EOnDemandInstallOptions::DoNotDownload.

#Synced-CL 50847432

#Synced-CL 50847432

diff --git a/Engine/Source/Runtime/Experimental/IoStore/OnDemand/Private/OnDemandContentInstaller.cpp b/Engine/Source/Runtime/Experimental/IoStore/OnDemand/Private/OnDemandContentInstaller.cpp
--- a/Engine/Source/Runtime/Experimental/IoStore/OnDemand/Private/OnDemandContentInstaller.cpp
+++ b/Engine/Source/Runtime/Experimental/IoStore/OnDemand/Private/OnDemandContentInstaller.cpp
@@ -1073,7 +1073,7 @@ void FOnDemandContentInstaller::ExecuteInstallRequest(FRequest& Request)
 
 	if (EnumHasAnyFlags(InstallRequest.Args.Options, EOnDemandInstallOptions::DoNotDownload))
 	{
-		Request.Result = MakeError(UE::UnifiedError::Core::ArgumentError::MakeError(TEXT("Options"), TEXT("DoNotDownload flag specified when download required")));
+		Request.Result = MakeError(UE::UnifiedError::IoStoreOnDemand::DownloadRequired::MakeError());
 		return CompleteInstallRequest(Request);
 	}
 
diff --git a/Engine/Source/Runtime/Experimental/IoStore/OnDemand/Private/OnDemandInstallCache.cpp b/Engine/Source/Runtime/Experimental/IoStore/OnDemand/Private/OnDemandInstallCache.cpp
--- a/Engine/Source/Runtime/Experimental/IoStore/OnDemand/Private/OnDemandInstallCache.cpp
+++ b/Engine/Source/Runtime/Experimental/IoStore/OnDemand/Private/OnDemandInstallCache.cpp
@@ -2249,11 +2249,6 @@ public:
 
 	// IOnDemandInstallCache
 	virtual bool								IsChunkCached(const FIoHash& ChunkHash) override;
-	virtual bool								TryPinChunks(
-													const FSharedOnDemandContainer& Container,
-													TConstArrayView<int32> EntryIndices,
-													FOnDemandContentHandle ContentHandle,
-													TArray<int32>& OutMissing) override;
 	virtual TResult<bool>						PutChunk(FIoBuffer&& Chunk, const FIoHash& ChunkHash) override;
 	virtual FResult								Purge(uint64 BytesToInstall) override;
 	virtual FResult								PurgeAllUnreferenced(bool bDefrag, const uint64* BytesToPurge = nullptr) override;
@@ -2821,48 +2816,6 @@ bool FOnDemandInstallCache::IsChunkCached(const FIoHash& ChunkHash)
 	return Loc.IsValid();
 }
 
-bool FOnDemandInstallCache::TryPinChunks(
-	const FSharedOnDemandContainer& Container,
-	TConstArrayView<int32> EntryIndices,
-	FOnDemandContentHandle ContentHandle,
-	TArray<int32>& OutMissing)
-{
-	TDynamicSharedLock SharedLock(PurgeDefragMutex, UE::DeferLock);
-
-	// Currently we can only pin chunks if we are not purging/defragging the cache 
-	if (SharedLock.TryLock() == false)
-	{
-		OutMissing = EntryIndices;
-		return false;
-	}
-
-	for (int32 EntryIndex : EntryIndices)
-	{
-		const FOnDemandChunkEntry& ChunkEntry = Container->ChunkEntries[EntryIndex];
-
-		{
-			TSharedLock PendingSharedLock(PendingChunks.SharedMutex);
-
-			if (PendingChunks.ContainsChunk(ChunkEntry.Hash))
-			{
-				IoStore.AddReference(Container, EntryIndex, ContentHandle);
-				continue;
-			}
-		}
-
-		if (FCasLocation Loc = Cas.FindChunk(ChunkEntry.Hash); Loc.IsValid())
-		{
-			IoStore.AddReference(Container, EntryIndex, ContentHandle);
-		}
-		else
-		{
-			OutMissing.Add(EntryIndex);
-		}
-	}
-
-	return true;
-}
-
 TResult<bool> FOnDemandInstallCache::PutChunk(FIoBuffer&& Chunk, const FIoHash& ChunkHash)
 {
 	if (PendingChunks.TotalSize > FPendingChunks::MaxPendingBytes)
diff --git a/Engine/Source/Runtime/Experimental/IoStore/OnDemand/Private/OnDemandInstallCache.h b/Engine/Source/Runtime/Experimental/IoStore/OnDemand/Private/OnDemandInstallCache.h
--- a/Engine/Source/Runtime/Experimental/IoStore/OnDemand/Private/OnDemandInstallCache.h
+++ b/Engine/Source/Runtime/Experimental/IoStore/OnDemand/Private/OnDemandInstallCache.h
@@ -30,11 +30,6 @@ class IOnDemandInstallCache
 public:
 	virtual										~IOnDemandInstallCache() = default;
 	virtual bool								IsChunkCached(const FIoHash& ChunkHash) = 0;
-	virtual bool								TryPinChunks(
-													const FSharedOnDemandContainer& Container,
-													TConstArrayView<int32> EntryIndices,
-													FOnDemandContentHandle ContentHandle,
-													TArray<int32>& OutMissing) = 0;
 	virtual TResult<bool>						PutChunk(FIoBuffer&& Chunk, const FIoHash& ChunkHash) = 0;
 	virtual FResult								Purge(uint64 BytesToInstall) = 0;
 	virtual FResult								PurgeAllUnreferenced(bool bDefrag, const uint64* BytesToPurge = nullptr) = 0;
diff --git a/Engine/Source/Runtime/Experimental/IoStore/OnDemand/Private/OnDemandIoStore.cpp b/Engine/Source/Runtime/Experimental/IoStore/OnDemand/Private/OnDemandIoStore.cpp
--- a/Engine/Source/Runtime/Experimental/IoStore/OnDemand/Private/OnDemandIoStore.cpp
+++ b/Engine/Source/Runtime/Experimental/IoStore/OnDemand/Private/OnDemandIoStore.cpp
@@ -746,56 +746,10 @@ TIoStatusOr<FOnDemandInstallSizeResult> FOnDemandIoStore::GetInstallSize(const F
 	
 	FOnDemandInstallSizeResult Result;
 
-	const bool bPin = Args.ContentHandle.IsValid();
-	if (bPin == false)
-	{
-		Result.InstallSize = Algo::TransformAccumulate(
-			ResolvedChunks,
-			[](const FResolvedContainerChunks& R) { return R.TotalSize; },
-			uint64(0));
-
-		return Result;
-	}
-
-	TArray<int32>						Uncached;
-	TArray<int32, TInlineAllocator<64>> Unreferenced;
-
-	Result.DownloadSize.Emplace(0);
-	for (const FResolvedContainerChunks& Resolved : ResolvedChunks)
-	{
-		Result.InstallSize += Resolved.TotalSize;
-
-		Uncached.Reset();
-		Unreferenced.Reset();
-		{
-			TUniqueLock Lock(ContainerMutex);
-			const FOnDemandChunkEntryReferences& Refs = Resolved.Container->FindOrAddChunkEntryReferences(*Args.ContentHandle.Handle);
-			for (int32 EntryIndex : Resolved.EntryIndices)
-			{
-				const bool bReferenced = Refs.Indices[EntryIndex];
-				if (bReferenced == false)
-				{
-					Unreferenced.Add(EntryIndex);
-				}
-			}
-		}
-
-		if (Unreferenced.IsEmpty() == false)
-		{
-			const bool bOk = InstallCache->TryPinChunks(
-					Resolved.Container.ToSharedRef(),
-					Unreferenced,
-					Args.ContentHandle,
-					Uncached);
-			UE_CLOG(!bOk, LogIoStoreOnDemand, Log, TEXT("Failed to pin install cache chunks due to purging/defragging"))
-
-			for (int32 EntryIndex : Uncached)
-			{
-				const FOnDemandChunkEntry& Entry = Resolved.Container->ChunkEntries[EntryIndex];
-				Result.DownloadSize.GetValue() += Entry.GetDiskSize(); 
-			}
-		}
-	}
+	Result.InstallSize = Algo::TransformAccumulate(
+		ResolvedChunks,
+		[](const FResolvedContainerChunks& R) { return R.TotalSize; },
+		uint64(0));
 
 	return Result;
 }
diff --git a/Engine/Source/Runtime/Experimental/IoStore/OnDemandCore/Private/OnDemandError.cpp b/Engine/Source/Runtime/Experimental/IoStore/OnDemandCore/Private/OnDemandError.cpp
--- a/Engine/Source/Runtime/Experimental/IoStore/OnDemandCore/Private/OnDemandError.cpp
+++ b/Engine/Source/Runtime/Experimental/IoStore/OnDemandCore/Private/OnDemandError.cpp
@@ -18,6 +18,7 @@ UE_DEFINE_ERROR(CasSnapshotError, IoStoreOnDemand);
 UE_DEFINE_ERROR(InstallerCacheFullError, IoStoreOnDemand);
 UE_DEFINE_ERROR(InstallReplayLoadError, IoStoreOnDemand);
 UE_DEFINE_ERROR(InstallReplaySaveError, IoStoreOnDemand);
+UE_DEFINE_ERROR(DownloadRequired, IoStoreOnDemand);
 
 namespace UE::UnifiedError::IoStore
 {
diff --git a/Engine/Source/Runtime/Experimental/IoStore/OnDemandCore/Public/IO/IoStoreOnDemand.h b/Engine/Source/Runtime/Experimental/IoStore/OnDemandCore/Public/IO/IoStoreOnDemand.h
--- a/Engine/Source/Runtime/Experimental/IoStore/OnDemandCore/Public/IO/IoStoreOnDemand.h
+++ b/Engine/Source/Runtime/Experimental/IoStore/OnDemandCore/Public/IO/IoStoreOnDemand.h
@@ -449,17 +449,13 @@ struct FOnDemandGetInstallSizeArgs
 	TArray<FPackageId> PackageIds;
 	/** Options */
 	EOnDemandGetInstallSizeOptions Options = EOnDemandGetInstallSizeOptions::None;
-	/** Pin existing chunk(s) with the specified content handle and compute remaining bytes to download. */
-	FOnDemandContentHandle ContentHandle;
 };
 
 /** Holds information about an install size request. */ 
 struct FOnDemandInstallSizeResult
 {
 	/** The total size of the content. */
-	uint64	InstallSize = 0;
-	/** The total size to download. This value is set when pinning existing chunk(s). */
-	TOptional<uint64> DownloadSize = 0;
+	uint64	InstallSize = 0;	
 };
 
 /** Options for controlling the behavior of the GetIsOnDemand request. */
diff --git a/Engine/Source/Runtime/Experimental/IoStore/OnDemandCore/Public/IO/OnDemandError.h b/Engine/Source/Runtime/Experimental/IoStore/OnDemandCore/Public/IO/OnDemandError.h
--- a/Engine/Source/Runtime/Experimental/IoStore/OnDemandCore/Public/IO/OnDemandError.h
+++ b/Engine/Source/Runtime/Experimental/IoStore/OnDemandCore/Public/IO/OnDemandError.h
@@ -28,6 +28,7 @@ UE_DECLARE_ERROR(UE_API,			CasSnapshotError,					11,	IoStoreOnDemand, NSLOCTEXT(
 UE_DECLARE_ERROR(UE_API,			InstallerCacheFullError,			12, IoStoreOnDemand, NSLOCTEXT("IoStoreOnDemand", "InstallerCacheFullError", "Cache full error"));
 UE_DECLARE_ERROR(UE_API,			InstallReplayLoadError,				13, IoStoreOnDemand, NSLOCTEXT("IoStoreOnDemand", "InstallReplayLoadError", "Failed to load install replay"));
 UE_DECLARE_ERROR(UE_API,			InstallReplaySaveError,				14, IoStoreOnDemand, NSLOCTEXT("IoStoreOnDemand", "InstallReplaySaveError", "Failed to save install replay"));
+UE_DECLARE_ERROR(UE_API,			DownloadRequired,					15, IoStoreOnDemand, NSLOCTEXT("IoStoreOnDemand", "DownloadRequired", "DoNotDownload flag specified when download required"));
 
 ////////////////////////////////////////////////////////////////////////////////
 namespace UE::IoStore
