Improved error handling for iad cache init and increased max cache chunk size.

#Synced-CL 50392955

diff --git a/Engine/Source/Runtime/Experimental/IoStore/OnDemand/Private/OnDemandInstallCache.cpp b/Engine/Source/Runtime/Experimental/IoStore/OnDemand/Private/OnDemandInstallCache.cpp
--- a/Engine/Source/Runtime/Experimental/IoStore/OnDemand/Private/OnDemandInstallCache.cpp
+++ b/Engine/Source/Runtime/Experimental/IoStore/OnDemand/Private/OnDemandInstallCache.cpp
@@ -837,7 +837,7 @@ FResult FCas::Configure(const FOnDemandInstallCacheConfig& Config)
 {
 	using namespace UE::UnifiedError::IoStoreOnDemand;
 
-	const uint32 MaxSupportedBlockSize = (128 << 20);
+	const uint32 MaxSupportedBlockSize = (256 << 20);
 
 	if (Config.MinBlockSize > MaxSupportedBlockSize)
 	{
@@ -2318,7 +2318,7 @@ FOnDemandInstallCache::FOnDemandInstallCache(const FOnDemandInstallCacheConfig&
 	FResult InitResult = Cas.Configure(Config);
 	if (InitResult.HasError())
 	{
-		UE_LOGFMT(LogIoStoreOnDemand, Error, "Failed to configure install cache, error: {Error}", InitResult.GetError());
+		UE_LOGFMT(LogIoStoreOnDemand, Fatal, "Failed to configure install cache, error: {Error}", InitResult.GetError());
 		return;
 	}
 
@@ -2328,7 +2328,7 @@ FOnDemandInstallCache::FOnDemandInstallCache(const FOnDemandInstallCacheConfig&
 	const uint64 MinDiskQuota = 2 * Cas.GetMaxBlockSize();
 	if (MaxCacheSize < MinDiskQuota)
 	{
-		UE_LOG(LogIoStoreOnDemand, Error, TEXT("Failed to initialize install cache - disk quota must be at least %.2lf MiB"), ToMiB(MinDiskQuota));
+		UE_LOG(LogIoStoreOnDemand, Fatal, TEXT("Failed to initialize install cache - disk quota must be at least %.2lf MiB"), ToMiB(MinDiskQuota));
 		return;
 	}
 
@@ -2339,7 +2339,7 @@ FOnDemandInstallCache::FOnDemandInstallCache(const FOnDemandInstallCacheConfig&
 	InitResult = Cas.Initialize(CacheDirectory);
 	if (InitResult.HasError())
 	{
-		UE_LOGFMT(LogIoStoreOnDemand, Error, "Failed to initialize install cache, error: {Error}", InitResult.GetError());
+		UE_LOGFMT(LogIoStoreOnDemand, Fatal, "Failed to initialize install cache, error: {Error}", InitResult.GetError());
 		return;
 	}
 
@@ -2462,7 +2462,7 @@ FOnDemandInstallCache::FOnDemandInstallCache(const FOnDemandInstallCacheConfig&
 	// Try to reset the cache if something has gone wrong 
 	if (InitResult.HasError())
 	{
-		UE_LOGFMT(LogIoStoreOnDemand, Error, "Resetting install cache, reason: {Error}", *InitResult.GetError().GetErrorMessage(true).ToString());
+		UE_LOGFMT(LogIoStoreOnDemand, Fatal, "Resetting install cache, reason: {Error}", *InitResult.GetError().GetErrorMessage(true).ToString());
 
 		FOnDemandInstallCacheStats::OnStartupError(InitResult);
 		InitResult = Reset();
@@ -2476,7 +2476,7 @@ FOnDemandInstallCache::FOnDemandInstallCache(const FOnDemandInstallCacheConfig&
 	}
 	else
 	{
-		UE_LOGFMT(LogIoStoreOnDemand, Error, "Failed to initialize install cache, reason: {Error}", *InitResult.GetError().GetErrorMessage(true).ToString());
+		UE_LOGFMT(LogIoStoreOnDemand, Fatal, "Failed to initialize install cache, reason: {Error}", *InitResult.GetError().GetErrorMessage(true).ToString());
 	}
 }
 
@@ -4092,6 +4092,11 @@ FResult FOnDemandInstallCache::FlushPendingChunksImpl(const FPendingChunks& Chun
 				break;
 			}
 
+			if (!ensure(Chunks.Chunks[ChunkIdx].GetSize() <= Cas.GetMaxBlockSize()))
+			{
+				UE_LOGFMT(LogIoStoreOnDemand, Error, "Chunk is larger then block size ({ChunkId})", ("ChunkId",Chunks.ChunkHashes[ChunkIdx]));
+			}
+
 			const FIoBuffer& Chunk = Chunks.Chunks[ChunkIdx];
 			Offsets.Add(CasBlockOffset + Ar.Tell());
 			Ar.Serialize(const_cast<uint8*>(Chunk.GetData()), Chunk.GetSize());
