Minor cleanup of HTTP I/O dispatcher stats

#rnx

#Synced-CL 50957382

diff --git a/Engine/Source/Runtime/Experimental/IoStore/OnDemand/Private/OnDemandHttpIoDispacher.cpp b/Engine/Source/Runtime/Experimental/IoStore/OnDemand/Private/OnDemandHttpIoDispacher.cpp
--- a/Engine/Source/Runtime/Experimental/IoStore/OnDemand/Private/OnDemandHttpIoDispacher.cpp
+++ b/Engine/Source/Runtime/Experimental/IoStore/OnDemand/Private/OnDemandHttpIoDispacher.cpp
@@ -823,9 +823,8 @@ bool FHttpQueue::Cancel(FHttpRequest* Request)
 		if(PausedRequests.Remove(Request) != 0)
 		{
 			// Cancel stats assume that the request was enqueued, so we have to add that back
-			FOnDemandIoBackendStats::Get()->OnHttpEnqueue(Request->GetHttpRequestType());
-
-			FOnDemandIoBackendStats::Get()->OnHttpCancel(Request->GetHttpRequestType());
+			FOnDemandIoBackendStats::Get()->OnHttpEnqueued(Request->GetHttpRequestType());
+			FOnDemandIoBackendStats::Get()->OnHttpDequeued(Request->GetHttpRequestType(), /*bWasCanceled=*/true);
 			FOnDemandIoBackendStats::Get()->OnHttpUnpaused(Request->GetHttpRequestType());
 			return true;
 		}
@@ -839,8 +838,7 @@ bool FHttpQueue::Cancel(FHttpRequest* Request)
 		if (int32 Index = Heap.Find(Request); Index != INDEX_NONE)
 		{
 			Heap.HeapRemoveAt(Index, HttpRequestSortPredicate);
-			FOnDemandIoBackendStats::Get()->OnHttpDequeue(Request->GetHttpRequestType());
-			FOnDemandIoBackendStats::Get()->OnHttpCancel(Request->GetHttpRequestType());
+			FOnDemandIoBackendStats::Get()->OnHttpDequeued(Request->GetHttpRequestType(), /*bWasCanceled=*/true);
 			return true;
 		}
 		return false;
@@ -858,15 +856,14 @@ bool FHttpQueue::Cancel(FHttpRequest* Request)
 	if (int32 Index = Heap.Find(Request); Index != INDEX_NONE)
 	{
 		Heap.HeapRemoveAt(Index, HttpRequestSortPredicate);
-		FOnDemandIoBackendStats::Get()->OnHttpDequeue(Request->GetHttpRequestType());
-		FOnDemandIoBackendStats::Get()->OnHttpCancel(Request->GetHttpRequestType());
+		FOnDemandIoBackendStats::Get()->OnHttpDequeued(Request->GetHttpRequestType(), /*bWasCanceled=*/true);
 
 		if (!List->IsEmpty())
 		{
 			FHttpRequest* Head = List->PeekHead();
 			ensure(Heap.Find(Head) == INDEX_NONE);
 			Heap.HeapPush(Head, HttpRequestSortPredicate);
-			FOnDemandIoBackendStats::Get()->OnHttpEnqueue(Head->GetHttpRequestType());
+			FOnDemandIoBackendStats::Get()->OnHttpEnqueued(Head->GetHttpRequestType());
 		}
 	}
 
@@ -906,7 +903,7 @@ void FHttpQueue::EnqueueInternal(FHttpRequest* Request)
 	{
 		Heap.HeapPush(Request, HttpRequestSortPredicate);
 		Request->Flags.Add(EHttpRequestFlags::HttpQueued);
-		FOnDemandIoBackendStats::Get()->OnHttpEnqueue(Request->GetHttpRequestType());
+		FOnDemandIoBackendStats::Get()->OnHttpEnqueued(Request->GetHttpRequestType());
 		return;
 	}
 
@@ -930,7 +927,7 @@ void FHttpQueue::EnqueueInternal(FHttpRequest* Request)
 	{
 		Heap.HeapPush(Request, HttpRequestSortPredicate);
 		ByChunkKey.Add(Request->ChunkHash, FHttpRequestList(Request));
-		FOnDemandIoBackendStats::Get()->OnHttpEnqueue(Request->GetHttpRequestType());
+		FOnDemandIoBackendStats::Get()->OnHttpEnqueued(Request->GetHttpRequestType());
 	}
 
 	Request->Flags.Add(EHttpRequestFlags::HttpQueued);
@@ -992,7 +989,7 @@ FHttpRequest* FHttpQueue::Dequeue()
 				FHttpRequest* Head = List->PeekHead();
 				ensure(Heap.Find(Head) == INDEX_NONE);
 				Heap.HeapPush(Head, HttpRequestSortPredicate);
-				FOnDemandIoBackendStats::Get()->OnHttpEnqueue(Head->GetHttpRequestType());
+				FOnDemandIoBackendStats::Get()->OnHttpEnqueued(Head->GetHttpRequestType());
 			}
 
 			Next = Batch.PeekHead();
@@ -1012,7 +1009,7 @@ FHttpRequest* FHttpQueue::Dequeue()
 
 	if (Next != nullptr)
 	{
-		FOnDemandIoBackendStats::Get()->OnHttpDequeue(Next->GetHttpRequestType());
+		FOnDemandIoBackendStats::Get()->OnHttpDequeued(Next->GetHttpRequestType(), /*bWasCanceled=*/false);
 	}
 
 	return Next;
@@ -1071,7 +1068,7 @@ void FHttpQueue::OnTogglePause(bool bPause, EHttpRequestTypeFilter Filter)
 
 					FHttpRequest* NewHeadRequest = List->PeekHead();
 					NewHeap.HeapPush(NewHeadRequest, HttpRequestSortPredicate);
-					FOnDemandIoBackendStats::Get()->OnHttpEnqueue(NewHeadRequest->GetHttpRequestType());
+					FOnDemandIoBackendStats::Get()->OnHttpEnqueued(NewHeadRequest->GetHttpRequestType());
 				}
 			}
 			else if (IsPaused(HeadRequest))
@@ -1803,6 +1800,7 @@ bool FHttpDispatcher::ProcessHttpRequests()
 			NextHttpRequest				= nullptr;
 
 			ensure(!HttpRequest->Flags.HasAny(EHttpRequestFlags::HttpQueued | EHttpRequestFlags::CacheInflight));
+			FOnDemandIoBackendStats::Get()->OnHttpStarted();
 
 			FIASHostGroup HostGroup = FHostGroupManager::Get().Find(HttpRequest->HostGroupName);
 			if (HostGroup.IsConnected() == false)
@@ -1811,7 +1809,8 @@ bool FHttpDispatcher::ProcessHttpRequests()
 				// an error itself and it is not being canceled by higher level code. However we do not
 				// currently have a statistic for that and we have to call one of the existing types in
 				// order to correctly reduce the pending count.
-				FOnDemandIoBackendStats::Get()->OnHttpCancel(HttpRequest->GetHttpRequestType());
+				const bool bWasCanceled = true;
+				FOnDemandIoBackendStats::Get()->OnHttpCompleted(HttpRequest->GetHttpRequestType(), 0, 0, 0, 0, 0, bWasCanceled);
 				CompleteRequest(*HttpRequest, HttpRequest->Buffer, FIoHttpRange(), 0, false);
 
 				check(NextHttpRequest == nullptr);
@@ -1850,8 +1849,6 @@ bool FHttpDispatcher::ProcessHttpRequests()
 #endif // UE_ALLOW_INVALID_URL_DEBUGGING
 
 			NumConcurrentRequests++;
-			FOnDemandIoBackendStats::Get()->OnHttpRequestStarted();
-
 			EMultiEndpointRequestFlags Flags = EnumHasAnyFlags(HttpRequest->Options.GetFlags(), EIoHttpFlags::ResponseHeaders)
 				? EMultiEndpointRequestFlags::ResponseHeaders
 				: EMultiEndpointRequestFlags::None;
@@ -1863,7 +1860,14 @@ bool FHttpDispatcher::ProcessHttpRequests()
 				(FMultiEndpointHttpClientResponse&& HttpResponse) mutable
 				{
 					TRACE_REQUEST_COMPLETED(HttpRequest, HttpResponse, HostGroup);
-					FOnDemandIoBackendStats::Get()->OnHttpRequestCompleted();
+					FOnDemandIoBackendStats::Get()->OnHttpCompleted(
+						HttpRequest->GetHttpRequestType(),
+						HttpResponse.StatusCode,
+						HttpResponse.Body.DataSize(),
+						HttpResponse.Sample.GetTotalMs(),
+						HttpResponse.RetryCount,
+						HttpResponse.CDNCacheStatus,
+						HttpResponse.IsCanceled());
 					NumConcurrentRequests--;
 
 					CompleteCancellationToken(HttpRequest);
@@ -1945,25 +1949,6 @@ FHttpRequest* FHttpDispatcher::DequeueHttpRequests()
 
 void FHttpDispatcher::CompleteHttpRequest(FHttpRequest& Request, FMultiEndpointHttpClientResponse&& HttpResponse, const FIoHttpRange& TotalRange)
 {
-	if (HttpResponse.RetryCount > 0)
-	{
-		FOnDemandIoBackendStats::Get()->OnHttpRetry(Request.GetHttpRequestType());
-	}
-
-	if (HttpResponse.IsOk())
-	{
-		FOnDemandIoBackendStats::Get()->OnHttpCdnCacheReply(Request.GetHttpRequestType(), HttpResponse.CDNCacheStatus);
-		FOnDemandIoBackendStats::Get()->OnHttpGet(Request.GetHttpRequestType(), HttpResponse.Body.DataSize(), HttpResponse.Sample.GetTotalMs());
-	}
-	else if (HttpResponse.IsCanceled())
-	{
-		FOnDemandIoBackendStats::Get()->OnHttpCancel(Request.GetHttpRequestType());
-	}
-	else
-	{
-		FOnDemandIoBackendStats::Get()->OnHttpError(Request.GetHttpRequestType());
-	}
-
 	CompleteRequest(Request, MoveTemp(HttpResponse.Headers), HttpResponse.Body, TotalRange, HttpResponse.StatusCode, false);
 }
 
diff --git a/Engine/Source/Runtime/Experimental/IoStore/OnDemand/Private/Statistics.cpp b/Engine/Source/Runtime/Experimental/IoStore/OnDemand/Private/Statistics.cpp
--- a/Engine/Source/Runtime/Experimental/IoStore/OnDemand/Private/Statistics.cpp
+++ b/Engine/Source/Runtime/Experimental/IoStore/OnDemand/Private/Statistics.cpp
@@ -1451,66 +1451,84 @@ void FOnDemandIoBackendStats::OnHttpDisconnected()
 	GHttpDisconnectCount.Increment();
 }
 
-void FOnDemandIoBackendStats::OnHttpEnqueue(EHttpRequestType Type)
+void FOnDemandIoBackendStats::OnHttpEnqueued(EHttpRequestType Type)
 {
 	GHttpPendingCount[(uint8)Type].Increment();
 }
 
-void FOnDemandIoBackendStats::OnHttpDequeue(EHttpRequestType Type)
-{
-	GHttpInflightCount[(uint8)Type].Increment();
-}
-
-void FOnDemandIoBackendStats::OnHttpGet(EHttpRequestType Type, uint64 SizeBytes, uint64 DurationMs)
+void FOnDemandIoBackendStats::OnHttpDequeued(EHttpRequestType Type, bool bWasCanceled)
 {
 	const uint8 Index = (uint8)Type;
-
-	GHttpPendingCount[Index].Decrement();
-	GHttpInflightCount[Index].Decrement();
-	GHttpGetCount[Index].Increment();
-	GHttpDownloadedBytes[Index].Add(SizeBytes);
-	GHttpDurationMsSum[Index] += DurationMs;
-	GHttpDurationMs[Index].Set(DurationMs);
-
-	GHttpHistory[Index].OnGet(SizeBytes, DurationMs);
-
-	GHttpBandwidthMbps[Index].Set(GHttpHistory[Index].GetBandwidthMbps());
-	GHttpDurationMsAvg[Index] = GHttpHistory[Index].GetAverage();
-
-	GHttpDurationMsMax[Index] = FMath::Max(GHttpDurationMsMax[Index], (int32)DurationMs);
+	if (bWasCanceled)
+	{
+		GHttpPendingCount[Index].Decrement();
+		GHttpCancelCount[Index].Increment();
+	}
+	else
+	{
+		GHttpInflightCount[Index].Increment();
+	}
 }
 
-void FOnDemandIoBackendStats::OnHttpCancel(EHttpRequestType Type)
+void FOnDemandIoBackendStats::OnHttpDecodeError(EHttpRequestType Type)
 {
-	GHttpInflightCount[(uint8)Type].Decrement();
-	GHttpPendingCount[(uint8)Type].Decrement();
-	GHttpCancelCount[(uint8)Type].Increment();
+	GHttpDecodeErrorCount[(uint8)Type].Increment();
 }
 
-void FOnDemandIoBackendStats::OnHttpRetry(EHttpRequestType Type)
+void FOnDemandIoBackendStats::OnHttpStarted()
 {
-	GHttpRetryCount[(uint8)Type].Increment();
+	++GHttpNumConcurrentRequests;
 }
 
-void FOnDemandIoBackendStats::OnHttpError(EHttpRequestType Type)
+void FOnDemandIoBackendStats::OnHttpCompleted(
+	EHttpRequestType Type,
+	uint32 StatusCode,
+	uint64 SizeBytes,
+	uint64 DurationMs,
+	uint32 RetryCount,
+	int8 CDNCacheStatus,
+	bool bCanceled)
 {
-	GHttpPendingCount[(uint8)Type].Decrement();
-	GHttpInflightCount[(uint8)Type].Decrement();
-	GHttpErrorCount[(uint8)Type].Increment();
-}
+	const uint8 Index	= (uint8)Type;
+	const bool bHttpOk	= StatusCode >= 200 && StatusCode < 300;
 
-void FOnDemandIoBackendStats::OnHttpDecodeError(EHttpRequestType Type)
-{
-	GHttpDecodeErrorCount[(uint8)Type].Increment();
-}
+	--GHttpNumConcurrentRequests;
+	GHttpPendingCount[Index].Decrement();
+	GHttpInflightCount[Index].Decrement();
 
-void FOnDemandIoBackendStats::OnHttpCdnCacheReply(EHttpRequestType Type, int32 Reply)
-{
-	switch (Reply)
+	if (RetryCount > 0)
 	{
-	case -1:	GHttpCdnCacheUnknown += 1;	break;	
-	case 0:		GHttpCdnCacheMiss += 1;		break;
-	default:	GHttpCdnCacheHit += 1;		break;
+		GHttpRetryCount[Index].Increment();
+	}
+
+	if (bHttpOk)
+	{
+		GHttpGetCount[Index].Increment();
+		GHttpDownloadedBytes[Index].Add(SizeBytes);
+		GHttpDurationMsSum[Index] += DurationMs;
+		GHttpDurationMs[Index].Set(DurationMs);
+
+		GHttpHistory[Index].OnGet(SizeBytes, DurationMs);
+
+		GHttpBandwidthMbps[Index].Set(GHttpHistory[Index].GetBandwidthMbps());
+		GHttpDurationMsAvg[Index] = GHttpHistory[Index].GetAverage();
+
+		GHttpDurationMsMax[Index] = FMath::Max(GHttpDurationMsMax[Index], (int32)DurationMs);
+
+		switch (CDNCacheStatus)
+		{
+		case -1:	GHttpCdnCacheUnknown += 1;	break;	
+		case 0:		GHttpCdnCacheMiss += 1;		break;
+		default:	GHttpCdnCacheHit += 1;		break;
+		}
+	}
+	else if (bCanceled)
+	{
+		GHttpCancelCount[Index].Increment();
+	}
+	else
+	{
+		GHttpErrorCount[Index].Increment();
 	}
 }
 
@@ -1529,16 +1547,6 @@ void FOnDemandIoBackendStats::OnHttpUnpaused(EHttpRequestType Type)
 	--GHttpNumPausedRequests[(uint8)Type];
 }
 
-void FOnDemandIoBackendStats::OnHttpRequestStarted()
-{
-	++GHttpNumConcurrentRequests;
-}
-
-void FOnDemandIoBackendStats::OnHttpRequestCompleted()
-{
-	--GHttpNumConcurrentRequests;
-}
-
 void FOnDemandIoBackendStats::OnApplicationResume()
 {
 	++GAppResumeCount;
diff --git a/Engine/Source/Runtime/Experimental/IoStore/OnDemand/Private/Statistics.h b/Engine/Source/Runtime/Experimental/IoStore/OnDemand/Private/Statistics.h
--- a/Engine/Source/Runtime/Experimental/IoStore/OnDemand/Private/Statistics.h
+++ b/Engine/Source/Runtime/Experimental/IoStore/OnDemand/Private/Statistics.h
@@ -97,23 +97,16 @@ public:
 	void OnHttpConnected() IAS_STATISTICS_IMPL()
 	void OnHttpDisconnected() IAS_STATISTICS_IMPL()
 
-	// From FOnDemandHttpThread
-	void OnHttpEnqueue(EHttpRequestType Type) IAS_STATISTICS_IMPL()
-	void OnHttpCancel(EHttpRequestType Type) IAS_STATISTICS_IMPL()
-	void OnHttpDequeue(EHttpRequestType Type) IAS_STATISTICS_IMPL()
-	void OnHttpGet(EHttpRequestType Type, uint64 SizeBytes, uint64 DurationMs) IAS_STATISTICS_IMPL()
-	void OnHttpRetry(EHttpRequestType Type) IAS_STATISTICS_IMPL()
-	void OnHttpError(EHttpRequestType Type) IAS_STATISTICS_IMPL()
+	// From HTTP I/O dispatcher
+	void OnHttpEnqueued(EHttpRequestType Type) IAS_STATISTICS_IMPL()
+	void OnHttpDequeued(EHttpRequestType Type, bool bWasCanceled) IAS_STATISTICS_IMPL()
+	void OnHttpStarted() IAS_STATISTICS_IMPL()
+	void OnHttpCompleted(EHttpRequestType Type, uint32 StatusCode, uint64 BodySize, uint64 DurationMs, uint32 RetryCount, int8 CDNCacheStatus, bool bCanceled) IAS_STATISTICS_IMPL()
 	void OnHttpDecodeError(EHttpRequestType Type) IAS_STATISTICS_IMPL()
-	void OnHttpCdnCacheReply(EHttpRequestType Type,int32 Reply) IAS_STATISTICS_IMPL()
-
 	void OnHttpOnRemovedPending(EHttpRequestType Type) IAS_STATISTICS_IMPL()
 	void OnHttpPaused(EHttpRequestType Type) IAS_STATISTICS_IMPL()
 	void OnHttpUnpaused(EHttpRequestType Type) IAS_STATISTICS_IMPL()
 
-	void OnHttpRequestStarted() IAS_STATISTICS_IMPL()
-	void OnHttpRequestCompleted() IAS_STATISTICS_IMPL()
-
 private:
 
 #if IAS_WITH_STATISTICS
