[TraceBasedDebuggers] removed ensure that was triggering for a valid case during shutdown

#rnx

#Synced-CL 50939189

diff --git a/Engine/Source/Runtime/Experimental/TraceBasedDebuggers/Private/RemoteSessionsManager.cpp b/Engine/Source/Runtime/Experimental/TraceBasedDebuggers/Private/RemoteSessionsManager.cpp
--- a/Engine/Source/Runtime/Experimental/TraceBasedDebuggers/Private/RemoteSessionsManager.cpp
+++ b/Engine/Source/Runtime/Experimental/TraceBasedDebuggers/Private/RemoteSessionsManager.cpp
@@ -356,7 +356,11 @@ void FRemoteSessionsManager::StopSessionDiscovery(const FGuid& InDebuggerGuid)
 
 	if (ActiveDebuggersUsingSessionDiscovery.IsEmpty())
 	{
-		if (ensureMsgf(TickHandle.IsValid(), TEXT("Tick handle should be valid when removing the last debugger requesting sessions discovery")))
+		// Tick handle should be valid when removing the last debugger requesting sessions discovery,
+		// but during shutdown it is possible that the CoreTicker gets reset before StopSessionDiscovery gets called.
+		// In that case (i.e., engine shutdown), we don't have anything to do.
+		if (!IsEngineExitRequested()
+			&& ensureMsgf(TickHandle.IsValid(), TEXT("Tick handle should be valid when removing the last debugger requesting sessions discovery")))
 		{
 			RemoveExpiredSessions(ERemoveSessionOptions::ForceRemoveAll);
 			FTSTicker::RemoveTicker(TickHandle);
