IAD Error log when reseting the cache should not be Fatal

#Synced-CL 50927740

diff --git a/Engine/Source/Runtime/Experimental/IoStore/OnDemand/Private/OnDemandInstallCache.cpp b/Engine/Source/Runtime/Experimental/IoStore/OnDemand/Private/OnDemandInstallCache.cpp
--- a/Engine/Source/Runtime/Experimental/IoStore/OnDemand/Private/OnDemandInstallCache.cpp
+++ b/Engine/Source/Runtime/Experimental/IoStore/OnDemand/Private/OnDemandInstallCache.cpp
@@ -2462,7 +2462,7 @@ FOnDemandInstallCache::FOnDemandInstallCache(const FOnDemandInstallCacheConfig&
 	// Try to reset the cache if something has gone wrong 
 	if (InitResult.HasError())
 	{
-		UE_LOGFMT(LogIoStoreOnDemand, Fatal, "Resetting install cache, reason: {Error}", *InitResult.GetError().GetErrorMessage(true).ToString());
+		UE_LOGFMT(LogIoStoreOnDemand, Error, "Resetting install cache, reason: {Error}", *InitResult.GetError().GetErrorMessage(true).ToString());
 
 		FOnDemandInstallCacheStats::OnStartupError(InitResult);
 		InitResult = Reset();
