Print whole flag values when an enum is encountered by JsonObjectGraph

[RN] JsonObjectGraph now properly encodes Enums that are being used as Flags

#Synced-CL 50958978

diff --git a/Engine/Source/Runtime/Experimental/JsonObjectGraph/Private/JsonStringifyImpl.cpp b/Engine/Source/Runtime/Experimental/JsonObjectGraph/Private/JsonStringifyImpl.cpp
--- a/Engine/Source/Runtime/Experimental/JsonObjectGraph/Private/JsonStringifyImpl.cpp
+++ b/Engine/Source/Runtime/Experimental/JsonObjectGraph/Private/JsonStringifyImpl.cpp
@@ -687,7 +687,7 @@ void FJsonStringifyImpl::WriteIntrinsicToJson(const void* Value, const FProperty
 		uint8 ValueAsIntrinsic = *(uint8*)Value;
 		if (const UEnum* EnumDef = ByteProperty->Enum)
 		{
-			FString StringValue = EnumDef->GetAuthoredNameStringByValue(ValueAsIntrinsic);
+			FString StringValue = EnumDef->GetValueOrBitfieldAsAuthoredNameString(ValueAsIntrinsic);
 			Writer->WriteValueInline(StringValue);
 		}
 		else
@@ -699,7 +699,7 @@ void FJsonStringifyImpl::WriteIntrinsicToJson(const void* Value, const FProperty
 	{
 		// export enums as strings
 		const UEnum* EnumDef = EnumProperty->GetEnum();
-		FString StringValue = EnumDef->GetAuthoredNameStringByValue(
+		FString StringValue = EnumDef->GetValueOrBitfieldAsAuthoredNameString(
 			EnumProperty->GetUnderlyingProperty()->GetSignedIntPropertyValue(Value));
 		Writer->WriteValueInline(StringValue);
 	}
