Collection Selection facade : optimize vertex to curve selection conversion to support hundreds of thousands of vertices and curves  ( tens of milliseconds results instead of several minutes for the worst case )

#Synced-CL 50953304

diff --git a/Engine/Source/Runtime/Experimental/Chaos/Private/GeometryCollection/Facades/CollectionTransformSelectionFacade.cpp b/Engine/Source/Runtime/Experimental/Chaos/Private/GeometryCollection/Facades/CollectionTransformSelectionFacade.cpp
--- a/Engine/Source/Runtime/Experimental/Chaos/Private/GeometryCollection/Facades/CollectionTransformSelectionFacade.cpp
+++ b/Engine/Source/Runtime/Experimental/Chaos/Private/GeometryCollection/Facades/CollectionTransformSelectionFacade.cpp
@@ -1813,48 +1813,76 @@ namespace GeometryCollection::Facades
 
 	TArray<int32> FCollectionTransformSelectionFacade::ConvertVertexSelectionToCurveSelection(const TArray<int32>& InVertexSelection, bool bAllElementsMustBeSelected) const
 	{
-		TArray<int32> OutSelection;
+		TArray<int32> OutSelectedCurves;
 
 		GeometryCollection::Facades::FCollectionCurveGeometryFacade CurveFacade(ConstCollection);
 		if (CurveFacade.IsValid())
 		{
 			const int32 NumCurves = CurveFacade.GetNumCurves();
-			for (int32 CurveIdx = 0; CurveIdx < NumCurves; ++CurveIdx)
+
+			TBitArray<> SelectedCurves;
+			SelectedCurves.Init(false, NumCurves);
+			const TArray<int32>& CurveIndicesByPoint = CurveFacade.GetPointCurveIndices();
+
+			for (const int32& VertexIndex : InVertexSelection)
 			{
-				const int32 PointIndexStart = CurveIdx == 0 ? 0 :CurveFacade.GetCurvePointOffsets()[CurveIdx-1];
-				const int32 PointIndexEnd = CurveFacade.GetCurvePointOffsets()[CurveIdx];
+				// there's two vertices for each point
+				const int32 PointIndex = (VertexIndex >> 1);
+				if (CurveIndicesByPoint.IsValidIndex(PointIndex))
+				{
+					const int32 CurveIndex = CurveIndicesByPoint[PointIndex];
+					if (SelectedCurves.IsValidIndex(CurveIndex) && !SelectedCurves[CurveIndex])
+					{
+						OutSelectedCurves.Add(CurveIndex);
+						SelectedCurves[CurveIndex] = true;
+					}
+				}
+			}
 
-				const int32 VertexIndexStart = PointIndexStart * 2;
-				const int32 VertexIndexEnd = PointIndexEnd * 2;
+			// let's filter the selected curves further with a more expensive check
+			if (bAllElementsMustBeSelected)
+			{
+				// Convert vertex index array into an O(1) addressable storage
+				TBitArray<> SelectedVertices;
+				SelectedVertices.Init(false, CurveFacade.GetPointRestOrientations().Num() * 2);
+				for (const int32& VertexIndex : InVertexSelection)
+				{
+					if (SelectedVertices.IsValidIndex(VertexIndex))
+					{
+						SelectedVertices[VertexIndex] = true;
+					}
+				}
 
-				bool bAllIncluded = true;
-				for (int32 VertexIdx = VertexIndexStart; VertexIdx < VertexIndexEnd; ++VertexIdx)
+				TArray<int32> PotentiallySelectedCurves = MoveTemp(OutSelectedCurves);
+				OutSelectedCurves.Reserve(PotentiallySelectedCurves.Num());
+				for (int32 CurveIdx : PotentiallySelectedCurves)
 				{
-					if(InVertexSelection.Contains(VertexIdx))
+					const int32 PointIndexStart = CurveIdx == 0 ? 0 : CurveFacade.GetCurvePointOffsets()[CurveIdx - 1];
+					const int32 PointIndexEnd = CurveFacade.GetCurvePointOffsets()[CurveIdx];
+
+					const int32 VertexIndexStart = PointIndexStart * 2;
+					const int32 VertexIndexEnd = PointIndexEnd * 2;
+
+					bool bAllIncluded = true;
+					for (int32 VertexIdx = VertexIndexStart; VertexIdx < VertexIndexEnd; ++VertexIdx)
 					{
-						if(!bAllElementsMustBeSelected)
+						if (SelectedVertices.IsValidIndex(VertexIdx))
 						{
-							OutSelection.Add(CurveIdx);
-							break;
+							if (!SelectedVertices[VertexIdx])
+							{
+								bAllIncluded = false;
+								break;
+							}
 						}
 					}
-					else
+					if (bAllIncluded)
 					{
-						bAllIncluded = false;
-						if(bAllElementsMustBeSelected)
-						{
-							break;
-						}
+						OutSelectedCurves.Add(CurveIdx);
 					}
 				}
-				if(bAllElementsMustBeSelected && bAllIncluded)
-				{
-					OutSelection.Add(CurveIdx);
-				}
 			}
 		}
-
-		return OutSelection;
+		return OutSelectedCurves;
 	}
 
 	TArray<int32> FCollectionTransformSelectionFacade::ConvertCurveSelectionToVertexSelection(const TArray<int32>& InCurveSelection) const
