Chaos : Curve facade : add method to resample a curve

#Synced-CL 50953304

diff --git a/Engine/Source/Runtime/Experimental/Chaos/Private/GeometryCollection/Facades/CollectionCurveFacade.cpp b/Engine/Source/Runtime/Experimental/Chaos/Private/GeometryCollection/Facades/CollectionCurveFacade.cpp
--- a/Engine/Source/Runtime/Experimental/Chaos/Private/GeometryCollection/Facades/CollectionCurveFacade.cpp
+++ b/Engine/Source/Runtime/Experimental/Chaos/Private/GeometryCollection/Facades/CollectionCurveFacade.cpp
@@ -313,6 +313,81 @@ namespace GeometryCollection::Facades
 		}
 	}
 
+	void FCollectionCurveGeometryFacade::GetCurveResampledPositions(const int32 CurveIndex, int32 NumSamples, TArray<FVector3f>& OutPoints) const
+	{
+		OutPoints.Reset();
+
+		if (NumSamples < 2)
+		{
+			return;
+		}
+
+		const int32 NumCurves = GetNumCurves();
+		if (CurveIndex < 0 || CurveIndex >= NumCurves)
+		{
+			return;
+		}
+
+		TArray<FVector3f> Positions;
+		GetCurvePointPositions(CurveIndex, Positions);
+		if (Positions.Num() < 2)
+		{
+			return;
+		}
+
+		// compute lengths
+		TArray<double> DistancesAt;
+		DistancesAt.Init(0, Positions.Num());
+
+		for (int32 PointIndex = 1; PointIndex < Positions.Num(); ++PointIndex)
+		{
+			const double SegmentSize = (Positions[PointIndex] - Positions[PointIndex - 1]).Size();
+			DistancesAt[PointIndex] = SegmentSize + DistancesAt[PointIndex - 1];
+		}
+
+		const double TotalCurveLength = DistancesAt.Last();
+		if (TotalCurveLength == 0)
+		{
+			return;
+		}
+
+		OutPoints.Reserve(NumSamples);
+		const double SamplingLength = TotalCurveLength / (double)(NumSamples - 1);
+
+		OutPoints.Add(Positions[0]);
+
+		double SampledDistance = SamplingLength;;
+		int32 PointIndex = 0;
+		while (PointIndex < DistancesAt.Num())
+		{
+			// find which segment this belong 
+			int NextIndex = (PointIndex + 1);
+			while (NextIndex < DistancesAt.Num())
+			{
+				if (SampledDistance < DistancesAt[NextIndex])
+				{
+					const double RelativeDistance = SampledDistance - DistancesAt[NextIndex - 1];
+					const double SegmentLength = DistancesAt[NextIndex] - DistancesAt[NextIndex - 1];
+					const double Alpha = (SegmentLength > 0) ? (RelativeDistance / SegmentLength) : 0;
+					const FVector3f InterpolatedPos = FMath::Lerp(Positions[NextIndex - 1], Positions[NextIndex], Alpha);
+					OutPoints.Add(InterpolatedPos);
+
+					SampledDistance += SamplingLength;
+					break;
+				}
+				++NextIndex;
+			}
+			PointIndex = (NextIndex - 1);
+
+			// we may have reach the end of gone too far
+			if (SampledDistance >= DistancesAt.Last())
+			{
+				OutPoints.Add(Positions.Last());
+				break;
+			}
+		}
+	}
+
 	FCollectionCurveHierarchyFacade::FCollectionCurveHierarchyFacade(FManagedArrayCollection& InCollection) :
 		ConstCollection(InCollection), Collection(&InCollection),
 		CurveParentIndices(InCollection, CurveParentIndicesAttribute, FCollectionCurveGeometryFacade::CurvesGroup),
diff --git a/Engine/Source/Runtime/Experimental/Chaos/Public/GeometryCollection/Facades/CollectionCurveFacade.h b/Engine/Source/Runtime/Experimental/Chaos/Public/GeometryCollection/Facades/CollectionCurveFacade.h
--- a/Engine/Source/Runtime/Experimental/Chaos/Public/GeometryCollection/Facades/CollectionCurveFacade.h
+++ b/Engine/Source/Runtime/Experimental/Chaos/Public/GeometryCollection/Facades/CollectionCurveFacade.h
@@ -102,6 +102,12 @@ namespace GeometryCollection::Facades
 		/** Get the point positions of a specific curve */
 		CHAOS_API void GetCurvePointPositions(const int32 CurveIndex, TArray<FVector3f>& OutPoints) const;
 
+		/** 
+		* Get the resampled positions of a specific curve 
+		* the curve will be sampled at regular interval along the curve, to match the NumSamples parameters
+		*/
+		CHAOS_API void GetCurveResampledPositions(const int32 CurveIndex, int32 NumSamples, TArray<FVector3f>& OutPoints) const;
+
 	protected :
 
 		/** Build geometry collection from the curves one */
