Added pending bytes trace counters for IAX HTTP requests.

#rnx

#Synced-CL 50957382

diff --git a/Engine/Source/Runtime/Experimental/IoStore/OnDemand/Private/OnDemandHttpIoDispacher.cpp b/Engine/Source/Runtime/Experimental/IoStore/OnDemand/Private/OnDemandHttpIoDispacher.cpp
--- a/Engine/Source/Runtime/Experimental/IoStore/OnDemand/Private/OnDemandHttpIoDispacher.cpp
+++ b/Engine/Source/Runtime/Experimental/IoStore/OnDemand/Private/OnDemandHttpIoDispacher.cpp
@@ -459,6 +459,7 @@ struct FHttpRequest final
 
 	EHttpRequestType				GetHttpRequestType() const		{ return Options.GetCategory() == 1 ? EHttpRequestType::Installed : EHttpRequestType::Streaming; }
 	virtual EIoErrorCode			GetStatus() const override		{ return CompletionStatus.load(std::memory_order_acquire); }
+	uint32							GetRangeSize() const			{ return Options.GetRange().IsValid() ? Options.GetRange().GetSize() : 0; }
 
 	virtual void					AddRef() const override			{ RefCount.fetch_add(1, std::memory_order_acq_rel); }
 	virtual FRefCountValue			GetRefCount() const override	{ return FRefCountValue(RefCount.load(std::memory_order_relaxed)); }
@@ -823,8 +824,8 @@ bool FHttpQueue::Cancel(FHttpRequest* Request)
 		if(PausedRequests.Remove(Request) != 0)
 		{
 			// Cancel stats assume that the request was enqueued, so we have to add that back
-			FOnDemandIoBackendStats::Get()->OnHttpEnqueued(Request->GetHttpRequestType());
-			FOnDemandIoBackendStats::Get()->OnHttpDequeued(Request->GetHttpRequestType(), /*bWasCanceled=*/true);
+			FOnDemandIoBackendStats::Get()->OnHttpEnqueued(Request->GetHttpRequestType(), 0);
+			FOnDemandIoBackendStats::Get()->OnHttpDequeued(Request->GetHttpRequestType(), 0, /*bWasCanceled=*/true);
 			FOnDemandIoBackendStats::Get()->OnHttpUnpaused(Request->GetHttpRequestType());
 			return true;
 		}
@@ -838,7 +839,8 @@ bool FHttpQueue::Cancel(FHttpRequest* Request)
 		if (int32 Index = Heap.Find(Request); Index != INDEX_NONE)
 		{
 			Heap.HeapRemoveAt(Index, HttpRequestSortPredicate);
-			FOnDemandIoBackendStats::Get()->OnHttpDequeued(Request->GetHttpRequestType(), /*bWasCanceled=*/true);
+			const bool bWasCanceled = true;
+			FOnDemandIoBackendStats::Get()->OnHttpDequeued(Request->GetHttpRequestType(), Request->GetRangeSize(), bWasCanceled);
 			return true;
 		}
 		return false;
@@ -856,14 +858,12 @@ bool FHttpQueue::Cancel(FHttpRequest* Request)
 	if (int32 Index = Heap.Find(Request); Index != INDEX_NONE)
 	{
 		Heap.HeapRemoveAt(Index, HttpRequestSortPredicate);
-		FOnDemandIoBackendStats::Get()->OnHttpDequeued(Request->GetHttpRequestType(), /*bWasCanceled=*/true);
 
 		if (!List->IsEmpty())
 		{
 			FHttpRequest* Head = List->PeekHead();
 			ensure(Heap.Find(Head) == INDEX_NONE);
 			Heap.HeapPush(Head, HttpRequestSortPredicate);
-			FOnDemandIoBackendStats::Get()->OnHttpEnqueued(Head->GetHttpRequestType());
 		}
 	}
 
@@ -872,6 +872,9 @@ bool FHttpQueue::Cancel(FHttpRequest* Request)
 		ByChunkKey.Remove(Request->ChunkHash);
 	}
 
+	const bool bWasCanceled = true;
+	FOnDemandIoBackendStats::Get()->OnHttpDequeued(Request->GetHttpRequestType(), Request->GetRangeSize(), bWasCanceled);
+
 	return true;
 }
 
@@ -897,13 +900,13 @@ void FHttpQueue::EnqueueInternal(FHttpRequest* Request)
 	check(Request->Next == nullptr);
 
 	Request->Flags.Remove(EHttpRequestFlags::CacheInflight);
+	FOnDemandIoBackendStats::Get()->OnHttpEnqueued(Request->GetHttpRequestType(), Request->GetRangeSize());
 
 	const bool bGenericHttpRequest = Request->ChunkHash.IsZero();
 	if (bGenericHttpRequest)
 	{
 		Heap.HeapPush(Request, HttpRequestSortPredicate);
 		Request->Flags.Add(EHttpRequestFlags::HttpQueued);
-		FOnDemandIoBackendStats::Get()->OnHttpEnqueued(Request->GetHttpRequestType());
 		return;
 	}
 
@@ -927,7 +930,6 @@ void FHttpQueue::EnqueueInternal(FHttpRequest* Request)
 	{
 		Heap.HeapPush(Request, HttpRequestSortPredicate);
 		ByChunkKey.Add(Request->ChunkHash, FHttpRequestList(Request));
-		FOnDemandIoBackendStats::Get()->OnHttpEnqueued(Request->GetHttpRequestType());
 	}
 
 	Request->Flags.Add(EHttpRequestFlags::HttpQueued);
@@ -989,7 +991,6 @@ FHttpRequest* FHttpQueue::Dequeue()
 				FHttpRequest* Head = List->PeekHead();
 				ensure(Heap.Find(Head) == INDEX_NONE);
 				Heap.HeapPush(Head, HttpRequestSortPredicate);
-				FOnDemandIoBackendStats::Get()->OnHttpEnqueued(Head->GetHttpRequestType());
 			}
 
 			Next = Batch.PeekHead();
@@ -1005,11 +1006,8 @@ FHttpRequest* FHttpQueue::Dequeue()
 	{
 		It->Flags.Add(EHttpRequestFlags::HttpInflight);
 		It->Flags.Remove(EHttpRequestFlags::HttpQueued);
-	}
-
-	if (Next != nullptr)
-	{
-		FOnDemandIoBackendStats::Get()->OnHttpDequeued(Next->GetHttpRequestType(), /*bWasCanceled=*/false);
+		bool bWasCanceled = false;
+		FOnDemandIoBackendStats::Get()->OnHttpDequeued(It->GetHttpRequestType(), It->GetRangeSize(), bWasCanceled);
 	}
 
 	return Next;
@@ -1068,7 +1066,7 @@ void FHttpQueue::OnTogglePause(bool bPause, EHttpRequestTypeFilter Filter)
 
 					FHttpRequest* NewHeadRequest = List->PeekHead();
 					NewHeap.HeapPush(NewHeadRequest, HttpRequestSortPredicate);
-					FOnDemandIoBackendStats::Get()->OnHttpEnqueued(NewHeadRequest->GetHttpRequestType());
+					FOnDemandIoBackendStats::Get()->OnHttpEnqueued(NewHeadRequest->GetHttpRequestType(), NewHeadRequest->GetRangeSize());
 				}
 			}
 			else if (IsPaused(HeadRequest))
diff --git a/Engine/Source/Runtime/Experimental/IoStore/OnDemand/Private/Statistics.cpp b/Engine/Source/Runtime/Experimental/IoStore/OnDemand/Private/Statistics.cpp
--- a/Engine/Source/Runtime/Experimental/IoStore/OnDemand/Private/Statistics.cpp
+++ b/Engine/Source/Runtime/Experimental/IoStore/OnDemand/Private/Statistics.cpp
@@ -371,6 +371,7 @@ private:
 		TCounterInt(...)  {}
 		void Set(int64 i) { V = i; }
 		void Add(int64 d) { V += d; }
+		void Subtract(int64 d) { V -= d; }
 		void Increment() { V++; }
 		void Decrement() { --V; }
 		int64 Get() const { return V;}
@@ -427,6 +428,7 @@ UE_IAX_COUNTER(HttpDecodeErrorCount, FCounterAtomicInt);
 UE_IAX_COUNTER(HttpRetryCount, FCounterInt);
 UE_IAX_COUNTER(HttpCancelCount, FCounterInt);
 UE_IAX_COUNTER(HttpPendingCount, FCounterAtomicInt);
+UE_IAX_MEMORY_COUNTER(HttpPendingBytes, FCounterAtomicInt);
 UE_IAX_COUNTER(HttpInflightCount, FCounterInt);
 UE_IAX_MEMORY_COUNTER(HttpDownloadedBytes, FCounterInt);
 UE_IAX_COUNTER(HttpDurationMs, FCounterInt);
@@ -1451,12 +1453,13 @@ void FOnDemandIoBackendStats::OnHttpDisconnected()
 	GHttpDisconnectCount.Increment();
 }
 
-void FOnDemandIoBackendStats::OnHttpEnqueued(EHttpRequestType Type)
+void FOnDemandIoBackendStats::OnHttpEnqueued(EHttpRequestType Type, uint64 RangeSize)
 {
 	GHttpPendingCount[(uint8)Type].Increment();
+	GHttpPendingBytes[(uint8)Type].Add(RangeSize);
 }
 
-void FOnDemandIoBackendStats::OnHttpDequeued(EHttpRequestType Type, bool bWasCanceled)
+void FOnDemandIoBackendStats::OnHttpDequeued(EHttpRequestType Type, uint64 RangeSize, bool bWasCanceled)
 {
 	const uint8 Index = (uint8)Type;
 	if (bWasCanceled)
@@ -1468,6 +1471,7 @@ void FOnDemandIoBackendStats::OnHttpDequeued(EHttpRequestType Type, bool bWasCan
 	{
 		GHttpInflightCount[Index].Increment();
 	}
+	GHttpPendingBytes[(uint8)Type].Subtract(RangeSize);
 }
 
 void FOnDemandIoBackendStats::OnHttpDecodeError(EHttpRequestType Type)
diff --git a/Engine/Source/Runtime/Experimental/IoStore/OnDemand/Private/Statistics.h b/Engine/Source/Runtime/Experimental/IoStore/OnDemand/Private/Statistics.h
--- a/Engine/Source/Runtime/Experimental/IoStore/OnDemand/Private/Statistics.h
+++ b/Engine/Source/Runtime/Experimental/IoStore/OnDemand/Private/Statistics.h
@@ -98,8 +98,8 @@ public:
 	void OnHttpDisconnected() IAS_STATISTICS_IMPL()
 
 	// From HTTP I/O dispatcher
-	void OnHttpEnqueued(EHttpRequestType Type) IAS_STATISTICS_IMPL()
-	void OnHttpDequeued(EHttpRequestType Type, bool bWasCanceled) IAS_STATISTICS_IMPL()
+	void OnHttpEnqueued(EHttpRequestType Type, uint64 RangeSize) IAS_STATISTICS_IMPL()
+	void OnHttpDequeued(EHttpRequestType Type, uint64 RangeSize, bool bWasCanceled) IAS_STATISTICS_IMPL()
 	void OnHttpStarted() IAS_STATISTICS_IMPL()
 	void OnHttpCompleted(EHttpRequestType Type, uint32 StatusCode, uint64 BodySize, uint64 DurationMs, uint32 RetryCount, int8 CDNCacheStatus, bool bCanceled) IAS_STATISTICS_IMPL()
 	void OnHttpDecodeError(EHttpRequestType Type) IAS_STATISTICS_IMPL()
