Dataflow : optimize even more the selection conversion from Vertex to Curve ( about x3 imporvement )
- Removed unecessary type conversion
- Added a bitarrayversion of the method on the selection facade

#Synced-CL 50953304

diff --git a/Engine/Source/Runtime/Experimental/Chaos/Private/GeometryCollection/Facades/CollectionTransformSelectionFacade.cpp b/Engine/Source/Runtime/Experimental/Chaos/Private/GeometryCollection/Facades/CollectionTransformSelectionFacade.cpp
--- a/Engine/Source/Runtime/Experimental/Chaos/Private/GeometryCollection/Facades/CollectionTransformSelectionFacade.cpp
+++ b/Engine/Source/Runtime/Experimental/Chaos/Private/GeometryCollection/Facades/CollectionTransformSelectionFacade.cpp
@@ -1839,7 +1839,7 @@ namespace GeometryCollection::Facades
 				}
 			}
 
-			// let's filter the selected curves further with a more expensive check
+			// let's filter the selected curves further with a more expensive check if needed
 			if (bAllElementsMustBeSelected)
 			{
 				// Convert vertex index array into an O(1) addressable storage
@@ -1885,6 +1885,61 @@ namespace GeometryCollection::Facades
 		return OutSelectedCurves;
 	}
 
+	void FCollectionTransformSelectionFacade::ConvertVertexSelectionToCurveSelection(const TBitArray<>& InVertexSelection, bool bAllElementsMustBeSelected, TBitArray<>& OutCurveSelection) const
+	{
+		GeometryCollection::Facades::FCollectionCurveGeometryFacade CurveFacade(ConstCollection);
+		if (CurveFacade.IsValid())
+		{
+			const int32 NumCurves = CurveFacade.GetNumCurves();
+			OutCurveSelection.Init(false, NumCurves);
+
+			const TArray<int32>& CurveIndicesByPoint = CurveFacade.GetPointCurveIndices();
+
+			const int32 NumVertices = InVertexSelection.Num();
+			for (int32 VertexIndex = 0; VertexIndex < NumVertices; ++VertexIndex)
+			{
+				if (InVertexSelection[VertexIndex])
+				{
+					// there's two vertices for each point
+					const int32 PointIndex = (VertexIndex >> 1);
+					if (CurveIndicesByPoint.IsValidIndex(PointIndex))
+					{
+						const int32 CurveIndex = CurveIndicesByPoint[PointIndex];
+						if (OutCurveSelection.IsValidIndex(CurveIndex) && !OutCurveSelection[CurveIndex])
+						{
+							OutCurveSelection[CurveIndex] = true;
+						}
+					}
+				}
+			}
+
+			// let's filter the selected curves further with a more expensive check if needed
+			if (bAllElementsMustBeSelected)
+			{
+				for (int32 CurveIdx = 0; CurveIdx < NumCurves; ++CurveIdx)
+				{
+					const int32 PointIndexStart = CurveIdx == 0 ? 0 : CurveFacade.GetCurvePointOffsets()[CurveIdx - 1];
+					const int32 PointIndexEnd = CurveFacade.GetCurvePointOffsets()[CurveIdx];
+
+					const int32 VertexIndexStart = PointIndexStart * 2;
+					const int32 VertexIndexEnd = PointIndexEnd * 2;
+
+					for (int32 VertexIdx = VertexIndexStart; VertexIdx < VertexIndexEnd; ++VertexIdx)
+					{
+						if (InVertexSelection.IsValidIndex(VertexIdx))
+						{
+							if (!InVertexSelection[VertexIdx])
+							{
+								OutCurveSelection[CurveIdx] = false;
+								break;
+							}
+						}
+					}
+				}
+			}
+		}
+	}
+
 	TArray<int32> FCollectionTransformSelectionFacade::ConvertCurveSelectionToVertexSelection(const TArray<int32>& InCurveSelection) const
 	{
 		TArray<int32> OutSelection;
diff --git a/Engine/Source/Runtime/Experimental/Chaos/Public/GeometryCollection/Facades/CollectionTransformSelectionFacade.h b/Engine/Source/Runtime/Experimental/Chaos/Public/GeometryCollection/Facades/CollectionTransformSelectionFacade.h
--- a/Engine/Source/Runtime/Experimental/Chaos/Public/GeometryCollection/Facades/CollectionTransformSelectionFacade.h
+++ b/Engine/Source/Runtime/Experimental/Chaos/Public/GeometryCollection/Facades/CollectionTransformSelectionFacade.h
@@ -275,6 +275,7 @@ namespace GeometryCollection::Facades
 		* Convert a VertexSelection to a CurveSelection
 		*/
 		CHAOS_API TArray<int32> ConvertVertexSelectionToCurveSelection(const TArray<int32>& InVertexSelection, bool bAllElementsMustBeSelected) const;
+		CHAOS_API void ConvertVertexSelectionToCurveSelection(const TBitArray<>& InVertexSelection, bool bAllElementsMustBeSelected, TBitArray<>& OutCurveSelection) const;
 
 		/**
 		* Convert a FaceSelection to a CurveSelection
diff --git a/Engine/Source/Runtime/Experimental/Dataflow/Core/Private/Dataflow/DataflowConvertNodes.cpp b/Engine/Source/Runtime/Experimental/Dataflow/Core/Private/Dataflow/DataflowConvertNodes.cpp
--- a/Engine/Source/Runtime/Experimental/Dataflow/Core/Private/Dataflow/DataflowConvertNodes.cpp
+++ b/Engine/Source/Runtime/Experimental/Dataflow/Core/Private/Dataflow/DataflowConvertNodes.cpp
@@ -296,10 +296,11 @@ void FConvertSelectionTypesDataflowNode::Evaluate(UE::Dataflow::FContext& Contex
 				}
 				else if (Output && Output->IsType<FDataflowCurveSelection>())
 				{
-					const TArray<int32>& SelectionArr = TransformSelectionFacade.ConvertVertexSelectionToCurveSelection(InValue.AsArray(), bAllElementsMustBeSelected);
+					TBitArray<> SelectedCurves;
+					TransformSelectionFacade.ConvertVertexSelectionToCurveSelection(InValue.GetBitArray(), bAllElementsMustBeSelected, SelectedCurves);
 
 					FDataflowCurveSelection NewSelection;
-					NewSelection.InitFromArray(InCollection, SelectionArr);
+					NewSelection.SetFromBitArray(MoveTemp(SelectedCurves));
 
 					SetValue(Context, MoveTemp(NewSelection), &Out);
 				}
diff --git a/Engine/Source/Runtime/Experimental/Dataflow/Core/Private/Dataflow/DataflowSelection.cpp b/Engine/Source/Runtime/Experimental/Dataflow/Core/Private/Dataflow/DataflowSelection.cpp
--- a/Engine/Source/Runtime/Experimental/Dataflow/Core/Private/Dataflow/DataflowSelection.cpp
+++ b/Engine/Source/Runtime/Experimental/Dataflow/Core/Private/Dataflow/DataflowSelection.cpp
@@ -114,6 +114,11 @@ bool FDataflowSelection::SetFromArray(const TArray<int32>& SelectionArr)
 	return bSuccess;
 }
 
+void FDataflowSelection::SetFromBitArray(TBitArray<>&& Other)
+{
+	SelectionArray = MoveTemp(Other);
+}
+
 void FDataflowSelection::SetFromArray(const TArray<bool>& SelectionArr)
 {
 	SelectionArray.Init(false, SelectionArray.Num());
diff --git a/Engine/Source/Runtime/Experimental/Dataflow/Core/Public/Dataflow/DataflowSelection.h b/Engine/Source/Runtime/Experimental/Dataflow/Core/Public/Dataflow/DataflowSelection.h
--- a/Engine/Source/Runtime/Experimental/Dataflow/Core/Public/Dataflow/DataflowSelection.h
+++ b/Engine/Source/Runtime/Experimental/Dataflow/Core/Public/Dataflow/DataflowSelection.h
@@ -48,6 +48,7 @@ struct FDataflowSelection
 	 * @return true if all indices were successfully set, false if some indices were invalid and could not be set
 	 */
 	DATAFLOWCORE_API bool SetFromArray(const TArray<int32>& SelectionArr);
+	DATAFLOWCORE_API void SetFromBitArray(TBitArray<>&& Other);
 
 	/**
 	* Sets the selection from a dense array (it contains a true/false element for every item)
