IOnDemandDevelopmentExtension - an interface that can alter the behavior when loading unreferenced packages on-demand.

#rnx

#Synced-CL 50957382

diff --git a/Engine/Source/Runtime/Experimental/IoStore/OnDemand/Private/OnDemandInstallCache.cpp b/Engine/Source/Runtime/Experimental/IoStore/OnDemand/Private/OnDemandInstallCache.cpp
--- a/Engine/Source/Runtime/Experimental/IoStore/OnDemand/Private/OnDemandInstallCache.cpp
+++ b/Engine/Source/Runtime/Experimental/IoStore/OnDemand/Private/OnDemandInstallCache.cpp
@@ -35,6 +35,10 @@
 #include "Tasks/Task.h"
 #include "ProfilingDebugging/IoStoreTrace.h"
 
+#if !UE_BUILD_SHIPPING
+#include "IO/OnDemandDevelopmentExtension.h"
+#endif
+
 #if WITH_IOSTORE_ONDEMAND_TESTS
 #include "Algo/Find.h"
 #include "TestHarness.h"
@@ -79,23 +83,6 @@ namespace CVars
 	);
 }
 
-///////////////////////////////////////////////////////////////////////////////
-static void LogUnreferencedChunk(const FIoChunkId& ChunkId, bool bLogAsWarning)
-{
-#if !UE_BUILD_SHIPPING
-	TUtf8StringBuilder<128> Filename;
-	TStringBuilder<128> Msg;
-
-	Msg << TEXT("Trying to load unreferenced chunk, PackageName='");
-	TryConvertChunkIdToPackageName(ChunkId, Filename, Msg);
-	Msg << TEXT("' Filename='") << Filename << TEXT("'");
-	Msg << TEXT(" ChunkId='") << ChunkId << TEXT("' (") << LexToString(ChunkId.GetChunkType()) << TEXT(")");
-
-	UE_CLOGFMT(bLogAsWarning,  LogIoStoreOnDemand, Warning, "{Msg}", Msg);
-	UE_CLOGFMT(!bLogAsWarning, LogIoStoreOnDemand, VeryVerbose, "{Msg}", Msg);
-#endif
-}
-
 ///////////////////////////////////////////////////////////////////////////////
 double ToKiB(uint64 Value)
 {
@@ -2610,13 +2597,25 @@ bool FOnDemandInstallCache::Resolve(FIoRequestImpl* Request)
 	{
 		if (ErrorCode == EIoErrorCode::NotInstalled)
 		{
-			const bool bLogAsWarning = IsDevModeEnabled();
-			LogUnreferencedChunk(Request->ChunkId, bLogAsWarning);
-			if (!IsDevModeEnabled())
+#if !UE_BUILD_SHIPPING
+			if (IOnDemandDevelopmentExtension* Ext = IoStore.TryGetDevelopmentExtension())
 			{
-				CompleteRequest(Request, EIoErrorCode::NotInstalled);
-				return true;
+				TUtf8StringBuilder<128> Filename;
+				TStringBuilder<128> Sb;
+				FStringView PackageName = TryConvertChunkIdToPackageName(Request->ChunkId, Filename, Sb);
+
+				if (Ext->GetUnreferencedIoChunkStatus(
+					Request->ChunkId,
+					Filename.ToView(),
+					PackageName,
+					IsDevModeEnabled()) != EIoErrorCode::NotInstalled)
+				{
+					return false;
+				}
 			}
+#endif			
+			CompleteRequest(Request, EIoErrorCode::NotInstalled);
+			return true;
 		}
 		return false;
 	}
diff --git a/Engine/Source/Runtime/Experimental/IoStore/OnDemand/Private/OnDemandIoStore.cpp b/Engine/Source/Runtime/Experimental/IoStore/OnDemand/Private/OnDemandIoStore.cpp
--- a/Engine/Source/Runtime/Experimental/IoStore/OnDemand/Private/OnDemandIoStore.cpp
+++ b/Engine/Source/Runtime/Experimental/IoStore/OnDemand/Private/OnDemandIoStore.cpp
@@ -2012,6 +2012,11 @@ void FOnDemandIoStore::OnImmediateAnalytic(FOnDemandImmediateAnalyticHandler Eve
 	OnDemandSetImmediateAnalyticHandler(MoveTemp(EventHandler));
 }
 
+void FOnDemandIoStore::SetDevelopmentExtension(TSharedPtr<IOnDemandDevelopmentExtension> Ext)
+{
+	DevExt = Ext;
+}
+
 } // namespace UE::IoStore
 
 ////////////////////////////////////////////////////////////////////////////////
diff --git a/Engine/Source/Runtime/Experimental/IoStore/OnDemand/Private/OnDemandIoStore.h b/Engine/Source/Runtime/Experimental/IoStore/OnDemand/Private/OnDemandIoStore.h
--- a/Engine/Source/Runtime/Experimental/IoStore/OnDemand/Private/OnDemandIoStore.h
+++ b/Engine/Source/Runtime/Experimental/IoStore/OnDemand/Private/OnDemandIoStore.h
@@ -14,6 +14,7 @@
 #include "IO/IoChunkId.h"
 #include "IO/IoHash.h"
 #include "IO/IoStoreOnDemandInternals.h"
+#include "IO/OnDemandDevelopmentExtension.h"
 #include "IO/Serialization/OnDemandContainerToc.h"
 #include "IasHostGroup.h"
 #include "Misc/AES.h"
@@ -371,10 +372,12 @@ public:
 	virtual void						ReportAnalytics(TArray<FAnalyticsEventAttribute>& OutAnalyticsArray) const override;
 	virtual TUniquePtr<IAnalyticsRecording> StartAnalyticsRecording() const override;
 	virtual void						OnImmediateAnalytic(FOnDemandImmediateAnalyticHandler EventHandler) override;
+	virtual void						SetDevelopmentExtension(FSharedOnDemandDevelopmentExtension Ext) override;
 	virtual void						CancelInstallRequest(FSharedInternalInstallRequest InstallRequest) override;
 	virtual void						UpdateInstallRequestPriority(FSharedInternalInstallRequest InstallRequest, int32 NewPriority) override;
 	virtual void						ReleaseContent(FOnDemandInternalContentHandle& ContentHandle) override;
 
+	IOnDemandDevelopmentExtension*		TryGetDevelopmentExtension() { return DevExt.Get(); }
 	TArray<FSharedOnDemandContainer>	GetContainers(EOnDemandContainerFlags ContainerFlags = EOnDemandContainerFlags::None) const;
 	void								FlushLastAccess(FOnDemandFlushLastAccessArgs&& Args, FOnDemandFlushLastAccessCompleted&& OnCompleted);
 
@@ -427,6 +430,7 @@ private:
 	bool									bTicking = false;
 	bool									bTickRequested = false;
 	TFuture<void>							TickFuture;
+	FSharedOnDemandDevelopmentExtension		DevExt;
 };
 
 } // namespace UE::IoStore
diff --git a/Engine/Source/Runtime/Experimental/IoStore/OnDemand/Private/OnDemandPackageStoreBackend.cpp b/Engine/Source/Runtime/Experimental/IoStore/OnDemand/Private/OnDemandPackageStoreBackend.cpp
--- a/Engine/Source/Runtime/Experimental/IoStore/OnDemand/Private/OnDemandPackageStoreBackend.cpp
+++ b/Engine/Source/Runtime/Experimental/IoStore/OnDemand/Private/OnDemandPackageStoreBackend.cpp
@@ -16,36 +16,12 @@
 #include "OnDemandIoStore.h"
 #include "OnDemandMisc.h"
 
-namespace UE::IoStore
-{
-
-///////////////////////////////////////////////////////////////////////////////
-static void LogUnreferencedPackage(FName PackageName, const FPackageId& PackageId, bool bLogAsWarning)
-{
 #if !UE_BUILD_SHIPPING
-	// Only log when querying the package store from the loader
-	if (!IsInAsyncLoadingThread() && !IsInParallelLoadingThread())
-	{
-		return;
-	}
-
-	TStringBuilder<128> Msg;
-	Msg << TEXT("Trying to load unreferenced package '");
-	if (PackageName.IsNone())
-	{
-		TUtf8StringBuilder<128> Filename;
-		TryConvertChunkIdToPackageName(CreatePackageDataChunkId(PackageId), Filename, Msg);
-	}
-	else
-	{
-		Msg << PackageName.ToString();
-	}
-	Msg << TEXT("' (") << PackageId << TEXT(")");
-
-	UE_CLOGFMT(bLogAsWarning,  LogIoStoreOnDemand, Warning, "{Msg}", Msg);
-	UE_CLOGFMT(!bLogAsWarning, LogIoStoreOnDemand, VeryVerbose, "{Msg}", Msg);
+#include "IO/OnDemandDevelopmentExtension.h"
 #endif
-}
+
+namespace UE::IoStore
+{
 
 ///////////////////////////////////////////////////////////////////////////////
 class FOnDemandPackageStoreBackend final
@@ -190,26 +166,48 @@ EPackageStoreEntryStatus FOnDemandPackageStoreBackend::GetPackageStoreEntry(
 	FName PackageName,
 	FPackageStoreEntry& OutPackageStoreEntry)
 {
-	if (const FOnDemandPackageStoreEntry* PackageEntry = FindPackageEntry(PackageId, OutPackageStoreEntry))
+	const FOnDemandPackageStoreEntry* PackageEntry = FindPackageEntry(PackageId, OutPackageStoreEntry);
+	if (PackageEntry == nullptr)
 	{
-		if (PackageEntry->IsReferenced || IsPackageStreamingEnabled())
-		{
-			return EPackageStoreEntryStatus::Ok;
-		}
+		return EPackageStoreEntryStatus::Missing;
+	}
 
-		const bool bLogAsWarning = IsDevModeEnabled();
-		LogUnreferencedPackage(PackageName, PackageId, bLogAsWarning);
+	if (PackageEntry->IsReferenced || IsPackageStreamingEnabled())
+	{
+		return EPackageStoreEntryStatus::Ok;
+	}
 
-		if (IsDevModeEnabled())
+	EPackageStoreEntryStatus EntryStatus = EPackageStoreEntryStatus::NotInstalled;
+#if !UE_BUILD_SHIPPING
+	if (TSharedPtr<FOnDemandIoStore> Pinned = OnDemandIoStore.Pin(); Pinned.IsValid())
+	{
+		if (IOnDemandDevelopmentExtension* Ext = Pinned->TryGetDevelopmentExtension())
 		{
-			// Return missing to let lower priority backends resolve the package store entry
-			return EPackageStoreEntryStatus::Missing;
-		}
+			if (PackageName.IsNone())
+			{
+				TUtf8StringBuilder<128> Filename;
+				TStringBuilder<128> Sb;
 
-		return EPackageStoreEntryStatus::NotInstalled;
+				if (FStringView PackageNameView =
+					TryConvertChunkIdToPackageName(CreatePackageDataChunkId(PackageId), Filename, Sb);
+					PackageNameView.IsEmpty() == false)
+				{
+					PackageName = FName(PackageNameView);
+				}
+			}
+
+			if (EPackageStoreEntryStatus Override = Ext->GetUnreferencedPackageStoreEntryStatus(
+				PackageName,
+				PackageId,
+				OutPackageStoreEntry,
+				IsDevModeEnabled()); Override == EPackageStoreEntryStatus::Missing)
+			{
+				EntryStatus = Override;
+			}
+		}
 	}
-	
-	return EPackageStoreEntryStatus::Missing;
+#endif
+	return EntryStatus;
 }
 
 EPackageStoreEntryStatus FOnDemandPackageStoreBackend::GetPackageStoreEntryInternal(
diff --git a/Engine/Source/Runtime/Experimental/IoStore/OnDemandCore/Public/IO/IoStoreOnDemand.h b/Engine/Source/Runtime/Experimental/IoStore/OnDemandCore/Public/IO/IoStoreOnDemand.h
--- a/Engine/Source/Runtime/Experimental/IoStore/OnDemandCore/Public/IO/IoStoreOnDemand.h
+++ b/Engine/Source/Runtime/Experimental/IoStore/OnDemandCore/Public/IO/IoStoreOnDemand.h
@@ -33,10 +33,12 @@ UE_API DECLARE_LOG_CATEGORY_EXTERN(LogIas, Log, All);
 namespace UE::IoStore
 {
 
-using FSharedOnDemandIoStore		= TSharedPtr<class IOnDemandIoStore, ESPMode::ThreadSafe>;
-using FWeakOnDemandIoStore			= TWeakPtr<class IOnDemandIoStore, ESPMode::ThreadSafe>;
-using FSharedInternalContentHandle	= TSharedPtr<class FOnDemandInternalContentHandle, ESPMode::ThreadSafe>;
-using FSharedInternalInstallRequest	= TSharedPtr<class FOnDemandInternalInstallRequest, ESPMode::ThreadSafe>;
+using FSharedOnDemandIoStore				= TSharedPtr<class IOnDemandIoStore, ESPMode::ThreadSafe>;
+using FWeakOnDemandIoStore					= TWeakPtr<class IOnDemandIoStore, ESPMode::ThreadSafe>;
+using FSharedInternalContentHandle			= TSharedPtr<class FOnDemandInternalContentHandle, ESPMode::ThreadSafe>;
+using FSharedInternalInstallRequest			= TSharedPtr<class FOnDemandInternalInstallRequest, ESPMode::ThreadSafe>;
+using FSharedOnDemandDevelopmentExtension	= TSharedPtr<class IOnDemandDevelopmentExtension, ESPMode::ThreadSafe>;
+;
 
 /** Interface for on-demand request types. */
 class FOnDemandRequest
@@ -707,6 +709,8 @@ public:
 	 * Useful for analytics that are better reported as single events.
 	 */
 	virtual void OnImmediateAnalytic(FOnDemandImmediateAnalyticHandler EventHandler) = 0;
+	/** Set the development extension. */
+	virtual void SetDevelopmentExtension(FSharedOnDemandDevelopmentExtension Ctx) { }
 
 private:
 	/** Initialize the I/O store. Called after the I/O store factory module feature has been registered. */
diff --git a/Engine/Source/Runtime/Experimental/IoStore/OnDemandCore/Public/IO/OnDemandDevelopmentExtension.h b/Engine/Source/Runtime/Experimental/IoStore/OnDemandCore/Public/IO/OnDemandDevelopmentExtension.h
new file mode 100644
--- /dev/null
+++ b/Engine/Source/Runtime/Experimental/IoStore/OnDemandCore/Public/IO/OnDemandDevelopmentExtension.h
@@ -0,0 +1,63 @@
+// Copyright Epic Games, Inc. All Rights Reserved.
+
+#pragma once
+
+#include "CoreTypes.h"
+#include "Containers/StringFwd.h"
+#include "IO/IoStatus.h"
+#include "Serialization/PackageStore.h"
+
+class FIoChunkId;
+
+namespace UE::IoStore
+{
+
+/**
+ * An interface for exceptional scenarios.
+ *
+ * Note: When running in IAX development mode (-Iax.DevMode),
+ * the on-demand package store and install cache backends
+ * are given higher priority than the file-based (installed
+ * content) backend. This allows content that has not yet been
+ * installed via an install request to be loaded by returning
+ * 'Missing' or 'UnknownChunkID'. If the chunk ID exists in the
+ * 'global.umeta' file, the corresponding package and/or
+ * filename will be provided.
+ */
+class IOnDemandDevelopmentExtension
+{
+public:
+	virtual ~IOnDemandDevelopmentExtension() = default;
+
+	/** Called when the on-demand package store finds an unreferenced package store entry
+	 *
+	 * Note: In shipping builds the package store will return 'NotInstalled' for packages
+	 * not pinned in the on-demand install cache. This currently halts the
+	 * package store traversel and affectively fails the package load request.
+	 * Returning 'Missing' will let lower priority package store backends handle
+	 * the request. Use 'IsInAsyncLoadingThread() || IsInParallelLoadingThread()' to only
+	 * trigger calls from the LoadPackage/Async.
+	 */
+	virtual EPackageStoreEntryStatus GetUnreferencedPackageStoreEntryStatus(
+		FName PackageName,
+		const FPackageId& PackageId,
+		const FPackageStoreEntry& Entry,
+		bool bDevMode) = 0;
+	/**
+	 * Called when an unreferenced I/O store chunk is loaded via the I/O dispatcher interface.
+	 * This happens when the package loader has resolved the package store entry or
+	 * when requesting unstructed data via FBulkData.
+	 *
+	 * Note: In shipping builds the on-demand install cache will return 'NotInstalled'
+	 * and fail the request when trying to load an unreferenced I/O store chunk.
+	 * Return 'UnknownChunkID' to let lower priority I/O dispacher backends handle
+	 * the I/O request. 
+	 */
+	virtual EIoErrorCode GetUnreferencedIoChunkStatus(
+		const FIoChunkId& ChunkId,
+		FUtf8StringView Filename,
+		FStringView PackageName,
+		bool bDevMode) = 0;
+};
+
+} // namespace UE::IoStore
