[NetPhysics] Fix issue where render interpolation did not snap when outside of interpolation distance if using the exponential decay version.

#Synced-CL 51043916

diff --git a/Engine/Source/Runtime/Experimental/Chaos/Private/PhysicsProxy/ClusterUnionPhysicsProxy.cpp b/Engine/Source/Runtime/Experimental/Chaos/Private/PhysicsProxy/ClusterUnionPhysicsProxy.cpp
--- a/Engine/Source/Runtime/Experimental/Chaos/Private/PhysicsProxy/ClusterUnionPhysicsProxy.cpp
+++ b/Engine/Source/Runtime/Experimental/Chaos/Private/PhysicsProxy/ClusterUnionPhysicsProxy.cpp
@@ -618,14 +618,13 @@ namespace Chaos
 					return ErrorDistanceSqr < (Velocity * MaxErrorDesyncTime).SizeSquared();
 				};
 
-			// If error is within interpolation limit, set the number of physics frames to interpolate over, else leave at 0 frames to instantly correct the error.
-			int32 RenderInterpErrorCorrectionDurationTicks = 0;
+			// Accumulate error if within interpolation limits
 			if (ErrorDistanceSqr < MaxErrorCorrection || VelocitySnapLimitHelper(Particle_External->V()) || VelocitySnapLimitHelper(PullData.V) || (NextPullData && VelocitySnapLimitHelper(NextPullData->V)))
 			{
-				RenderInterpErrorCorrectionDurationTicks = FMath::FloorToInt32(ErrorCorrectionDuration / AsyncFixedTimeStep); // Convert duration from seconds to simulation ticks
-			}
+				const int32 RenderInterpErrorCorrectionDurationTicks = FMath::FloorToInt32(ErrorCorrectionDuration / AsyncFixedTimeStep); // Convert duration from seconds to simulation ticks
 
-			InterpData->AccumulateErrorXR(Error->ErrorX, Error->ErrorR, SolverSyncTimestamp, RenderInterpErrorCorrectionDurationTicks);
+				InterpData->AccumulateErrorXR(Error->ErrorX, Error->ErrorR, SolverSyncTimestamp, RenderInterpErrorCorrectionDurationTicks);
+			}
 		}
 		const bool bHasInterpolationData = InterpData != nullptr;
 
diff --git a/Engine/Source/Runtime/Experimental/Chaos/Private/PhysicsProxy/GeometryCollectionPhysicsProxy.cpp b/Engine/Source/Runtime/Experimental/Chaos/Private/PhysicsProxy/GeometryCollectionPhysicsProxy.cpp
--- a/Engine/Source/Runtime/Experimental/Chaos/Private/PhysicsProxy/GeometryCollectionPhysicsProxy.cpp
+++ b/Engine/Source/Runtime/Experimental/Chaos/Private/PhysicsProxy/GeometryCollectionPhysicsProxy.cpp
@@ -4695,13 +4695,12 @@ bool FGeometryCollectionPhysicsProxy::PullFromPhysicsState(const Chaos::FDirtyGe
 			ErrorCorrectionDuration = ErrorSettings->ErrorCorrectionDuration;
 		}
 
-		// If error is within interpolation limit, set the number of physics frames to interpolate over, else leave at 0 frames to instantly correct the error.
-		int32 RenderInterpErrorCorrectionDurationTicks = 0;
+		// Accumulate error if within interpolation limits
 		if (Error->ErrorX.SizeSquared() < (MaxErrorCorrection * MaxErrorCorrection))
 		{
-			RenderInterpErrorCorrectionDurationTicks = FMath::FloorToInt32(ErrorCorrectionDuration / AsyncFixedTimeStep); // Convert duration from seconds to simulation ticks
+			const int32 RenderInterpErrorCorrectionDurationTicks = FMath::FloorToInt32(ErrorCorrectionDuration / AsyncFixedTimeStep); // Convert duration from seconds to simulation ticks
+			InterpData->AccumulateErrorXR(Error->ErrorX, Error->ErrorR, SolverSyncTimestamp, RenderInterpErrorCorrectionDurationTicks);
 		}
-		InterpData->AccumulateErrorXR(Error->ErrorX, Error->ErrorR, SolverSyncTimestamp, RenderInterpErrorCorrectionDurationTicks);
 	}
 	const bool bHasInterpolationData = InterpData != nullptr;
 
diff --git a/Engine/Source/Runtime/Experimental/Chaos/Private/PhysicsProxy/SingleParticlePhysicsProxy.cpp b/Engine/Source/Runtime/Experimental/Chaos/Private/PhysicsProxy/SingleParticlePhysicsProxy.cpp
--- a/Engine/Source/Runtime/Experimental/Chaos/Private/PhysicsProxy/SingleParticlePhysicsProxy.cpp
+++ b/Engine/Source/Runtime/Experimental/Chaos/Private/PhysicsProxy/SingleParticlePhysicsProxy.cpp
@@ -366,17 +366,15 @@ bool FSingleParticlePhysicsProxy::PullFromPhysicsState(const Chaos::FDirtyRigidP
 				return ErrorDistanceSqr < (Velocity * MaxErrorDesyncTime).SizeSquared();
 			};
 
-			// If error is within interpolation limit, set the number of physics frames to interpolate over, else leave at 0 frames to instantly correct the error.
-			int32 RenderInterpErrorCorrectionDurationTicks = 0;
+			// Accumulate error if within interpolation limits
 			if (ErrorDistanceSqr < MaxErrorCorrection || VelocitySnapLimitHelper(Rigid->V()) || VelocitySnapLimitHelper(PullData.V) || (NextPullData && VelocitySnapLimitHelper(NextPullData->V)))
 			{
-				RenderInterpErrorCorrectionDurationTicks = FMath::FloorToInt32(ErrorCorrectionDuration / AsyncFixedTimeStep); // Convert duration from seconds to simulation ticks
-			}
+				const int32 RenderInterpErrorCorrectionDurationTicks = FMath::FloorToInt32(ErrorCorrectionDuration / AsyncFixedTimeStep); // Convert duration from seconds to simulation ticks
+				InterpData->AccumulateErrorXR(Error->ErrorX, Error->ErrorR, SolverSyncTimestamp, RenderInterpErrorCorrectionDurationTicks);
 
-			InterpData->AccumulateErrorXR(Error->ErrorX, Error->ErrorR, SolverSyncTimestamp, RenderInterpErrorCorrectionDurationTicks);
-			
-			const int32 RenderInterpErrorVelocitySmoothingDurationTicks = FMath::FloorToInt32(RenderInterpolationCVars::RenderInterpErrorVelocitySmoothingDuration / AsyncFixedTimeStep); // Convert duration from seconds to simulation ticks
-			InterpData->SetVelocitySmoothing(Rigid->V(), Rigid->X(), RenderInterpErrorVelocitySmoothingDurationTicks);
+				const int32 RenderInterpErrorVelocitySmoothingDurationTicks = FMath::FloorToInt32(RenderInterpolationCVars::RenderInterpErrorVelocitySmoothingDuration / AsyncFixedTimeStep); // Convert duration from seconds to simulation ticks
+				InterpData->SetVelocitySmoothing(Rigid->V(), Rigid->X(), RenderInterpErrorVelocitySmoothingDurationTicks);
+			}
 		}
 		const bool bHasInterpolationData = InterpData != nullptr;
 
