[NetPhysics] Add inverse input decay setting to modular vehicle

#Synced-CL 51043916

diff --git a/Engine/Source/Runtime/Experimental/ChaosVehicles/ChaosVehiclesCore/Private/SimModule/ModuleInput.cpp b/Engine/Source/Runtime/Experimental/ChaosVehicles/ChaosVehiclesCore/Private/SimModule/ModuleInput.cpp
--- a/Engine/Source/Runtime/Experimental/ChaosVehicles/ChaosVehiclesCore/Private/SimModule/ModuleInput.cpp
+++ b/Engine/Source/Runtime/Experimental/ChaosVehicles/ChaosVehiclesCore/Private/SimModule/ModuleInput.cpp
@@ -253,6 +253,7 @@ void FModuleInputValue::Serialize(FArchive& Ar, UPackageMap* Map, bool& bOutSucc
 	Ar.SerializeBits(&bIsNonZero, 1);
 	Ar.SerializeBits(&ValueType, 3);
 	Ar.SerializeBits(&bApplyInputDecay, 1);
+	Ar.SerializeBits(&bInverseInputDecay, 1);
 	Ar.SerializeBits(&bClearAfterConsumed, 1);
 	Ar.SerializeBits(&bHasValue, 1);
 
@@ -308,6 +309,7 @@ void FModuleInputValue::NetSerialize(FArchive& Ar, UPackageMap* Map, bool& bOutS
 	{
 		Ar.SerializeBits(&ValueType, 3);
 		Ar.SerializeBits(&bApplyInputDecay, 1);
+		Ar.SerializeBits(&bInverseInputDecay, 1);
 		Ar.SerializeBits(&bClearAfterConsumed, 1);
 	}
 
@@ -371,6 +373,7 @@ void FModuleInputValue::DeltaNetSerialize(FArchive& Ar, UPackageMap* Map, bool&
 		Value = PreviousInputValue.Value;
 		ValueInt = PreviousInputValue.ValueInt;
 		bApplyInputDecay = PreviousInputValue.bApplyInputDecay;
+		bInverseInputDecay = PreviousInputValue.bInverseInputDecay;
 		bClearAfterConsumed = PreviousInputValue.bClearAfterConsumed;
 		bHasValue = PreviousInputValue.bHasValue;
 		return;
@@ -502,6 +505,9 @@ void FModuleInputValue::Decay(const float DecayAmount)
 		return;
 	}
 
+	/** Set the multiplier that that applies the decay amount, decay is passed in as for example 0.1 to decay 10 */
+	const float DecayMultiplier = bInverseInputDecay ? DecayAmount : (1.0f - DecayAmount);
+
 	switch (GetValueType())
 	{
 	case EModuleInputValueType::MBoolean:
@@ -514,7 +520,7 @@ void FModuleInputValue::Decay(const float DecayAmount)
 	case EModuleInputValueType::MAxis2D:
 	case EModuleInputValueType::MAxis3D:
 		// Reduce input value by DecayAmount
-		Value = (1.0f - DecayAmount) * Value;
+		Value *= DecayMultiplier;
 		break;
 	default:
 		checkf(false, TEXT("Unsupported value type for module input value!"));
@@ -610,6 +616,7 @@ void FModuleInputContainer::Initialize(TArray<FModuleInputSetup>& SetupData, FIn
 	{
 		int Index = AddInput(Setup.Type);
 		InputValues[Index].SetApplyInputDecay(Setup.bApplyInputDecay);
+		InputValues[Index].SetInverseInputDecay(Setup.bInverseInputDecay);
 		InputValues[Index].SetClearAfterConsumed(Setup.bClearAfterConsumed);
 		NameMapOut.Add(Setup.Name, Index);
 	}
diff --git a/Engine/Source/Runtime/Experimental/ChaosVehicles/ChaosVehiclesCore/Public/SimModule/ModuleInput.h b/Engine/Source/Runtime/Experimental/ChaosVehicles/ChaosVehiclesCore/Public/SimModule/ModuleInput.h
--- a/Engine/Source/Runtime/Experimental/ChaosVehicles/ChaosVehiclesCore/Public/SimModule/ModuleInput.h
+++ b/Engine/Source/Runtime/Experimental/ChaosVehicles/ChaosVehiclesCore/Public/SimModule/ModuleInput.h
@@ -356,7 +356,11 @@ public:
 	void SetApplyInputDecay(const bool bInApplyInputDecay) { bApplyInputDecay = bInApplyInputDecay; }
 	bool ShouldApplyInputDecay() const { return  bApplyInputDecay; }
 
-	/** After simulation first consumes the input, clear the value doesn't persist across multiple async physics ticks */
+	/** During physics resimulation, if input is being decayed, inverse the decay value, this produces better results sometimes for some inputs like steering */
+	void SetInverseInputDecay(const bool bInInverseInputDecay) { bInverseInputDecay = bInInverseInputDecay; }
+	bool ShouldInverseInputDecay() const { return  bInverseInputDecay; }
+
+	/** After simulation first consumes the input, clear the value so that it  doesn't persist across multiple async physics ticks */
 	void SetClearAfterConsumed(const bool bInClearAfterConsumed) { bClearAfterConsumed = bInClearAfterConsumed; }
 	bool ShouldClearAfterConsumed() const { return  bClearAfterConsumed; }
 
@@ -375,12 +379,17 @@ protected:
 	UPROPERTY()
 	EModuleInputValueType ValueType = EModuleInputValueType::MBoolean;
 
+	/** During physics resimulation, apply decay for this input while it's being extrapolated */
 	UPROPERTY()
 	bool bApplyInputDecay = false;
 
+	/** During physics resimulation, if input is being decayed, inverse the decay value, this produces better results sometimes for some inputs like steering */
 	UPROPERTY()
-	bool bClearAfterConsumed = false;
+	bool bInverseInputDecay = false;
 
+	/** After simulation first consumes the input, clear the value so that it doesn't persist across multiple async physics ticks */
+	UPROPERTY()
+	bool bClearAfterConsumed = false;
 };
 
 // Supported getter specializations
@@ -510,6 +519,7 @@ struct FModuleInputSetup
 	{
 		Type = EModuleInputValueType::MBoolean;
 		bApplyInputDecay = false;
+		bInverseInputDecay = false;
 		bClearAfterConsumed = false;
 	}
 
@@ -517,6 +527,7 @@ struct FModuleInputSetup
 		: Name(InName)
 		, Type(InType)
 		, bApplyInputDecay(false)
+		, bInverseInputDecay(false)
 		, bClearAfterConsumed(false)
 	{
 	}
@@ -536,6 +547,9 @@ struct FModuleInputSetup
 	UPROPERTY(EditAnywhere, Category = VehicleInput)
 	bool bApplyInputDecay;
 
+	UPROPERTY(EditAnywhere, Category = VehicleInput)
+	bool bInverseInputDecay;
+
 	/** After simulation first consumes the input, clear the value doesn't persist across multiple async physics ticks */
 	UPROPERTY(EditAnywhere, Category = VehicleInput)
 	bool bClearAfterConsumed;
