Add a new cvar 'iax.HttpNoRateLimits' that when enabled will disable all limits on bandwidth and allow IAX to run at full speed.

#rnx

- The cvar is compiled out in shipping and is for development only.
- FOnDemandIoBackendStats now calls GetHttpRateLimitKiBPerSecond to find the current bandwidth limit, meaning it will work with this new cvar and the FDownlinkBandwidthManager.
-- Previously we only showed the limit for the 'ias.HttpRateLimitKiBPerSecond' cvar which would be misleading if FDownlinkBandwidthManager was enabled.
- In order for FOnDemandIoBackendStats to be able to call GetHttpRateLimitKiBPerSecond I had to move 'GIasHttpNoRateLimits' to be a global.
-- In a future submit I will clean this up and the bandwidth limit will be set 'onchange' from all cvars and the FDownlinkBandwidthManager, at which point we can push the value to FOnDemandIoBackendStats instead.

#Synced-CL 51086784

diff --git a/Engine/Source/Runtime/Experimental/IoStore/OnDemand/Private/OnDemandHttpIoDispacher.cpp b/Engine/Source/Runtime/Experimental/IoStore/OnDemand/Private/OnDemandHttpIoDispacher.cpp
--- a/Engine/Source/Runtime/Experimental/IoStore/OnDemand/Private/OnDemandHttpIoDispacher.cpp
+++ b/Engine/Source/Runtime/Experimental/IoStore/OnDemand/Private/OnDemandHttpIoDispacher.cpp
@@ -43,8 +43,12 @@
 
 // When enabled the cvar 'iax.InvalidUrlChance' will be enabled and allow us to simulate invalid urls
 #define UE_ALLOW_INVALID_URL_DEBUGGING !UE_BUILD_SHIPPING
+// When enabled it is possible to poison data loaded from the IAS cache via the cvar 'ias.PoisonCache'
 #define UE_ALLOW_CACHE_POISONING !UE_BUILD_SHIPPING
+// When enabled it is possible to pause the http requests via the cvar 'iax.HttpPause'
 #define UE_ALLOW_HTTP_PAUSE (1 && !NO_CVARS && !UE_BUILD_SHIPPING)
+// When enabled it is possible to disable all http bandwidth limits via the 'iax.HttpNoRateLimits' cvar
+#define UE_ALLOW_DISABLE_HTTP_LIMITS !UE_BUILD_SHIPPING
 
 ////////////////////////////////////////////////////////////////////////////////
 const TCHAR* LexToString(EThreadPriority Priority)
@@ -301,6 +305,15 @@ static FAutoConsoleVariableRef CVar_GIasHttpRateLimitKiBPerSecond(
 	TEXT("Http throttle limit in KiBPerSecond")
 );
 
+#if UE_ALLOW_DISABLE_HTTP_LIMITS
+bool GIasHttpNoRateLimits = false;
+static FAutoConsoleVariableRef CVar_GIasHttpNoRateLimit(
+	TEXT("iax.HttpNoRateLimits"),
+	GIasHttpNoRateLimits,
+	TEXT("When true, we ignore all bandwidth limits and completely unthrottle the connection. Not avaliable in UE_BUILD_SHIPPING.")
+);
+#endif // UE_ALLOW_DISABLE_HTTP_LIMITS
+
 int32 GIasHttpPollTimeoutMs = 17;
 static FAutoConsoleVariableRef CVar_GIasHttpPollTimeoutMs(
 	TEXT("ias.HttpPollTimeoutMs"),
@@ -343,6 +356,23 @@ int32 CalculateRequestQueueLengthMax()
 	return FMath::Min(Value, int32(HTTP::FEventLoop::MaxActiveTickets));
 }
 
+////////////////////////////////////////////////////////////////////////////////
+
+// TODO: Should be a member of FHttpDispatcher but needs to be global so that we can expose GetHttpRateLimitKiBPerSecond to FOnDemandIoBackendStats for now
+std::atomic<int32> DLBWHttpRateLimitKiBPerSecond = 0;
+
+int32 GetHttpRateLimitKiBPerSecond()
+{
+#if UE_ALLOW_DISABLE_HTTP_LIMITS
+	if (GIasHttpNoRateLimits)
+	{
+		return 0;
+	}
+#endif // UE_ALLOW_DISABLE_HTTP_LIMITS
+
+	return (FDownlinkBandwidthManager::IsEnabled()) ? DLBWHttpRateLimitKiBPerSecond.load() : GIasHttpRateLimitKiBPerSecond;
+}
+
 ////////////////////////////////////////////////////////////////////////////////
 static constexpr int32 GThreadPriorities[5] =
 {
@@ -1184,7 +1214,7 @@ private:
 	void							RecreateHttpClientIfNeeded();
 
 	bool							TryReadFromCache(FHttpRequest* Request);
-	int32							GetHttpRateLimitKiBPerSecond() const;
+	
 	bool							ProcessHttpRequests();
 	void							CancelRemainingHttpRequests();
 	FHttpRequest*					DequeueHttpRequests();
@@ -1217,7 +1247,7 @@ private:
 	std::atomic_bool						bRecreateHttpClient = false;
 	std::atomic_bool						bStopRequested = false;
 	std::atomic_int32_t						InflightRequestCount{0};
-	FThreadSafeCounter						DLBWHttpRateLimitKiBPerSecond = 0;
+	
 	UE::BandwidthManager::FServiceHandleID	BandwidthManagerServiceHandle = UE_INVALID_BANDWIDTH_SERVICE_HANDLE;
 
 #if UE_ALLOW_HTTP_PAUSE
@@ -1309,10 +1339,10 @@ bool FHttpDispatcher::Init()
 		BandwidthConfig.InGameplayPriority = UE::BandwidthManager::EBandwidthPriority::LowPriority;
 
 		BandwidthManagerServiceHandle = DBManager.RegisterMonitoredService(TEXT("HttpDispatcher"), BandwidthConfig, true,
-			[this](int32 AllottedBandwidthKbps) 
+			[](int32 AllottedBandwidthKbps) 
 			{
 				// 1 Kilobyte is equal to (1000 / 1024) kibibytes
-				DLBWHttpRateLimitKiBPerSecond.Set(static_cast<int32>((float)AllottedBandwidthKbps * (1000.0f / 1024.0f)));
+				DLBWHttpRateLimitKiBPerSecond = static_cast<int32>((float)AllottedBandwidthKbps * (1000.0f / 1024.0f));
 			});
 	}
 
@@ -1769,11 +1799,6 @@ bool FHttpDispatcher::TryReadFromCache(FHttpRequest* Request)
 	return false;
 }
 
-int32 FHttpDispatcher::GetHttpRateLimitKiBPerSecond() const
-{
-	return (FDownlinkBandwidthManager::IsEnabled()) ? DLBWHttpRateLimitKiBPerSecond.GetValue() : GIasHttpRateLimitKiBPerSecond;
-}
-
 bool FHttpDispatcher::ProcessHttpRequests()
 {
 	TRACE_CPUPROFILER_EVENT_SCOPE(FHttpDispatcher::ProcessHttpRequests);
diff --git a/Engine/Source/Runtime/Experimental/IoStore/OnDemand/Private/Statistics.cpp b/Engine/Source/Runtime/Experimental/IoStore/OnDemand/Private/Statistics.cpp
--- a/Engine/Source/Runtime/Experimental/IoStore/OnDemand/Private/Statistics.cpp
+++ b/Engine/Source/Runtime/Experimental/IoStore/OnDemand/Private/Statistics.cpp
@@ -31,6 +31,7 @@ namespace UE::IoStore
 extern int32 GHttpIoDispatcherMaxRangeKiB;
 
 int32 CalculateRequestQueueLengthMax();
+int32 GetHttpRateLimitKiBPerSecond();
 
 float GIasStatisticsLogInterval = 60.f;
 static FAutoConsoleVariableRef CVar_StatisticsLogInterval(
@@ -83,7 +84,6 @@ FAutoConsoleCommand CmdLogOnDemandContentInstallerStats(
 );
 #endif // !UE_BUILD_SHIPPING
 
-extern int32 GIasHttpRateLimitKiBPerSecond;
 extern int32 GIasHttpConcurrentRequests;
 
 ////////////////////////////////////////////////////////////////////////////////
@@ -960,9 +960,10 @@ void FOnDemandIoBackendStats::PrintOnScreenStatistics(FCoreDelegates::FSeverityM
 			}
 
 			TStringBuilder<32> RateLimit;
-			if (GIasHttpRateLimitKiBPerSecond > 0)
+			const int32 RateLimitKiBPerSecond = GetHttpRateLimitKiBPerSecond();
+			if (RateLimitKiBPerSecond > 0)
 			{
-				RateLimit << GIasHttpRateLimitKiBPerSecond;
+				RateLimit << RateLimitKiBPerSecond;
 			}
 			else
 			{
