[Chaos] Improve performance when there are only few constraints.
Tasks are not dispatched anymore in EndCollision if there are only few constraints.
This is valid for RBAN and for the main physics pipeline.

#Synced-CL 50494661

diff --git a/Engine/Source/Runtime/Experimental/Chaos/Private/Chaos/Collision/CollisionConstraintAllocator.cpp b/Engine/Source/Runtime/Experimental/Chaos/Private/Chaos/Collision/CollisionConstraintAllocator.cpp
--- a/Engine/Source/Runtime/Experimental/Chaos/Private/Chaos/Collision/CollisionConstraintAllocator.cpp
+++ b/Engine/Source/Runtime/Experimental/Chaos/Private/Chaos/Collision/CollisionConstraintAllocator.cpp
@@ -5,6 +5,10 @@
 
 namespace Chaos
 {
+
+int32 MinBatchSizeEndCollision = 40;
+FAutoConsoleVariableRef CVarMinBatchSizeEndCollision(TEXT("p.MinBatchSizeEndCollision"), MinBatchSizeEndCollision, TEXT("Minimum number of constraints to batch tasks in EndCollision."));
+
 	namespace Private
 	{
 		FPBDCollisionConstraintPtr FCollisionContextAllocator::CreateConstraint(
@@ -93,11 +97,9 @@ namespace Chaos
 
 		void FCollisionConstraintAllocator::ProcessNewNonCcdConstraints(const TArray<FPBDCollisionConstraint*>& NewConstraints, TArray<FPBDCollisionConstraint*>& AllCCDConstraints)
 		{
-			int32 NumTasks = bSingleWorkerPhysics ? 1 : FMath::Min(FTaskGraphInterface::Get().GetNumWorkerThreads(), Chaos::MaxNumWorkers);
-			NumTasks = FMath::Max(FMath::Min(NumTasks, NewConstraints.Num()), 1);
-
+			const int32 NumTasks = FMath::Max(FMath::Min(FTaskGraphInterface::Get().GetNumWorkerThreads(), Chaos::MaxNumWorkers), 1);
 			const int32 NumConstraints = NewConstraints.Num();
-			const int32 ConstraintByTask = FMath::Max(FMath::DivideAndRoundUp(NumConstraints, NumTasks), 1);
+			const int32 ConstraintByTask = FMath::Max(FMath::DivideAndRoundUp(NumConstraints, NumTasks), MinBatchSizeEndCollision);
 			const int32 NumBatches = FMath::DivideAndRoundUp(NumConstraints, ConstraintByTask);
 
 			if (Chaos::ShouldExecuteTasksInParallel(NumBatches))
@@ -159,11 +161,10 @@ namespace Chaos
 
 		void FCollisionConstraintAllocator::ProcessNewCcdConstraints()
 		{
-			int32 NumTasks = bSingleWorkerPhysics ? 1 : FMath::Min(FTaskGraphInterface::Get().GetNumWorkerThreads(), Chaos::MaxNumWorkers);
-			NumTasks = FMath::Max(FMath::Min(NumTasks, ActiveCCDConstraints.Num()), 1);
 
+			const int32 NumTasks = FMath::Max(FMath::Min(FTaskGraphInterface::Get().GetNumWorkerThreads(), Chaos::MaxNumWorkers), 1);
 			const int32 NumConstraints = ActiveCCDConstraints.Num();
-			const int32 ConstraintByTask = FMath::Max(FMath::DivideAndRoundUp(NumConstraints, NumTasks), 1);
+			const int32 ConstraintByTask = FMath::Max(FMath::DivideAndRoundUp(NumConstraints, NumTasks), MinBatchSizeEndCollision);
 			const int32 NumBatches = FMath::DivideAndRoundUp(NumConstraints, ConstraintByTask);
 
 			if (Chaos::ShouldExecuteTasksInParallel(NumBatches))
@@ -221,7 +222,13 @@ namespace Chaos
 
 		void FCollisionConstraintAllocator::ProcessNewItems()
 		{
-			if (Chaos::ShouldExecuteTasksInParallel(2))
+
+			const int32 NumTasks = FMath::Max(FMath::Min(FTaskGraphInterface::Get().GetNumWorkerThreads(), Chaos::MaxNumWorkers), 1);
+			const int32 NumConstraints = ContextAllocators.IsEmpty() ? 0 : ContextAllocators[0]->NewActiveConstraints.Num();
+			const int32 ConstraintByTask = FMath::Max(FMath::DivideAndRoundUp(NumConstraints, NumTasks), MinBatchSizeEndCollision);
+			const int32 NumBatches = FMath::DivideAndRoundUp(NumConstraints, ConstraintByTask);
+
+			if (Chaos::ShouldExecuteTasksInParallel(NumBatches))
 			{
 				UE::Tasks::FTask PendingTask = UE::Tasks::Launch(UE_SOURCE_LOCATION, [this] 
 				{
diff --git a/Engine/Source/Runtime/Experimental/Chaos/Private/Chaos/Framework/Parallel.cpp b/Engine/Source/Runtime/Experimental/Chaos/Private/Chaos/Framework/Parallel.cpp
--- a/Engine/Source/Runtime/Experimental/Chaos/Private/Chaos/Framework/Parallel.cpp
+++ b/Engine/Source/Runtime/Experimental/Chaos/Private/Chaos/Framework/Parallel.cpp
@@ -16,7 +16,7 @@ namespace Chaos
 	CHAOS_API int32 MaxNumWorkers = 100;
 	CHAOS_API int32 SmallBatchSize = 1;
 	CHAOS_API int32 LargeBatchSize = 100;
-	CHAOS_API int32 MinParallelTaskSize = 1;
+	CHAOS_API int32 MinParallelTaskSize = 2;
 	
 	CHAOS_API bool bDisablePhysicsParallelFor = false;
 	CHAOS_API bool bDisableParticleParallelFor = false;
