Chaos API: Adding scene id and a scene registry.
Scene id is not yet propagated to the handles as that will complicate the review.

#Synced-CL 51097280

diff --git a/Engine/Plugins/Experimental/ChaosRigidPhysicsAsync/Source/ChaosRigidPhysicsAsync/Private/ChaosRigidPhysicsAsync/RigidFactoryAsync.cpp b/Engine/Plugins/Experimental/ChaosRigidPhysicsAsync/Source/ChaosRigidPhysicsAsync/Private/ChaosRigidPhysicsAsync/RigidFactoryAsync.cpp
--- a/Engine/Plugins/Experimental/ChaosRigidPhysicsAsync/Source/ChaosRigidPhysicsAsync/Private/ChaosRigidPhysicsAsync/RigidFactoryAsync.cpp
+++ b/Engine/Plugins/Experimental/ChaosRigidPhysicsAsync/Source/ChaosRigidPhysicsAsync/Private/ChaosRigidPhysicsAsync/RigidFactoryAsync.cpp
@@ -6,6 +6,7 @@
 #include "Logging/StructuredLog.h"
 #include "RigidPhysics/RigidBody.h"
 #include "RigidPhysics/RigidLog.h"
+#include "RigidPhysics/RigidSceneRegistry.h"
 
 #if UE_RIGIDPHYSICS_API_ENABLED
 
@@ -43,6 +44,7 @@ namespace Chaos::Rigids::Async
 		{
 			if (FRigidSceneAsyncGT* Scene = new FRigidSceneAsyncGT(InName, *SettingsAsync))
 			{
+				Scene->SetId(FRigidSceneRegistry::GetInstance().Register(Scene));
 				Scenes.Add(Scene);
 
 				Scene->Startup();
@@ -58,7 +60,7 @@ namespace Chaos::Rigids::Async
 	{
 		using namespace UE::Physics;
 
-		if (IRigidScene* Scene = InSceneHandle.Get())
+		if (IRigidScene* Scene = InSceneHandle.GetScene())
 		{
 			if (FRigidSceneAsyncGT* SceneGT = Scene->AsA<FRigidSceneAsyncGT>())
 			{
@@ -67,6 +69,7 @@ namespace Chaos::Rigids::Async
 				SceneGT->Shutdown();
 
 				Scenes.Remove(SceneGT);
+				FRigidSceneRegistry::GetInstance().Unregister(SceneGT->GetId());
 
 				delete SceneGT;
 			}
diff --git a/Engine/Plugins/Experimental/ChaosRigidPhysicsAsync/Source/ChaosRigidPhysicsAsync/Private/ChaosRigidPhysicsAsync/RigidSceneAsync.cpp b/Engine/Plugins/Experimental/ChaosRigidPhysicsAsync/Source/ChaosRigidPhysicsAsync/Private/ChaosRigidPhysicsAsync/RigidSceneAsync.cpp
--- a/Engine/Plugins/Experimental/ChaosRigidPhysicsAsync/Source/ChaosRigidPhysicsAsync/Private/ChaosRigidPhysicsAsync/RigidSceneAsync.cpp
+++ b/Engine/Plugins/Experimental/ChaosRigidPhysicsAsync/Source/ChaosRigidPhysicsAsync/Private/ChaosRigidPhysicsAsync/RigidSceneAsync.cpp
@@ -119,6 +119,16 @@ namespace Chaos::Rigids::Async
 		InContainer->SetId(FRigidObjectId());
 	}
 
+	UE::Physics::FRigidSceneId FRigidSceneAsyncGT::GetId() const
+	{
+		return SceneId;
+	}
+
+	void FRigidSceneAsyncGT::SetId(const UE::Physics::FRigidSceneId InId)
+	{
+		SceneId = InId;
+	}
+
 	UE::Physics::FRigidSceneHandle FRigidSceneAsyncGT::GetHandle() const
 	{
 		return UE::Physics::FRigidSceneHandle(const_cast<FRigidSceneAsyncGT*>(this));
@@ -508,6 +518,11 @@ namespace Chaos::Rigids::Async
 	void FRigidSceneAsyncPT::Unlock(UE::Physics::ERigidLockType InLockType) const
 	{
 	}
+	
+	UE::Physics::FRigidSceneId FRigidSceneAsyncPT::GetId() const
+	{
+		return SceneHandle.GetId();
+	}
 
 	UE::Physics::FRigidSceneHandle FRigidSceneAsyncPT::GetHandle() const
 	{
diff --git a/Engine/Plugins/Experimental/ChaosRigidPhysicsAsync/Source/ChaosRigidPhysicsAsync/Public/ChaosRigidPhysicsAsync/RigidSceneAsync.h b/Engine/Plugins/Experimental/ChaosRigidPhysicsAsync/Source/ChaosRigidPhysicsAsync/Public/ChaosRigidPhysicsAsync/RigidSceneAsync.h
--- a/Engine/Plugins/Experimental/ChaosRigidPhysicsAsync/Source/ChaosRigidPhysicsAsync/Public/ChaosRigidPhysicsAsync/RigidSceneAsync.h
+++ b/Engine/Plugins/Experimental/ChaosRigidPhysicsAsync/Source/ChaosRigidPhysicsAsync/Public/ChaosRigidPhysicsAsync/RigidSceneAsync.h
@@ -65,6 +65,9 @@ namespace Chaos::Rigids::Async
 		// IRIGIDSCENE BEGIN
 		//
 	public:
+	
+		UE_INTERNAL virtual UE::Physics::FRigidSceneId GetId() const override final;
+		UE_INTERNAL void SetId(const UE::Physics::FRigidSceneId InId);
 		UE_INTERNAL virtual UE::Physics::FRigidSceneHandle GetHandle() const override final;
 		UE_INTERNAL virtual void RegisterModifier(UE::Physics::IRigidSceneModifier* Modifier) override final;
 		UE_INTERNAL virtual void UnregisterModifier(UE::Physics::IRigidSceneModifier* Modifier) override final;
@@ -109,6 +112,7 @@ namespace Chaos::Rigids::Async
 		UE::Physics::FRigidDebugName Name;
 
 		FRigidSceneSettingsAsync SceneSettings;
+		UE::Physics::FRigidSceneId SceneId;
 
 		FChaosScene* Scene;
 		TStrongObjectPtr<UObject> Owner;
@@ -147,6 +151,7 @@ namespace Chaos::Rigids::Async
 		// IRIGIDSCENE BEGIN
 		//
 	public:
+		UE_INTERNAL virtual UE::Physics::FRigidSceneId GetId() const override final;
 		UE_INTERNAL virtual UE::Physics::FRigidSceneHandle GetHandle() const override final;
 		UE_INTERNAL virtual UE::Physics::IRigidBody* GetBody(const UE::Physics::FRigidObjectId& InBodyId) const override final;
 		UE_INTERNAL virtual UE::Physics::IRigidBodyContainer* GetBodyContainer(const UE::Physics::FRigidObjectId& InBodyContainerId) const override final;
diff --git a/Engine/Plugins/Experimental/ChaosRigidPhysicsAsync/Tests/Private/TestGroupEvents.cpp b/Engine/Plugins/Experimental/ChaosRigidPhysicsAsync/Tests/Private/TestGroupEvents.cpp
new file mode 100644
--- /dev/null
+++ b/Engine/Plugins/Experimental/ChaosRigidPhysicsAsync/Tests/Private/TestGroupEvents.cpp
@@ -0,0 +1,13 @@
+// Copyright Epic Games, Inc. All Rights Reserved.
+#include <catch2/catch_test_macros.hpp>
+
+#if UE_RIGIDPHYSICS_API_ENABLED
+
+#include "RigidPhysics/RigidSceneRegistry.h"
+
+GROUP_BEFORE_EACH(Catch::DefaultGroup)
+{
+	UE::Physics::FRigidSceneRegistry::GetInstance().Reset();
+}
+
+#endif // UE_RIGIDPHYSICS_API_ENABLED
diff --git a/Engine/Source/Runtime/Experimental/RigidPhysics/Private/RigidPhysics/RigidSceneHandle.cpp b/Engine/Source/Runtime/Experimental/RigidPhysics/Private/RigidPhysics/RigidSceneHandle.cpp
--- a/Engine/Source/Runtime/Experimental/RigidPhysics/Private/RigidPhysics/RigidSceneHandle.cpp
+++ b/Engine/Source/Runtime/Experimental/RigidPhysics/Private/RigidPhysics/RigidSceneHandle.cpp
@@ -1,12 +1,12 @@
 // Copyright Epic Games, Inc. All Rights Reserved.
 #include "RigidPhysics/RigidSceneHandle.h"
 
+#if UE_RIGIDPHYSICS_API_ENABLED
+
 #include "RigidPhysics/RigidContext.h"
 #include "RigidPhysics/RigidLog.h"
 #include "RigidPhysics/RigidScene.h"
 
-#if UE_RIGIDPHYSICS_API_ENABLED
-
 namespace UE::Physics
 {
 	FRigidSceneHandle::FRigidSceneHandle(FRigidScene* InScene)
@@ -20,13 +20,13 @@ namespace UE::Physics
 	FRigidContextGameRO FRigidSceneHandle::LockRO() const
 	{
 		DoSafetyCheck();
-		return FRigidContextGameRO(SceneRaw);
+		return FRigidContextGameRO(GetScene());
 	}
 
 	FRigidContextGameRW FRigidSceneHandle::LockRW() const
 	{
 		DoSafetyCheck();
-		return FRigidContextGameRW(SceneRaw);
+		return FRigidContextGameRW(GetScene());
 	}
 
 	void FRigidSceneHandle::Reset()
@@ -37,7 +37,12 @@ namespace UE::Physics
 #endif
 	}
 
-	IRigidScene* FRigidSceneHandle::Get() const
+	FRigidSceneId FRigidSceneHandle::GetId() const
+	{
+		return GetScene()->GetId();
+	}
+
+	IRigidScene* FRigidSceneHandle::GetScene() const
 	{
 		return SceneRaw;
 	}
diff --git a/Engine/Source/Runtime/Experimental/RigidPhysics/Private/RigidPhysics/RigidSceneId.cpp b/Engine/Source/Runtime/Experimental/RigidPhysics/Private/RigidPhysics/RigidSceneId.cpp
new file mode 100644
--- /dev/null
+++ b/Engine/Source/Runtime/Experimental/RigidPhysics/Private/RigidPhysics/RigidSceneId.cpp
@@ -0,0 +1,23 @@
+// Copyright Epic Games, Inc. All Rights Reserved.
+#include "RigidPhysics/RigidSceneId.h"
+
+#if UE_RIGIDPHYSICS_API_ENABLED
+
+#include "RigidPhysics/RigidScene.h"
+#include "RigidPhysics/RigidSceneRegistry.h"
+
+namespace UE::Physics
+{
+	IRigidScene* FRigidSceneId::Get() const
+	{
+		FRigidSceneRegistry& Registry = FRigidSceneRegistry::GetInstance();
+		return Registry.Get(*this);
+	}
+
+	void FRigidSceneId::Reset()
+	{
+		Id = INDEX_NONE;
+	}
+} // namespace UE::Physics
+
+#endif // UE_RIGIDPHYSICS_API_ENABLED
diff --git a/Engine/Source/Runtime/Experimental/RigidPhysics/Private/RigidPhysics/RigidSceneRegistry.cpp b/Engine/Source/Runtime/Experimental/RigidPhysics/Private/RigidPhysics/RigidSceneRegistry.cpp
new file mode 100644
--- /dev/null
+++ b/Engine/Source/Runtime/Experimental/RigidPhysics/Private/RigidPhysics/RigidSceneRegistry.cpp
@@ -0,0 +1,67 @@
+// Copyright Epic Games, Inc. All Rights Reserved.
+#include "RigidPhysics/RigidSceneRegistry.h"
+
+#if UE_RIGIDPHYSICS_API_ENABLED
+
+#include "RigidPhysics/RigidLog.h"
+
+namespace UE::Physics
+{
+	FRigidSceneRegistry FRigidSceneRegistry::Instance;
+	
+	FRigidSceneRegistry& FRigidSceneRegistry::GetInstance()
+	{
+		return Instance;
+	}
+
+	FRigidSceneId FRigidSceneRegistry::Register(IRigidScene* Scene)
+	{
+		checkf(SceneCount < MaxScenes, TEXT("FRigidSceneRegistry: Ran out of scenes. Trying to create scene %d when %d is the max"), SceneCount, MaxScenes);
+		if (SceneCount >= MaxScenes)
+		{
+			return FRigidSceneId();
+		}
+		
+		uint8 Id = SceneCount;
+		++SceneCount;
+
+		FRigidSceneId Result;
+		Result.Id = Id;
+		Scenes[Id] = Scene;
+		return Result;
+	}
+
+	void FRigidSceneRegistry::Unregister(const FRigidSceneId Id)
+	{
+		const int32 ArrayIndex = Id.Id;
+		check(ArrayIndex < SceneCount);
+
+		if (ArrayIndex < SceneCount)
+		{
+			// TODO: This is technically not thread safe as someone could be reading this pointer at the same time.
+			// Consider deferring the null for later with a queue.
+			Scenes[ArrayIndex] = nullptr;
+		}
+	}
+
+	IRigidScene* FRigidSceneRegistry::Get(const FRigidSceneId Id) const
+	{
+		const int32 ArrayIndex = Id.Id;
+		if (ArrayIndex < SceneCount)
+		{
+			return Scenes[ArrayIndex];
+		}
+		return nullptr;
+	}
+
+	void FRigidSceneRegistry::Reset()
+	{
+		for (int32 I = 0; I < MaxScenes; ++I)
+		{
+			Scenes[I] = nullptr;
+		}
+		SceneCount = 0;
+	}
+} // namespace UE::Physics
+
+#endif // UE_RIGIDPHYSICS_API_ENABLED
diff --git a/Engine/Source/Runtime/Experimental/RigidPhysics/Public/RigidPhysics/Internal/IRigidScene.h b/Engine/Source/Runtime/Experimental/RigidPhysics/Public/RigidPhysics/Internal/IRigidScene.h
--- a/Engine/Source/Runtime/Experimental/RigidPhysics/Public/RigidPhysics/Internal/IRigidScene.h
+++ b/Engine/Source/Runtime/Experimental/RigidPhysics/Public/RigidPhysics/Internal/IRigidScene.h
@@ -23,6 +23,7 @@ namespace UE::Physics
 		IRigidScene() = default;
 		virtual ~IRigidScene() = default;
 
+		UE_INTERNAL RIGIDPHYSICS_API virtual FRigidSceneId GetId() const = 0;
 		UE_INTERNAL RIGIDPHYSICS_API virtual FRigidSceneHandle GetHandle() const = 0;
 
 		UE_INTERNAL RIGIDPHYSICS_API virtual IRigidBody* GetBody(const FRigidObjectId& InBodyId) const = 0;
diff --git a/Engine/Source/Runtime/Experimental/RigidPhysics/Public/RigidPhysics/RigidFwd.h b/Engine/Source/Runtime/Experimental/RigidPhysics/Public/RigidPhysics/RigidFwd.h
--- a/Engine/Source/Runtime/Experimental/RigidPhysics/Public/RigidPhysics/RigidFwd.h
+++ b/Engine/Source/Runtime/Experimental/RigidPhysics/Public/RigidPhysics/RigidFwd.h
@@ -64,6 +64,7 @@ namespace UE::Physics
 	class FRigidQueryDesc;
 	class FRigidQueryHit;
 	class FRigidSceneHandle;
+	class FRigidSceneId;
 	class FRigidShapeInstance;
 
 	class IJointConstraint;
diff --git a/Engine/Source/Runtime/Experimental/RigidPhysics/Public/RigidPhysics/RigidSceneHandle.h b/Engine/Source/Runtime/Experimental/RigidPhysics/Public/RigidPhysics/RigidSceneHandle.h
--- a/Engine/Source/Runtime/Experimental/RigidPhysics/Public/RigidPhysics/RigidSceneHandle.h
+++ b/Engine/Source/Runtime/Experimental/RigidPhysics/Public/RigidPhysics/RigidSceneHandle.h
@@ -5,6 +5,8 @@
 
 #if UE_RIGIDPHYSICS_API_ENABLED
 
+#include "RigidPhysics/RigidSceneId.h"
+
 namespace UE::Physics
 {
 	class FRigidScene;
@@ -26,7 +28,8 @@ namespace UE::Physics
 		RIGIDPHYSICS_API FRigidContextGameRW LockRW() const;
 
 		UE_INTERNAL RIGIDPHYSICS_API void Reset();
-		UE_INTERNAL RIGIDPHYSICS_API IRigidScene* Get() const;
+		UE_INTERNAL RIGIDPHYSICS_API FRigidSceneId GetId() const;
+		UE_INTERNAL RIGIDPHYSICS_API IRigidScene* GetScene() const;
 
 	private:
 		void DoSafetyCheck() const;
diff --git a/Engine/Source/Runtime/Experimental/RigidPhysics/Public/RigidPhysics/RigidSceneId.h b/Engine/Source/Runtime/Experimental/RigidPhysics/Public/RigidPhysics/RigidSceneId.h
new file mode 100644
--- /dev/null
+++ b/Engine/Source/Runtime/Experimental/RigidPhysics/Public/RigidPhysics/RigidSceneId.h
@@ -0,0 +1,27 @@
+// Copyright Epic Games, Inc. All Rights Reserved.
+#pragma once
+
+#include "RigidPhysics/RigidFwd.h"
+
+#if UE_RIGIDPHYSICS_API_ENABLED
+
+UE_EXPERIMENTAL(5.8, "The new Chaos API is experimental")
+namespace UE::Physics
+{
+	class FRigidSceneId
+	{
+	public:
+		bool operator==(const FRigidSceneId&) const = default;
+		bool operator!=(const FRigidSceneId&) const = default;
+
+		UE_INTERNAL RIGIDPHYSICS_API IRigidScene* Get() const;
+		UE_INTERNAL RIGIDPHYSICS_API void Reset();
+		
+	private:
+		friend class FRigidSceneRegistry;
+
+		uint8 Id = static_cast<uint8>(INDEX_NONE);
+	};
+} // namespace UE::Physics
+
+#endif // UE_RIGIDPHYSICS_API_ENABLED
diff --git a/Engine/Source/Runtime/Experimental/RigidPhysics/Public/RigidPhysics/RigidSceneRegistry.h b/Engine/Source/Runtime/Experimental/RigidPhysics/Public/RigidPhysics/RigidSceneRegistry.h
new file mode 100644
--- /dev/null
+++ b/Engine/Source/Runtime/Experimental/RigidPhysics/Public/RigidPhysics/RigidSceneRegistry.h
@@ -0,0 +1,38 @@
+// Copyright Epic Games, Inc. All Rights Reserved.
+#pragma once
+
+#include "RigidPhysics/RigidFwd.h"
+
+#if UE_RIGIDPHYSICS_API_ENABLED
+
+#include "Containers/StaticArray.h"
+#include "RigidPhysics/RigidSceneId.h"
+
+UE_EXPERIMENTAL(5.8, "The new Chaos API is experimental")
+namespace UE::Physics
+{
+	class FRigidSceneRegistry
+	{
+	public:
+		RIGIDPHYSICS_API static FRigidSceneRegistry& GetInstance();
+		
+		RIGIDPHYSICS_API FRigidSceneId Register(IRigidScene* Scene);
+		RIGIDPHYSICS_API void Unregister(const FRigidSceneId Id);
+
+		RIGIDPHYSICS_API IRigidScene* Get(const FRigidSceneId Id) const;
+
+		// Resets the singleton back to empty. Primarily used for tests.
+		UE_INTERNAL RIGIDPHYSICS_API void Reset();
+
+	private:
+		static constexpr uint32 MaxScenes = 255;
+		static FRigidSceneRegistry Instance;
+
+		FRigidSceneRegistry() = default;
+		
+		uint8 SceneCount = 0;
+		TStaticArray<IRigidScene*, MaxScenes> Scenes;
+	};
+} // namespace UE::Physics
+
+#endif // UE_RIGIDPHYSICS_API_ENABLED
