Remove the cvar 'ias.HttpConcurrentRequests' as it does not do anything anymore.

#rnx

- Removed 'ias.HttpConcurrentRequests' as it is no longer used.
- Added 'iax.Http1RequestPerConnection' which is used to control how many requests are allowed to run on each connection when http1 is used rather than hard coding it to 1.
-- A new cvar was made rather than re-using 'ias.HttpConcurrentRequests' to avoid confusion as the use case for it has slightly changed.
-- This value will be passed to the http client when it is created so changing it will need to trigger ::OnHttpClientCVarChanged so that the client is rebuilt with the new values.
- When creating the http client for http2 there is no real reason to set a limit on the number of requests on the connection (since we only have the 1 connection) so we might as well set it to HTTP::FEventLoop::MaxActiveTickets.
-- This means that 'iax.HttpMaxInflight' now only controls how many requests we issue at once there is no reason for us to recreate the http client when the value changes so it no longer invokes ::OnHttpClientCVarChanged.

#Synced-CL 51086784

diff --git a/Engine/Source/Runtime/Experimental/IoStore/OnDemand/Private/OnDemandHttpClient.cpp b/Engine/Source/Runtime/Experimental/IoStore/OnDemand/Private/OnDemandHttpClient.cpp
--- a/Engine/Source/Runtime/Experimental/IoStore/OnDemand/Private/OnDemandHttpClient.cpp
+++ b/Engine/Source/Runtime/Experimental/IoStore/OnDemand/Private/OnDemandHttpClient.cpp
@@ -17,7 +17,7 @@
 namespace UE::IoStore
 {
 
-extern int32 GIaxHttpMaxInflight;
+extern int32 GIaxHttpOneRequestsPerConnection;
 extern int32 GIaxHttpVersion;
 
 /** Not intended as a long term cvar, only including so that we can revert the new behavior easily if needed */
@@ -776,16 +776,16 @@ TUniquePtr<HTTP::FConnectionPool> FMultiEndpointHttpClient::CreateConnection(FAn
 	bUseHttpTwo &= HostUrl.StartsWith("https://");
 	Params.HttpVersion = bUseHttpTwo ? EHttpVersion::Two : EHttpVersion::One;
 
-	Params.ConnectionCount = uint16(Config.MaxConnectionCount);
-	Params.MaxInflight = uint8(FMath::Clamp(GIaxHttpMaxInflight, 1, 64));
 	if (Params.HttpVersion == EHttpVersion::Two)
 	{
 		Params.ConnectionCount = 1; // rfc9113 conformance
 		Params.VerifyCert = HTTP::FCertRoots::Default();
+		Params.MaxInflight = HTTP::FEventLoop::MaxActiveTickets;
 	}
 	else
 	{
-		Params.MaxInflight = 1;
+		Params.ConnectionCount = uint16(Config.MaxConnectionCount);
+		Params.MaxInflight = (uint8)FMath::Clamp(GIaxHttpOneRequestsPerConnection, 1, MAX_uint8);
 	}
 
 	return MakeUnique<HTTP::FConnectionPool>(Params);
diff --git a/Engine/Source/Runtime/Experimental/IoStore/OnDemand/Private/OnDemandHttpIoDispacher.cpp b/Engine/Source/Runtime/Experimental/IoStore/OnDemand/Private/OnDemandHttpIoDispacher.cpp
--- a/Engine/Source/Runtime/Experimental/IoStore/OnDemand/Private/OnDemandHttpIoDispacher.cpp
+++ b/Engine/Source/Runtime/Experimental/IoStore/OnDemand/Private/OnDemandHttpIoDispacher.cpp
@@ -247,7 +247,14 @@ int32 GIaxHttpMaxInflight = 8;
 static FAutoConsoleVariableRef CVar_IaxHttpMaxInflight(
 	TEXT("iax.HttpMaxInflight"),
 	GIaxHttpMaxInflight,
-	TEXT("Number of requests to issue at once per connection"),
+	TEXT("Number of requests to issue at once")
+);
+
+int32 GIaxHttpOneRequestsPerConnection = 2;
+static FAutoConsoleVariableRef CVar_IaxHttpOneRequestsPerConnection(
+	TEXT("iax.Http1RequestPerConnection"),
+	GIaxHttpOneRequestsPerConnection,
+	TEXT("The number of requests to allow on each connection when GIaxHttpVersion == 1"),
 	FConsoleVariableDelegate::CreateStatic(&OnHttpClientCVarChanged)
 );
 
@@ -291,13 +298,6 @@ static FAutoConsoleVariableRef CVar_IasHttpAllowChunkedXfer(
 	FConsoleVariableDelegate::CreateStatic(&OnHttpClientCVarChanged)
 );
 
-int32 GIasHttpConcurrentRequests = 8;
-static FAutoConsoleVariableRef CVar_IasHttpConcurrentRequests(
-	TEXT("ias.HttpConcurrentRequests"),
-	GIasHttpConcurrentRequests,
-	TEXT("Number of concurrent requests in the http client.")
-);
-
 int32 GIasHttpRateLimitKiBPerSecond = 0;
 static FAutoConsoleVariableRef CVar_GIasHttpRateLimitKiBPerSecond(
 	TEXT("ias.HttpRateLimitKiBPerSecond"),
@@ -351,8 +351,9 @@ int32 CalculateRequestQueueLengthMax()
 	// latency and the effect of CDN cache misses. Reducing what is in the queue
 	// DOES NOT have any effect on throughput throttling.
 	int32 ConnNum = (GIaxHttpVersion == 2) ? 1 : GIasHttpConnectionCount;
-	int32 InflightNum = (GIaxHttpVersion == 2) ? GIaxHttpMaxInflight : 1;
+	int32 InflightNum = (GIaxHttpVersion == 2) ? GIaxHttpMaxInflight : GIaxHttpOneRequestsPerConnection;
 	int32 Value = FMath::Max((ConnNum * InflightNum) + 2, 1);
+
 	return FMath::Min(Value, int32(HTTP::FEventLoop::MaxActiveTickets));
 }
 
diff --git a/Engine/Source/Runtime/Experimental/IoStore/OnDemand/Private/Statistics.cpp b/Engine/Source/Runtime/Experimental/IoStore/OnDemand/Private/Statistics.cpp
--- a/Engine/Source/Runtime/Experimental/IoStore/OnDemand/Private/Statistics.cpp
+++ b/Engine/Source/Runtime/Experimental/IoStore/OnDemand/Private/Statistics.cpp
@@ -84,8 +84,6 @@ FAutoConsoleCommand CmdLogOnDemandContentInstallerStats(
 );
 #endif // !UE_BUILD_SHIPPING
 
-extern int32 GIasHttpConcurrentRequests;
-
 ////////////////////////////////////////////////////////////////////////////////
 static FOnDemandImmediateAnalyticHandler OnDemandImmediateAnalyticEventHandler;
 static TArray<FOnDemandImmediateAnalytic> QueuedOnDemandImmediateAnalyticEventStartupEvents;
