Set the drefault value of 'iax.Http1RequestPerConnection' to 1, to retain the original behavior before the cvar was added.

#rnx

#Synced-CL 51086784

diff --git a/Engine/Source/Runtime/Experimental/IoStore/OnDemand/Private/OnDemandHttpIoDispacher.cpp b/Engine/Source/Runtime/Experimental/IoStore/OnDemand/Private/OnDemandHttpIoDispacher.cpp
--- a/Engine/Source/Runtime/Experimental/IoStore/OnDemand/Private/OnDemandHttpIoDispacher.cpp
+++ b/Engine/Source/Runtime/Experimental/IoStore/OnDemand/Private/OnDemandHttpIoDispacher.cpp
@@ -250,7 +250,7 @@ static FAutoConsoleVariableRef CVar_IaxHttpMaxInflight(
 	TEXT("Number of requests to issue at once")
 );
 
-int32 GIaxHttpOneRequestsPerConnection = 2;
+int32 GIaxHttpOneRequestsPerConnection = 1;
 static FAutoConsoleVariableRef CVar_IaxHttpOneRequestsPerConnection(
 	TEXT("iax.Http1RequestPerConnection"),
 	GIaxHttpOneRequestsPerConnection,
