IAD - Fix timestamp tracking accidentally being compiled out in shipping.

#Synced-CL 51003583

diff --git a/Engine/Source/Runtime/Experimental/IoStore/OnDemand/Private/OnDemandInstallCache.cpp b/Engine/Source/Runtime/Experimental/IoStore/OnDemand/Private/OnDemandInstallCache.cpp
--- a/Engine/Source/Runtime/Experimental/IoStore/OnDemand/Private/OnDemandInstallCache.cpp
+++ b/Engine/Source/Runtime/Experimental/IoStore/OnDemand/Private/OnDemandInstallCache.cpp
@@ -4147,7 +4147,7 @@ FResult FOnDemandInstallCache::FlushPendingChunksImpl(const FPendingChunks& Chun
 			else
 			{
 				const int64 Now = FDateTime::UtcNow().GetTicks();
-				check(Cas.TrackAccessIf(ECasTrackAccessType::Always, CurrentBlockId, Now, bDirty));
+				verify(Cas.TrackAccessIf(ECasTrackAccessType::Always, CurrentBlockId, Now, bDirty));
 				Transaction.BlockAccess(CurrentBlockId, Now);
 			}
 
@@ -4161,7 +4161,7 @@ FResult FOnDemandInstallCache::FlushPendingChunksImpl(const FPendingChunks& Chun
 			else
 			{
 				const int64 Now = FDateTime::UtcNow().GetTicks();
-				check(Cas.TrackModificationIf(ECasTrackModificationType::Always, CurrentBlockId, Now));
+				verify(Cas.TrackModificationIf(ECasTrackModificationType::Always, CurrentBlockId, Now));
 				Transaction.BlockModification(CurrentBlockId, Now);
 			}
 
