[Backout] - CL51243704
#rnx
Original CL Desc
-----------------------------------------------------------------
Chaos Dataflow Solver Actor

diff --git a/Engine/Plugins/Experimental/ChaosDataflowSolver/ChaosDataflowSolver.uplugin b/Engine/Plugins/Experimental/ChaosDataflowSolver/ChaosDataflowSolver.uplugin
deleted file mode 100644
diff --git a/Engine/Plugins/Experimental/ChaosDataflowSolver/Source/ChaosDataflowSolver/ChaosDataflowSolver.Build.cs b/Engine/Plugins/Experimental/ChaosDataflowSolver/Source/ChaosDataflowSolver/ChaosDataflowSolver.Build.cs
deleted file mode 100644
diff --git a/Engine/Plugins/Experimental/ChaosDataflowSolver/Source/ChaosDataflowSolver/Private/ChaosDataflowSolver.cpp b/Engine/Plugins/Experimental/ChaosDataflowSolver/Source/ChaosDataflowSolver/Private/ChaosDataflowSolver.cpp
deleted file mode 100644
diff --git a/Engine/Plugins/Experimental/ChaosDataflowSolver/Source/ChaosDataflowSolver/Private/ChaosDataflowSolverActor.cpp b/Engine/Plugins/Experimental/ChaosDataflowSolver/Source/ChaosDataflowSolver/Private/ChaosDataflowSolverActor.cpp
deleted file mode 100644
diff --git a/Engine/Plugins/Experimental/ChaosDataflowSolver/Source/ChaosDataflowSolver/Private/ChaosSolverBindingComponent.cpp b/Engine/Plugins/Experimental/ChaosDataflowSolver/Source/ChaosDataflowSolver/Private/ChaosSolverBindingComponent.cpp
deleted file mode 100644
diff --git a/Engine/Plugins/Experimental/ChaosDataflowSolver/Source/ChaosDataflowSolver/Private/RigidComponents/InstancedRigidComponentInterface.h b/Engine/Plugins/Experimental/ChaosDataflowSolver/Source/ChaosDataflowSolver/Private/RigidComponents/InstancedRigidComponentInterface.h
deleted file mode 100644
diff --git a/Engine/Plugins/Experimental/ChaosDataflowSolver/Source/ChaosDataflowSolver/Private/RigidComponents/LandscapeRigidComponentInterface.h b/Engine/Plugins/Experimental/ChaosDataflowSolver/Source/ChaosDataflowSolver/Private/RigidComponents/LandscapeRigidComponentInterface.h
deleted file mode 100644
diff --git a/Engine/Plugins/Experimental/ChaosDataflowSolver/Source/ChaosDataflowSolver/Private/RigidComponents/PrimitiveRigidComponentInterface.h b/Engine/Plugins/Experimental/ChaosDataflowSolver/Source/ChaosDataflowSolver/Private/RigidComponents/PrimitiveRigidComponentInterface.h
deleted file mode 100644
diff --git a/Engine/Plugins/Experimental/ChaosDataflowSolver/Source/ChaosDataflowSolver/Private/RigidComponents/SkeletalRigidComponentInterface.h b/Engine/Plugins/Experimental/ChaosDataflowSolver/Source/ChaosDataflowSolver/Private/RigidComponents/SkeletalRigidComponentInterface.h
deleted file mode 100644
diff --git a/Engine/Plugins/Experimental/ChaosDataflowSolver/Source/ChaosDataflowSolver/Public/ChaosDataflowSolver.h b/Engine/Plugins/Experimental/ChaosDataflowSolver/Source/ChaosDataflowSolver/Public/ChaosDataflowSolver.h
deleted file mode 100644
diff --git a/Engine/Plugins/Experimental/ChaosDataflowSolver/Source/ChaosDataflowSolver/Public/ChaosDataflowSolverActor.h b/Engine/Plugins/Experimental/ChaosDataflowSolver/Source/ChaosDataflowSolver/Public/ChaosDataflowSolverActor.h
deleted file mode 100644
diff --git a/Engine/Plugins/Experimental/ChaosDataflowSolver/Source/ChaosDataflowSolver/Public/ChaosSolverBindingComponent.h b/Engine/Plugins/Experimental/ChaosDataflowSolver/Source/ChaosDataflowSolver/Public/ChaosSolverBindingComponent.h
deleted file mode 100644
diff --git a/Engine/Plugins/Experimental/ChaosDataflowSolver/Source/ChaosDataflowSolver/Public/RigidAssetUserData.h b/Engine/Plugins/Experimental/ChaosDataflowSolver/Source/ChaosDataflowSolver/Public/RigidAssetUserData.h
deleted file mode 100644
diff --git a/Engine/Plugins/Experimental/ChaosRigidPhysicsAsync/Source/ChaosRigidPhysicsAsync/ChaosRigidPhysicsAsync.Build.cs b/Engine/Plugins/Experimental/ChaosRigidPhysicsAsync/Source/ChaosRigidPhysicsAsync/ChaosRigidPhysicsAsync.Build.cs
--- a/Engine/Plugins/Experimental/ChaosRigidPhysicsAsync/Source/ChaosRigidPhysicsAsync/ChaosRigidPhysicsAsync.Build.cs
+++ b/Engine/Plugins/Experimental/ChaosRigidPhysicsAsync/Source/ChaosRigidPhysicsAsync/ChaosRigidPhysicsAsync.Build.cs
@@ -14,9 +14,6 @@ namespace UnrealBuildTool.Rules
                 "Core",
                 "CoreUObject",
                 "RigidPhysics",
-				"DataflowCore",
-				"DataflowEngine",
-				"DataflowSimulation",
 				}
 			);
 
diff --git a/Engine/Plugins/Experimental/ChaosRigidPhysicsAsync/Source/ChaosRigidPhysicsAsync/ChaosRigidPhysicsAsyncModule.cpp b/Engine/Plugins/Experimental/ChaosRigidPhysicsAsync/Source/ChaosRigidPhysicsAsync/ChaosRigidPhysicsAsyncModule.cpp
--- a/Engine/Plugins/Experimental/ChaosRigidPhysicsAsync/Source/ChaosRigidPhysicsAsync/ChaosRigidPhysicsAsyncModule.cpp
+++ b/Engine/Plugins/Experimental/ChaosRigidPhysicsAsync/Source/ChaosRigidPhysicsAsync/ChaosRigidPhysicsAsyncModule.cpp
@@ -2,17 +2,4 @@
 
 #include "Modules/ModuleManager.h"
 
-#include "ChaosRigidPhysicsAsync/Dataflow/RigidPhysicsAsyncDataflowNodes.h"
-
-class FChaosRigidPhysicsAsyncModule : public FDefaultModuleImpl
-{
-public:
-	virtual void StartupModule() override
-	{
-		RegisterRigidPhysicsAsyncDataflowSimulationNodes();
-	}
-	virtual void ShutdownModule() override
-	{
-	}
-};
-IMPLEMENT_MODULE(FChaosRigidPhysicsAsyncModule, ChaosRigidPhysicsAsync)
\ No newline at end of file
+IMPLEMENT_MODULE(FDefaultModuleImpl, ChaosRigidPhysicsAsync)
\ No newline at end of file
diff --git a/Engine/Plugins/Experimental/ChaosRigidPhysicsAsync/Source/ChaosRigidPhysicsAsync/Private/ChaosRigidPhysicsAsync/Dataflow/RigidPhysicsAsyncDataflowNodes.cpp b/Engine/Plugins/Experimental/ChaosRigidPhysicsAsync/Source/ChaosRigidPhysicsAsync/Private/ChaosRigidPhysicsAsync/Dataflow/RigidPhysicsAsyncDataflowNodes.cpp
deleted file mode 100644
diff --git a/Engine/Plugins/Experimental/ChaosRigidPhysicsAsync/Source/ChaosRigidPhysicsAsync/Public/ChaosRigidPhysicsAsync/Dataflow/RigidPhysicsAsyncDataflowNodes.h b/Engine/Plugins/Experimental/ChaosRigidPhysicsAsync/Source/ChaosRigidPhysicsAsync/Public/ChaosRigidPhysicsAsync/Dataflow/RigidPhysicsAsyncDataflowNodes.h
deleted file mode 100644
diff --git a/Engine/Source/Runtime/Experimental/ChaosSolverEngine/ChaosSolverEngine.Build.cs b/Engine/Source/Runtime/Experimental/ChaosSolverEngine/ChaosSolverEngine.Build.cs
--- a/Engine/Source/Runtime/Experimental/ChaosSolverEngine/ChaosSolverEngine.Build.cs
+++ b/Engine/Source/Runtime/Experimental/ChaosSolverEngine/ChaosSolverEngine.Build.cs
@@ -26,8 +26,7 @@ namespace UnrealBuildTool.Rules
 			PrivateDependencyModuleNames.AddRange(
 				new string[]
 				{
-					"Messaging",
-					"RigidPhysics",
+					"Messaging"
 				}
 				);
 
diff --git a/Engine/Source/Runtime/Experimental/Dataflow/Simulation/DataflowSimulation.Build.cs b/Engine/Source/Runtime/Experimental/Dataflow/Simulation/DataflowSimulation.Build.cs
--- a/Engine/Source/Runtime/Experimental/Dataflow/Simulation/DataflowSimulation.Build.cs
+++ b/Engine/Source/Runtime/Experimental/Dataflow/Simulation/DataflowSimulation.Build.cs
@@ -14,7 +14,6 @@ namespace UnrealBuildTool.Rules
 					"CoreUObject",
 					"DataflowCore",
 					"DataflowEngine",
-					"RigidPhysics",
 				}
 			);
 			PrivateDependencyModuleNames.AddRange(
diff --git a/Engine/Source/Runtime/Experimental/Dataflow/Simulation/Private/Dataflow/DataflowSimulationContext.cpp b/Engine/Source/Runtime/Experimental/Dataflow/Simulation/Private/Dataflow/DataflowSimulationContext.cpp
--- a/Engine/Source/Runtime/Experimental/Dataflow/Simulation/Private/Dataflow/DataflowSimulationContext.cpp
+++ b/Engine/Source/Runtime/Experimental/Dataflow/Simulation/Private/Dataflow/DataflowSimulationContext.cpp
@@ -1,8 +1,7 @@
-// Copyright Epic Games, Inc. All Rights Reserved.
+﻿// Copyright Epic Games, Inc. All Rights Reserved.
 
 #include "Dataflow/DataflowSimulationContext.h"
 #include "Dataflow/DataflowSimulationProxy.h"
-#include "Dataflow/RigidPhysicsSceneState.h"
 
 namespace UE::Dataflow
 {
@@ -92,33 +91,6 @@ int32 TSimulationContext<Base>::NumSimulationProxies(const FString& ProxyType) c
 	return TypedProxies ? TypedProxies->Num() : 0;
 }
 
-#if UE_RIGIDPHYSICS_API_ENABLED 
-
-template<typename Base>
-FRigidPhysicsSceneDataflowState* TSimulationContext<Base>::GetRigidState()
-{
-	if (!RigidState)
-	{
-		RigidState = MakePimpl<FRigidPhysicsSceneDataflowState>();
-	}
-
-	return RigidState.Get();
-}
-
-template<typename Base>
-void TSimulationContext<Base>::SetRigidSceneName(const FString& InSceneName)
-{
-	SceneName = InSceneName;
-}
-
-template<typename Base>
-const FString& TSimulationContext<Base>::GetRigidSceneName()
-{
-	return SceneName;
-}
-
-#endif
-
 template class TSimulationContext<UE::Dataflow::FContextSingle>;
 template class TSimulationContext<UE::Dataflow::FContextThreaded>;
 	
diff --git a/Engine/Source/Runtime/Experimental/Dataflow/Simulation/Private/Dataflow/DataflowSimulationNodes.cpp b/Engine/Source/Runtime/Experimental/Dataflow/Simulation/Private/Dataflow/DataflowSimulationNodes.cpp
--- a/Engine/Source/Runtime/Experimental/Dataflow/Simulation/Private/Dataflow/DataflowSimulationNodes.cpp
+++ b/Engine/Source/Runtime/Experimental/Dataflow/Simulation/Private/Dataflow/DataflowSimulationNodes.cpp
@@ -6,9 +6,6 @@
 #include "Dataflow/DataflowNodeColorsRegistry.h"
 #include "Dataflow/Interfaces/DataflowPhysicsSolver.h"
 
-#include "RigidPhysics/RigidContext.h"
-#include "RigidPhysics/RigidScene.h"
-
 #include UE_INLINE_GENERATED_CPP_BY_NAME(DataflowSimulationNodes)
 
 namespace UE::Dataflow
@@ -25,14 +22,6 @@ namespace UE::Dataflow
 		// Proxies nodes
 		DATAFLOW_NODE_REGISTER_CREATION_FACTORY(FFilterSimulationProxiesDataflowNode);
 		DATAFLOW_NODE_REGISTER_CREATION_FACTORY(FSimulationProxiesTerminalDataflowNode);
-
-#if UE_RIGIDPHYSICS_API_ENABLED
-		// Scene Setup
-		DATAFLOW_NODE_REGISTER_CREATION_FACTORY(FRigidSceneSetupTerminalDataflowNode);
-
-		DATAFLOW_NODE_REGISTER_CREATION_FACTORY(FAdvanceRigidPhysicsSolverDataflowNode);
-		DATAFLOW_NODE_REGISTER_CREATION_FACTORY(FSceneSimulationTerminalDataflowNode);
-#endif
 		
 		static constexpr FLinearColor CDefaultNodeBodyTintColor = FLinearColor(0.f, 0.f, 0.0f, 0.5f);
 		
@@ -108,37 +97,3 @@ void FSimulationProxiesTerminalDataflowNode::EvaluateSimulation(UE::Dataflow::FD
 {
 	const TArray<FDataflowSimulationProperty> SolverProperties = GetValue(SimulationContext, &SimulationProxies);
 }
-
-void FAdvanceRigidPhysicsSolverDataflowNode::EvaluateSimulation(UE::Dataflow::FDataflowSimulationContext& SimulationContext, const FDataflowOutput* Output) const
-{
-#if UE_RIGIDPHYSICS_API_ENABLED
-	const FRigidPhysicsSceneDataflowState InRigidState = GetValue(SimulationContext, &State);
-
-	const float SimulationDeltaTime = GetValue(SimulationContext, &SimulationTime).DeltaTime;
-	
-	if(FRigidPhysicsSceneDataflowState* RigidState = SimulationContext.GetRigidState())
-	{
-		if (UE::Physics::FRigidContextGameRW Context = RigidState->Handle.LockRW())
-		{
-			Context->Tick(SimulationDeltaTime);
-		}
-	}
-
-	SetValue(SimulationContext, InRigidState, &State);
-#endif
-}
-
-void FSceneSimulationTerminalDataflowNode::EvaluateSimulation(UE::Dataflow::FDataflowSimulationContext& SimulationContext, const FDataflowOutput* Output) const
-{
-#if UE_RIGIDPHYSICS_API_ENABLED
-	FRigidPhysicsSceneDataflowState InState = GetValue(SimulationContext, &State);
-#endif
-}
-
-void FRigidSceneSetupTerminalDataflowNode::EvaluateSimulation(UE::Dataflow::FDataflowSimulationContext& SimulationContext, const FDataflowOutput* Output) const
-{
-#if UE_RIGIDPHYSICS_API_ENABLED
-	FRigidPhysicsSceneDataflowState InState = GetValue(SimulationContext, &State);
-	*SimulationContext.GetRigidState() = InState;
-#endif
-}
\ No newline at end of file
diff --git a/Engine/Source/Runtime/Experimental/Dataflow/Simulation/Private/Dataflow/DataflowSimulationPlugin.cpp b/Engine/Source/Runtime/Experimental/Dataflow/Simulation/Private/Dataflow/DataflowSimulationPlugin.cpp
--- a/Engine/Source/Runtime/Experimental/Dataflow/Simulation/Private/Dataflow/DataflowSimulationPlugin.cpp
+++ b/Engine/Source/Runtime/Experimental/Dataflow/Simulation/Private/Dataflow/DataflowSimulationPlugin.cpp
@@ -1,4 +1,4 @@
-// Copyright Epic Games, Inc. All Rights Reserved.
+﻿// Copyright Epic Games, Inc. All Rights Reserved.
 
 #include "Dataflow/DataflowSimulationPlugin.h"
 #include "Dataflow/DataflowSimulationNodes.h"
@@ -16,11 +16,6 @@ void IDataflowSimulationPlugin::StartupModule()
 	UE::Dataflow::RegisterNodeFilter(FDataflowInvalidNode::StaticType());
 	UE::Dataflow::RegisterNodeFilter(FDataflowExecutionNode::StaticType());
 
-#if UE_RIGIDPHYSICS_API_ENABLED
-	UE::Dataflow::RegisterNodeFilter(FRigidSceneSetupTerminalDataflowNode::StaticType());
-	UE::Dataflow::RegisterNodeFilter(FSceneSimulationTerminalDataflowNode::StaticType());
-#endif
-
 	UDataflowSimulationManager::OnStartup();
 }
 
diff --git a/Engine/Source/Runtime/Experimental/Dataflow/Simulation/Public/Dataflow/DataflowSimulationContext.h b/Engine/Source/Runtime/Experimental/Dataflow/Simulation/Public/Dataflow/DataflowSimulationContext.h
--- a/Engine/Source/Runtime/Experimental/Dataflow/Simulation/Public/Dataflow/DataflowSimulationContext.h
+++ b/Engine/Source/Runtime/Experimental/Dataflow/Simulation/Public/Dataflow/DataflowSimulationContext.h
@@ -7,10 +7,6 @@
 #include "Dataflow/DataflowObjectInterface.h"
 #include "Dataflow/DataflowSimulationProxy.h"
 
-#include "RigidPhysics/RigidFwd.h"
-
-struct FRigidPhysicsSceneDataflowState;
-
 namespace UE::Dataflow
 {
 	/** Simulation context that will be used by all the simulation nodes*/
@@ -60,17 +56,6 @@ namespace UE::Dataflow
 		/** Build group bits based on the string array */
 		DATAFLOWSIMULATION_API void BuildGroupBits(const TArray<FString>& SimulationGroups, TBitArray<>& GroupBits) const;
 
-#if UE_RIGIDPHYSICS_API_ENABLED 
-		/** Return state associated with the rigids dataflow simulation */
-		UE_INTERNAL DATAFLOWSIMULATION_API FRigidPhysicsSceneDataflowState* GetRigidState();
-
-		/** Set the name of scene created through the rigids dataflow */
-		UE_INTERNAL DATAFLOWSIMULATION_API void SetRigidSceneName(const FString& InSceneName);
-
-		/** Get the name of scene created through the rigids dataflow */
-		UE_INTERNAL DATAFLOWSIMULATION_API const FString& GetRigidSceneName();
-#endif
-
 		/** Push another level of iteration indices */
 		void PushIterationIndex()
 		{
@@ -111,11 +96,6 @@ namespace UE::Dataflow
 
 		/** List of all the nested iteration indices */
 		TArray<int32> IterationIndices;
-
-#if UE_RIGIDPHYSICS_API_ENABLED
-		UE_INTERNAL TPimplPtr<FRigidPhysicsSceneDataflowState> RigidState;
-		UE_INTERNAL FString SceneName = "";
-#endif
 	};
 
 	template<typename Base>
diff --git a/Engine/Source/Runtime/Experimental/Dataflow/Simulation/Public/Dataflow/DataflowSimulationNodes.h b/Engine/Source/Runtime/Experimental/Dataflow/Simulation/Public/Dataflow/DataflowSimulationNodes.h
--- a/Engine/Source/Runtime/Experimental/Dataflow/Simulation/Public/Dataflow/DataflowSimulationNodes.h
+++ b/Engine/Source/Runtime/Experimental/Dataflow/Simulation/Public/Dataflow/DataflowSimulationNodes.h
@@ -5,7 +5,6 @@
 #include "Dataflow/DataflowObject.h"
 #include "Dataflow/DataflowSimulationProxy.h"
 #include "Dataflow/DataflowSimulationContext.h"
-#include "Dataflow/RigidPhysicsSceneState.h"
 
 #include "DataflowSimulationNodes.generated.h"
 
@@ -294,74 +293,6 @@ public:
 	virtual void EvaluateSimulation(UE::Dataflow::FDataflowSimulationContext& SimulationContext, const FDataflowOutput* Output) const override;
 };
 
-/** Advance the rigid solver in time */
-USTRUCT(meta = (DataflowSimulation))
-struct FAdvanceRigidPhysicsSolverDataflowNode : public FDataflowSimulationNode
-{
-	GENERATED_BODY()
-	DATAFLOW_NODE_DEFINE_INTERNAL(FAdvanceRigidPhysicsSolverDataflowNode, "AdvanceRigidSolver", "Physics|Solver", UDataflow::SimulationTag)
-
-public:
-	
-	FAdvanceRigidPhysicsSolverDataflowNode(const UE::Dataflow::FNodeParameters& InParam, FGuid InGuid = FGuid::NewGuid())
-	: FDataflowSimulationNode(InParam, InGuid)
-	{
-		RegisterInputConnection(&SimulationTime);
-		RegisterInputConnection(&State);
-		RegisterOutputConnection(&State,&State);
-	}
-
-	/** Delta time to use to advance the solver*/
-	UPROPERTY(Transient, SkipSerialization, Meta = (DataflowInput))
-	FDataflowSimulationTime SimulationTime = FDataflowSimulationTime(0.0f,0.0f);
-
-	UPROPERTY(Transient, SkipSerialization, Meta = (DataflowInput, DataflowOutput, DataflowPassthrough = "State"))
-	FRigidPhysicsSceneDataflowState State;
-
-	/** Evaluate simulation node given a simulation context */
-	virtual void EvaluateSimulation(UE::Dataflow::FDataflowSimulationContext& SimulationContext, const FDataflowOutput* Output) const override;
-};
-
-/** Main terminal node for simulation */
-USTRUCT(meta = (DataflowSimulation))
-struct FSceneSimulationTerminalDataflowNode : public FDataflowSimulationNode
-{
-	GENERATED_BODY()
-	DATAFLOW_NODE_DEFINE_INTERNAL(FSceneSimulationTerminalDataflowNode, "SceneSimulationTerminal", "Terminal|Scene", UDataflow::SimulationTag)
-public:
-	
-	FSceneSimulationTerminalDataflowNode(const UE::Dataflow::FNodeParameters& InParam, FGuid InGuid = FGuid::NewGuid())
-	: FDataflowSimulationNode(InParam, InGuid)
-	{
-		RegisterInputConnection(&State);
-	}
-	UPROPERTY(Transient, SkipSerialization, Meta = (DataflowInput))
-	FRigidPhysicsSceneDataflowState State;
-	
-	/** Evaluate simulation node given a simulation context */
-	virtual void EvaluateSimulation(UE::Dataflow::FDataflowSimulationContext& SimulationContext, const FDataflowOutput* Output) const override;
-};
-
-/** Main terminal node for scene setup */
-USTRUCT(meta = (DataflowSimulation))
-struct FRigidSceneSetupTerminalDataflowNode : public FDataflowSimulationNode
-{
-	GENERATED_BODY()
-	DATAFLOW_NODE_DEFINE_INTERNAL(FRigidSceneSetupTerminalDataflowNode, "RigidSceneSetupTerminal", "Terminal|Scene", UDataflow::SimulationTag)
-public:
-	
-	FRigidSceneSetupTerminalDataflowNode(const UE::Dataflow::FNodeParameters& InParam, FGuid InGuid = FGuid::NewGuid())
-	: FDataflowSimulationNode(InParam, InGuid)
-	{
-		RegisterInputConnection(&State);
-	}
-	UPROPERTY(Transient, SkipSerialization, Meta = (DataflowInput))
-	FRigidPhysicsSceneDataflowState State;
-	
-	/** Evaluate simulation node given a simulation context */
-	virtual void EvaluateSimulation(UE::Dataflow::FDataflowSimulationContext& SimulationContext, const FDataflowOutput* Output) const override;
-};
-
 namespace UE::Dataflow
 {
 	void RegisterDataflowSimulationNodes();
diff --git a/Engine/Source/Runtime/Experimental/Dataflow/Simulation/Public/Dataflow/RigidPhysicsSceneState.h b/Engine/Source/Runtime/Experimental/Dataflow/Simulation/Public/Dataflow/RigidPhysicsSceneState.h
deleted file mode 100644
diff --git a/Engine/Source/Runtime/Experimental/RigidPhysics/Public/RigidPhysics/RigidFwd.h b/Engine/Source/Runtime/Experimental/RigidPhysics/Public/RigidPhysics/RigidFwd.h
--- a/Engine/Source/Runtime/Experimental/RigidPhysics/Public/RigidPhysics/RigidFwd.h
+++ b/Engine/Source/Runtime/Experimental/RigidPhysics/Public/RigidPhysics/RigidFwd.h
@@ -98,7 +98,7 @@ namespace UE::Physics
 
 	using FConstRigidShapeInstancePtr = TRigidShapeInstancePtr<true>;
 	using FRigidShapeInstancePtr = TRigidShapeInstancePtr<false>;
-	using FRigidDebugName = ::Chaos::FSharedDebugName;
+	using FRigidDebugName = Chaos::FSharedDebugName;
 	using FRigidSceneWeakPtr = TWeakPtr<IRigidScene, ESPMode::ThreadSafe>;
 }
 
