[NetPhysics] Fix multiple history registration in a row along with not clearing out history weakptr array correctly.
- Also add guard against resimulation performing a performance collapse.

#Synced-CL 51291913

diff --git a/Engine/Source/Runtime/Engine/Private/PhysicsEngine/NetworkPhysicsComponent.cpp b/Engine/Source/Runtime/Engine/Private/PhysicsEngine/NetworkPhysicsComponent.cpp
--- a/Engine/Source/Runtime/Engine/Private/PhysicsEngine/NetworkPhysicsComponent.cpp
+++ b/Engine/Source/Runtime/Engine/Private/PhysicsEngine/NetworkPhysicsComponent.cpp
@@ -915,6 +915,7 @@ void UNetworkPhysicsComponent::UninitializePhysicsReplication()
 			}
 		}
 	}
+
 	NetworkPhysicsComponent_Internal = nullptr;
 	PhysicsObject = nullptr;
 }
@@ -2229,32 +2230,45 @@ void FAsyncNetworkPhysicsComponent::ConsumeAsyncInput_Internal(const int32 Physi
 			{
 				bCompareInputToTriggerRewind = *AsyncInput->bCompareInputToTriggerRewind;
 			}
-			if (AsyncInput->StateHelper.IsSet())
-			{
-				// Setup rewind data if not already done, and get history size
-				const int32 NumFrames = SetupRewindData();
 
-				StateHelper = (*AsyncInput->StateHelper)->Clone();
+			const bool bStateHelperNeedsUpdate = AsyncInput->StateHelper.IsSet() && AsyncInput->StateHelper.GetValue().IsValid()
+				&& (!StateHelper || StateHelper->GetTypeId() != AsyncInput->StateHelper.GetValue()->GetTypeId());
 
-				// Create state history and local property
-				StateData = StateHelper->CreateUniqueData();
-				StateHistory = MakeShareable(StateHelper->CreateUniqueRewindHistory(NumFrames).Release());
-				RegisterDataHistoryInRewindData();
-				bIsUsingLegacyData = StateHelper->IsUsingLegacyData();
-			}
-			if (AsyncInput->InputHelper.IsSet())
+			const bool bInputHelperNeedsUpdate = AsyncInput->InputHelper.IsSet() && AsyncInput->InputHelper.GetValue().IsValid()
+				&& (!InputHelper || InputHelper->GetTypeId() != AsyncInput->InputHelper.GetValue()->GetTypeId());
+
+			if (bStateHelperNeedsUpdate || bInputHelperNeedsUpdate)
 			{
-				// Setup rewind data if not already done, and get history size
-				const int32 NumFrames = SetupRewindData();
+				UnregisterDataHistoryFromRewindData();
+
+				if (bStateHelperNeedsUpdate)
+				{
+					// Setup rewind data if not already done, and get history size
+					const int32 NumFrames = SetupRewindData();
 
-				InputHelper = (*AsyncInput->InputHelper)->Clone();
+					StateHelper = (*AsyncInput->StateHelper)->Clone();
+
+					// Create state history and local property
+					StateData = StateHelper->CreateUniqueData();
+					StateHistory = MakeShareable(StateHelper->CreateUniqueRewindHistory(NumFrames).Release());
+					bIsUsingLegacyData = StateHelper->IsUsingLegacyData();
+				}
+
+				if (bInputHelperNeedsUpdate)
+				{
+					// Setup rewind data if not already done, and get history size
+					const int32 NumFrames = SetupRewindData();
+
+					InputHelper = (*AsyncInput->InputHelper)->Clone();
+
+					// Create input history and local data properties
+					InputData = InputHelper->CreateUniqueData();
+					LatestInputReceiveData = InputHelper->CreateUniqueData();
+					InputHistory = MakeShareable(InputHelper->CreateUniqueRewindHistory(NumFrames).Release());
+					bIsUsingLegacyData = InputHelper->IsUsingLegacyData();
+				}
 
-				// Create input history and local data properties
-				InputData = InputHelper->CreateUniqueData();
-				LatestInputReceiveData = InputHelper->CreateUniqueData();
-				InputHistory = MakeShareable(InputHelper->CreateUniqueRewindHistory(NumFrames).Release());
 				RegisterDataHistoryInRewindData();
-				bIsUsingLegacyData = InputHelper->IsUsingLegacyData();
 			}
 		}
 
diff --git a/Engine/Source/Runtime/Engine/Public/Physics/NetworkPhysicsComponent.h b/Engine/Source/Runtime/Engine/Public/Physics/NetworkPhysicsComponent.h
--- a/Engine/Source/Runtime/Engine/Public/Physics/NetworkPhysicsComponent.h
+++ b/Engine/Source/Runtime/Engine/Public/Physics/NetworkPhysicsComponent.h
@@ -1309,7 +1309,7 @@ struct FNetworkPhysicsData : public FNetworkPhysicsPayload
 struct FNetworkPhysicsDataHelper
 {
 	virtual ~FNetworkPhysicsDataHelper() = default;
-	
+	virtual const void* GetTypeId() const = 0;
 	virtual TUniquePtr<FNetworkPhysicsDataHelper> Clone() const = 0;
 	virtual TUniquePtr<FNetworkPhysicsPayload> CreateUniqueData() const = 0;
 	virtual TUniquePtr<Chaos::FBaseRewindHistory> CreateUniqueRewindHistory(const int32 Size) const = 0;
@@ -1346,6 +1346,13 @@ template<typename DataType, bool bLegacyData = false>
 struct TNetworkPhysicsDataHelper : FNetworkPhysicsDataHelper
 {
 	~TNetworkPhysicsDataHelper() = default;
+
+	// Implement a static type id, all structs with the same templated types will share the same address to this
+	inline static const int32 StaticTypeId = 0;
+
+	// Get the address to the static type id
+	virtual const void* GetTypeId() const override { return &StaticTypeId; }
+
 	TUniquePtr<FNetworkPhysicsDataHelper> Clone() const override { return MakeUnique<TNetworkPhysicsDataHelper>(*this); }
 	TUniquePtr<FNetworkPhysicsPayload> CreateUniqueData() const override { return MakeUnique<DataType>(); }
 	TUniquePtr<Chaos::FBaseRewindHistory> CreateUniqueRewindHistory(const int32 Size) const override { return MakeUnique<TNetRewindHistory<DataType, bLegacyData>>(Size); }
diff --git a/Engine/Source/Runtime/Experimental/Chaos/Private/PBDRigidsSolver.cpp b/Engine/Source/Runtime/Experimental/Chaos/Private/PBDRigidsSolver.cpp
--- a/Engine/Source/Runtime/Experimental/Chaos/Private/PBDRigidsSolver.cpp
+++ b/Engine/Source/Runtime/Experimental/Chaos/Private/PBDRigidsSolver.cpp
@@ -2144,7 +2144,7 @@ namespace Chaos
 		{
 			const int32 LastStep = MRewindData->CurrentFrame() - 1;
 			const int32 ResimStep = MRewindCallback->TriggerRewindIfNeeded_Internal(LastStep);
-			const int32 NumResimSteps = LastStep - ResimStep + 1;
+			const int32 NumResimSteps = MRewindData->CurrentFrame() - ResimStep;
 
 			if (ResimStep < 0)
 			{
@@ -2154,9 +2154,8 @@ namespace Chaos
 				return;
 			}
 
-			const bool bEnableNetworkPredictionDebug = IsNetworkPhysicsPredictionEnabled() && CanDebugNetworkPhysicsPrediction();
-			
 #if !(UE_BUILD_SHIPPING || UE_BUILD_TEST)
+			const bool bEnableNetworkPredictionDebug = IsNetworkPhysicsPredictionEnabled() && CanDebugNetworkPhysicsPrediction();
 			if (bEnableNetworkPredictionDebug)
 			{
 				UE_LOG(LogChaos, Log, TEXT("	-> Trying to Rewind Frame = %d | At Time = %f | Num Steps = %d | Resim Step = %d | Last Step = %d | Manager Size = %d"), CurrentFrame,  MTime, NumResimSteps, ResimStep, LastStep, MarshallingManager.GetNumHistory_Internal());
@@ -2176,165 +2175,160 @@ namespace Chaos
 			}
 #endif
 
-			if ((ResimStep <= LastStep) && NumResimSteps <= MarshallingManager.GetNumHistory_Internal())
 			{
 				QUICK_SCOPE_CYCLE_COUNTER(ChaosRewindAndResim);
+
 				if (MRewindData->CanRewindToFrame(ResimStep))
 				{
-					MRewindData->RewindToFrame(ResimStep);
-					
+					if (MRewindData->AllowRewindToFrame(ResimStep))
+					{
+						MRewindData->RewindToFrame(ResimStep);
+
 #if DEBUG_REWIND_DATA
-					UE_LOG(LogChaos, Warning, TEXT("CLIENT | PT | ConditionalApplyRewind_Internal | PERFORMING RESIMULATION | Resim From Frame = %d | Num Steps = %d | To Current Frame: %d"), ResimStep, NumResimSteps, CurrentFrame);
+						UE_LOG(LogChaos, Warning, TEXT("CLIENT | PT | ConditionalApplyRewind_Internal | PERFORMING RESIMULATION | Resim From Frame = %d | Num Steps = %d | To Current Frame: %d"), ResimStep, NumResimSteps, CurrentFrame);
 #endif
 
-					SetIsResimming(true);
-					CurrentFrame = ResimStep;
+						SetIsResimming(true);
+						CurrentFrame = ResimStep;
 
-					FPushPhysicsData* ConsumerData = MarshallingManager.GetConsumerData_Internal();
-					TArray<FPushPhysicsData*> RecordedPushData = MarshallingManager.StealHistory_Internal(NumResimSteps);
-					bool bFirst = true;
-
-					FResimDebugInfo DebugInfo;
-					FDurationTimer ResimTimer(DebugInfo.ResimTime);
-					{
-						QUICK_SCOPE_CYCLE_COUNTER(ChaosRewindAndResim_Resim);
+						FPushPhysicsData* ConsumerData = MarshallingManager.GetConsumerData_Internal();
+						TArray<FPushPhysicsData*> RecordedPushData = MarshallingManager.StealHistory_Internal(NumResimSteps);
+						bool bFirst = true;
 
-#if WITH_CHAOS_VISUAL_DEBUGGER
-						bool bIsTracingResimDetails = (FChaosVisualDebuggerTrace::IsTracing() && CVDDC_ResimDetails->IsChannelEnabled());
-						if (bIsTracingResimDetails)
+						FResimDebugInfo DebugInfo;
+						FDurationTimer ResimTimer(DebugInfo.ResimTime);
 						{
-							// This frame is redundant with the one in AdvanceSolver, but if we don't do this, subsequent particle traces and steps will appear as the last steps of the current frame we rewind from
-							// and CVD will not correctly advance the server timeline to the frame we are landing on, making it harder to compare the target set locally with server state at the same frame
-							CVD_TRACE_SOLVER_START_FRAME(FPhysicsSolverBase, *this);
-							CVD_SCOPE_TRACE_SOLVER_STEP(CVDDC_ResimDetails, TEXT("Rewind To Frame"));
-							MRewindData->TraceParticlesModifiedDuringRewind();
-						}
-#endif
+							QUICK_SCOPE_CYCLE_COUNTER(ChaosRewindAndResim_Resim);
 
-						// Do rollback as necessary
-						for (int32 Step = ResimStep; Step <= LastStep; ++Step)
-						{
 #if WITH_CHAOS_VISUAL_DEBUGGER
-							if (!bFirst && bIsTracingResimDetails)
+							bool bIsTracingResimDetails = (FChaosVisualDebuggerTrace::IsTracing() && CVDDC_ResimDetails->IsChannelEnabled());
+							if (bIsTracingResimDetails)
 							{
+								// This frame is redundant with the one in AdvanceSolver, but if we don't do this, subsequent particle traces and steps will appear as the last steps of the current frame we rewind from
+								// and CVD will not correctly advance the server timeline to the frame we are landing on, making it harder to compare the target set locally with server state at the same frame
 								CVD_TRACE_SOLVER_START_FRAME(FPhysicsSolverBase, *this);
-								CVD_SCOPE_TRACE_SOLVER_STEP(CVDDC_ResimDetails, TEXT("Resim Step Start"));
-								CVD_TRACE_PARTICLES_SOA(GetEvolution()->GetParticles(), &GetEvolution()->GetRigidClustering());
+								CVD_SCOPE_TRACE_SOLVER_STEP(CVDDC_ResimDetails, TEXT("Rewind To Frame"));
+								MRewindData->TraceParticlesModifiedDuringRewind();
 							}
 #endif
 
-							if ((LastStep - Step) < RecordedPushData.Num())
+							// Do rollback as necessary
+							for (int32 Step = ResimStep; Step <= LastStep; ++Step)
 							{
-								if (!bFirst && !bUseCollisionResimCache)
+#if WITH_CHAOS_VISUAL_DEBUGGER
+								if (!bFirst && bIsTracingResimDetails)
 								{
-									MRewindData->StepNonResimParticles(Step);
+									CVD_TRACE_SOLVER_START_FRAME(FPhysicsSolverBase, *this);
+									CVD_SCOPE_TRACE_SOLVER_STEP(CVDDC_ResimDetails, TEXT("Resim Step Start"));
+									CVD_TRACE_PARTICLES_SOA(GetEvolution()->GetParticles(), &GetEvolution()->GetRigidClustering());
 								}
+#endif
 
-								if (PhysicsReplicationCVars::ResimulationCVars::bApplyTargetsWhileResimulating || bFirst)
+								if ((LastStep - Step) < RecordedPushData.Num())
 								{
-									// Update all the particles having received a target from the server
-									MRewindData->ApplyTargets(Step, bFirst);
-								}
+									if (!bFirst && !bUseCollisionResimCache)
+									{
+										MRewindData->StepNonResimParticles(Step);
+									}
 
-								FPushPhysicsData* PushData = RecordedPushData[LastStep - Step];	//push data is sorted as latest first
-								if (bFirst)
-								{
-									SCOPED_NAMED_EVENT(ChaosRewindAndResim_ResetCollisions, FColor::Red);
-									MTime = PushData->StartTime;	//not sure if sub-steps have proper StartTime so just do this once and let solver evolve remaining time
+									if (PhysicsReplicationCVars::ResimulationCVars::bApplyTargetsWhileResimulating || bFirst)
+									{
+										// Update all the particles having received a target from the server
+										MRewindData->ApplyTargets(Step, bFirst);
+									}
 
-									GetEvolution()->ResetCollisions();
-								}
+									FPushPhysicsData* PushData = RecordedPushData[LastStep - Step];	//push data is sorted as latest first
+									if (bFirst)
+									{
+										SCOPED_NAMED_EVENT(ChaosRewindAndResim_ResetCollisions, FColor::Red);
+										MTime = PushData->StartTime;	//not sure if sub-steps have proper StartTime so just do this once and let solver evolve remaining time
+
+										GetEvolution()->ResetCollisions();
+									}
 
-								GetEvolution()->SetReset(bFirst);
+									GetEvolution()->SetReset(bFirst);
 
 #if !(UE_BUILD_SHIPPING || UE_BUILD_TEST)
-								if (bEnableNetworkPredictionDebug)
-								{
-									UE_LOG(LogChaos, Log, TEXT("		-> Re-simulating Frame = %d "), Step);
-									for (auto& Handle : GetParticles().GetNonDisabledDynamicView())
+									if (bEnableNetworkPredictionDebug)
 									{
-										if (GetEvolution()->GetIslandManager().GetParticleIsland(Handle.Handle()))
+										UE_LOG(LogChaos, Log, TEXT("		-> Re-simulating Frame = %d "), Step);
+										for (auto& Handle : GetParticles().GetNonDisabledDynamicView())
 										{
-											if (!Handle.IsSleeping())
-											{
-												UE_LOG(LogChaos, Log, TEXT("Particle Dynamic At Position = %s | Velocity = %s | Quaternion = %s | Omega = %s | Resim Frame = %d | Sync State = %d | Needs Resim = %d | Unique Idx = %d"),
-													*Handle.GetX().ToString(), *Handle.GetV().ToString(), *Handle.GetR().ToString(), *Handle.GetW().ToString(),
-													GetEvolution()->GetIslandManager().GetParticleIsland(Handle.Handle())->GetResimFrame(), (uint8)Handle.SyncState(),
-													GetEvolution()->GetIslandManager().GetParticleIsland(Handle.Handle())->NeedsResim(), Handle.UniqueIdx().Idx);
-											}
-											else
+											if (GetEvolution()->GetIslandManager().GetParticleIsland(Handle.Handle()))
 											{
-												UE_LOG(LogChaos, Log, TEXT("Particle Sleeping At Position = %s | Velocity = %s | Quaternion = %s | Omega = %s | Resim Frame = %d | Sync State = %d | Needs Resim = %d | Unique Idx = %d"),
-													*Handle.GetX().ToString(), *Handle.GetV().ToString(), *Handle.GetR().ToString(), *Handle.GetW().ToString(),
-													GetEvolution()->GetIslandManager().GetParticleIsland(Handle.Handle())->GetResimFrame(), (uint8)Handle.SyncState(),
-													GetEvolution()->GetIslandManager().GetParticleIsland(Handle.Handle())->NeedsResim(), Handle.UniqueIdx().Idx);
+												if (!Handle.IsSleeping())
+												{
+													UE_LOG(LogChaos, Log, TEXT("Particle Dynamic At Position = %s | Velocity = %s | Quaternion = %s | Omega = %s | Resim Frame = %d | Sync State = %d | Needs Resim = %d | Unique Idx = %d"),
+														*Handle.GetX().ToString(), *Handle.GetV().ToString(), *Handle.GetR().ToString(), *Handle.GetW().ToString(),
+														GetEvolution()->GetIslandManager().GetParticleIsland(Handle.Handle())->GetResimFrame(), (uint8)Handle.SyncState(),
+														GetEvolution()->GetIslandManager().GetParticleIsland(Handle.Handle())->NeedsResim(), Handle.UniqueIdx().Idx);
+												}
+												else
+												{
+													UE_LOG(LogChaos, Log, TEXT("Particle Sleeping At Position = %s | Velocity = %s | Quaternion = %s | Omega = %s | Resim Frame = %d | Sync State = %d | Needs Resim = %d | Unique Idx = %d"),
+														*Handle.GetX().ToString(), *Handle.GetV().ToString(), *Handle.GetR().ToString(), *Handle.GetW().ToString(),
+														GetEvolution()->GetIslandManager().GetParticleIsland(Handle.Handle())->GetResimFrame(), (uint8)Handle.SyncState(),
+														GetEvolution()->GetIslandManager().GetParticleIsland(Handle.Handle())->NeedsResim(), Handle.UniqueIdx().Idx);
+												}
 											}
 										}
 									}
-								}
 #endif
 #if WITH_CHAOS_VISUAL_DEBUGGER
-								if (bFirst && bIsTracingResimDetails)
-								{
-									CVD_TRACE_SOLVER_END_FRAME(FPhysicsSolverBase, *this);
-									CVD_TRACE_SOLVER_START_FRAME(FPhysicsSolverBase, *this);
-									CVD_SCOPE_TRACE_SOLVER_STEP(CVDDC_ResimDetails, TEXT("Resim Step Start"));
-									CVD_TRACE_PARTICLES_SOA(GetEvolution()->GetParticles(), &GetEvolution()->GetRigidClustering());
-								}
+									if (bFirst && bIsTracingResimDetails)
+									{
+										CVD_TRACE_SOLVER_END_FRAME(FPhysicsSolverBase, *this);
+										CVD_TRACE_SOLVER_START_FRAME(FPhysicsSolverBase, *this);
+										CVD_SCOPE_TRACE_SOLVER_STEP(CVDDC_ResimDetails, TEXT("Resim Step Start"));
+										CVD_TRACE_PARTICLES_SOA(GetEvolution()->GetParticles(), &GetEvolution()->GetRigidClustering());
+									}
 #endif
 
-								{
-									QUICK_SCOPE_CYCLE_COUNTER(ChaosRewindAndResim_RewindCallback_PreResimStep_Internal);
-									MRewindCallback->PreResimStep_Internal(Step, bFirst);
-								}
+									{
+										QUICK_SCOPE_CYCLE_COUNTER(ChaosRewindAndResim_RewindCallback_PreResimStep_Internal);
+										MRewindCallback->PreResimStep_Internal(Step, bFirst);
+									}
 
 #if WITH_CHAOS_VISUAL_DEBUGGER
-								if (bIsTracingResimDetails)
-								{
-									CVD_TRACE_SOLVER_END_FRAME(FPhysicsSolverBase, *this);
-								}
+									if (bIsTracingResimDetails)
+									{
+										CVD_TRACE_SOLVER_END_FRAME(FPhysicsSolverBase, *this);
+									}
 #endif
-								// Run the advance but omit the game thread callbacks as we're executing during the physics tick
-								FSolverTasksPTOnly ImmediateTask(*this, PushData);
-								//ensure(bSolverHasFrozenGameThreadCallbacks == false);	//We don't support this for resim as it's very expensive and difficult to schedule
-								ImmediateTask.AdvanceSolver();
-								{
-									QUICK_SCOPE_CYCLE_COUNTER(ChaosRewindAndResim_RewindCallback_PostResimStep_Internal);
-									MRewindCallback->PostResimStep_Internal(Step);
-								}
+									// Run the advance but omit the game thread callbacks as we're executing during the physics tick
+									FSolverTasksPTOnly ImmediateTask(*this, PushData);
+									//ensure(bSolverHasFrozenGameThreadCallbacks == false);	//We don't support this for resim as it's very expensive and difficult to schedule
+									ImmediateTask.AdvanceSolver();
+									{
+										QUICK_SCOPE_CYCLE_COUNTER(ChaosRewindAndResim_RewindCallback_PostResimStep_Internal);
+										MRewindCallback->PostResimStep_Internal(Step);
+									}
 
-								bFirst = false;
+									bFirst = false;
+								}
 							}
 						}
-					}
 
-					// Clean up after resim. This is where particles are unfrozen if they were frozen due to bubble resim.
-					MRewindData->EndResim();
+						// Clean up after resim. This is where particles are unfrozen if they were frozen due to bubble resim.
+						MRewindData->EndResim();
 
-					MarshallingManager.SetConsumerData_Internal(ConsumerData);
-					GetEvolution()->GetIslandManager().ResetParticleResimFrame();
+						MarshallingManager.SetConsumerData_Internal(ConsumerData);
+						GetEvolution()->GetIslandManager().ResetParticleResimFrame();
 
-					SetIsResimming(false);
-					GetEvolution()->SetReset(false);
+						SetIsResimming(false);
+						GetEvolution()->SetReset(false);
 
-					ResimTimer.Stop();
-					MRewindCallback->SetResimDebugInfo_Internal(DebugInfo);
+						ResimTimer.Stop();
+						MRewindCallback->SetResimDebugInfo_Internal(DebugInfo);
+					}
 				}
 #if DEBUG_REWIND_DATA
-				else
+				else // !CanRewindToFrame
 				{
 					UE_LOG(LogChaos, Log, TEXT("CLIENT | PT | ConditionalApplyRewind_Internal | Resimulation failed, FRewindData::CanRewindToFrame returned false | Current Frame = %d | Num Steps = %d | Resim Frame = %d | Last Frame = %d | Rewind History Size = %d"), CurrentFrame, NumResimSteps, ResimStep, LastStep, MarshallingManager.GetNumHistory_Internal());
 				}
 #endif
 			}
-#if DEBUG_REWIND_DATA
-			else
-			{
-				UE_LOG(LogChaos, Log, TEXT("CLIENT | PT | ConditionalApplyRewind_Internal | Resimulation failed, invalid rewind frame data | Current Frame = %d | Num Steps = %d | Resim Frame = %d | Last Frame = %d | Rewind History Size = %d"), CurrentFrame, NumResimSteps, ResimStep, LastStep, MarshallingManager.GetNumHistory_Internal());
-			}
-#endif
-			// Clear the ResimFrame no matter if resimulation succeeded or failed (if it failed it's not going to succeed next frame either based on the same ResimFrame)
-			MRewindData->SetResimFrame(INDEX_NONE);
 		}
 	}
 
diff --git a/Engine/Source/Runtime/Experimental/Chaos/Private/RewindData.cpp b/Engine/Source/Runtime/Experimental/Chaos/Private/RewindData.cpp
--- a/Engine/Source/Runtime/Experimental/Chaos/Private/RewindData.cpp
+++ b/Engine/Source/Runtime/Experimental/Chaos/Private/RewindData.cpp
@@ -864,31 +864,91 @@ void FRewindData::UnfreezeParticlesOutsideResimBubbles()
 bool bForceSkipResim = false;
 FAutoConsoleVariableRef CVarForceSkipResim(TEXT("p.Resim.ForceSkipResim"), bForceSkipResim, TEXT("Don't perform resimulations while this is true"));
 
-bool FRewindData::CanRewindToFrame(int32 RewindFrame) const
+bool bGuardFromPerformanceCollapse = true;
+FAutoConsoleVariableRef CVarResimGuardFromPerformanceCollapse(TEXT("p.Resim.GuardFromPerformanceCollapse"), bGuardFromPerformanceCollapse, TEXT("Don't perform another resimulation before as much realtime has passed as last resim took."));
+
+bool FRewindData::AllowRewindToFrame(int32 RewindFrame) const
 {
+	// Don't perform a resimulation if we are currently under cooldown from a previous resimulation
+	if (LatestFrame < ResimCooldownUntilFrame)
+	{
+		return false;
+	}
+
+	// Stop resim from causing a performance collapse by taking longer to calculate than physics has moved forward in realtime
+	if (bGuardFromPerformanceCollapse && ResimPerformanceCooldownCycleCounter > FPlatformTime::Cycles64())
+	{
+		return false;
+	}
+
+	return true;
+}
+
+bool FRewindData::CanRewindToFrame(int32 RewindFrame)
+{
+	auto ClearResimHelper = [&]()
+		{
+			// Clear ResimFrame since we will never be able to resim to this frame
+			SetResimFrame(INDEX_NONE);
+			if (Solver && Solver->GetEvolution())
+			{
+				Solver->GetEvolution()->GetIslandManager().ResetParticleResimFrame();
+			}
+		};
+
 	if (bForceSkipResim)
 	{
+		ClearResimHelper();
+		return false;
+	}
+
+	// We can't rewind forward
+	if (RewindFrame > LatestFrame)
+	{
+		if (Chaos::ResimConsoleVars::bEnableRewindLogs)
+		{
+			UE_LOG(LogChaos, Log, TEXT("CLIENT | PT | CanRewindToFrame | Failed due to rewind frame being later than latest simulated frame | Rewind Frame: %d | Latest Frame: %d"), RewindFrame, LatestFrame);
+		}
+
+		ClearResimHelper();
+		return false;
+	}
+
+	// We need marshaling manager data for the number of frames back we want to resim
+	const int32 NumResimFrames = LatestFrame - RewindFrame;
+	if (Solver && NumResimFrames > Solver->GetMarshallingManager().GetNumHistory_Internal())
+	{
+		if (Chaos::ResimConsoleVars::bEnableRewindLogs)
+		{
+			UE_LOG(LogChaos, Log, TEXT("CLIENT | PT | CanRewindToFrame | Failed due to NumResimFrames is larger than the cached data in marshaling manager | Rewind Frame: %d | Resim Frames: %d | Marshaling History Count:: %d"), RewindFrame, NumResimFrames, Solver->GetMarshallingManager().GetNumHistory_Internal());
+		}
+
+		ClearResimHelper();
 		return false;
 	}
 
-	//Can't go too far back
+	// We can't rewind outside of the rewind history
 	const int32 EarliestFrame = GetEarliestFrame_Internal();
 	if (RewindFrame < EarliestFrame)
 	{
 		if (Chaos::ResimConsoleVars::bEnableRewindLogs)
 		{
-			UE_LOG(LogChaos, Log, TEXT("CLIENT | PT | RewindToFrame | Failed due to rewind frame earlier than available history | Rewind Frame: %d | Earliest Frame: %d"), RewindFrame, EarliestFrame);
+			UE_LOG(LogChaos, Log, TEXT("CLIENT | PT | CanRewindToFrame | Failed due to rewind frame earlier than available history | Rewind Frame: %d | Earliest Frame: %d"), RewindFrame, EarliestFrame);
 		}
+
+		ClearResimHelper();
 		return false;
 	}
 
-	//If we need to save and we are right on the edge of the buffer, we can't go back to earliest frame
+	// If we need to save and we are right on the edge of the buffer, we can't go back to earliest frame
 	if (RewindFrame == EarliestFrame && bNeedsSave && FramesSaved == Managers.Capacity())
 	{
 		if (Chaos::ResimConsoleVars::bEnableRewindLogs)
 		{
-			UE_LOG(LogChaos, Log, TEXT("CLIENT | PT | RewindToFrame | Failed due to rewinding to last available frame and bNeedsSave is set to true"));
+			UE_LOG(LogChaos, Log, TEXT("CLIENT | PT | CanRewindToFrame | Failed due to rewinding to last available frame and bNeedsSave is set to true"));
 		}
+
+		ClearResimHelper();
 		return false;
 	}
 
@@ -912,6 +972,9 @@ void FRewindData::RewindToFrame(int32 RewindFrame)
 		return;
 	}
 
+	// Start measuring resim time
+	ResimStartCycleCounter = FPlatformTime::Cycles64();
+
 	QUICK_SCOPE_CYCLE_COUNTER(RewindToFrame);
 
 	EnsureIsInPhysicsThreadContext();
@@ -1153,6 +1216,10 @@ void FRewindData::EndResim()
 #if WITH_CHAOS_VISUAL_DEBUGGER
 	ParticlesModifiedDuringRewind.Empty();
 #endif
+
+	// Measure the resim time and set a cooldown for that amount of time
+	const uint64 CyclesCount = FPlatformTime::Cycles64() - ResimStartCycleCounter;
+	ResimPerformanceCooldownCycleCounter = FPlatformTime::Cycles64() + CyclesCount;
 }
 
 void FRewindData::StepNonResimParticles(const int32 Frame)
@@ -1173,14 +1240,15 @@ void FRewindData::StepNonResimParticles(const int32 Frame)
 	for (FDirtyParticleInfo& DirtyParticleInfo : DirtyParticles)
 	{
 		FGeometryParticleHandle* GeometryParticle = DirtyParticleInfo.GetObjectPtr();
-		FKinematicGeometryParticleHandle* KinematicParticle = GeometryParticle->CastToKinematicParticle();
-		FPBDRigidParticleHandle* RigidParticle = GeometryParticle->CastToRigidParticle();
 
 		if (GeometryParticle->ResimType() != EResimType::ResimAsFollower)
 		{
 			continue;
 		}
 
+		FKinematicGeometryParticleHandle* KinematicParticle = GeometryParticle->CastToKinematicParticle();
+		FPBDRigidParticleHandle* RigidParticle = GeometryParticle->CastToRigidParticle();
+
 		bool bHasChanged = false;
 		const FGeometryParticleStateBase& History = DirtyParticleInfo.GetHistory();
 
@@ -1394,7 +1462,7 @@ void FRewindData::FinishFrame()
 			FinalResimHelper(DirtyJoints);
 		}
 	}
-	
+
 	++CurFrame;
 	LatestFrame = FMath::Max(LatestFrame, CurFrame);
 }
@@ -2747,7 +2815,7 @@ int32 FRewindData::GetLatestReceivedTargetStateFrame(Chaos::FGeometryParticleHan
 
 void FRewindData::RequestResimulation(int32 Frame, Chaos::FGeometryParticleHandle* Particle)
 {
-	if (IsResim() || LatestFrame < ResimCooldownUntilFrame)
+	if (IsResim())
 	{
 		return;
 	}
diff --git a/Engine/Source/Runtime/Experimental/Chaos/Public/RewindData.h b/Engine/Source/Runtime/Experimental/Chaos/Public/RewindData.h
--- a/Engine/Source/Runtime/Experimental/Chaos/Public/RewindData.h
+++ b/Engine/Source/Runtime/Experimental/Chaos/Public/RewindData.h
@@ -2215,8 +2215,11 @@ private:
 		return TObjState(State, Handle, PropertiesPool, { Frame, Phase });
 	}
 
-	/** Returns whether rewinding to RewindFrame is allowed and possible */
-	bool CanRewindToFrame(int32 RewindFrame) const;
+	/** Returns whether rewinding to RewindFrame is possible (if we have history for the frame etc.) */
+	bool CanRewindToFrame(int32 RewindFrame);
+
+	/** Returns whether rewinding to RewindFrame is allowed (if we are currently under cooldown etc.) */
+	bool AllowRewindToFrame(int32 RewindFrame) const;
 
 	/** Apply the cached history state for the given frame cached particles and joints */
 	void RewindToFrame(int32 RewindFrame);
@@ -2262,6 +2265,10 @@ private:
 	bool AllowResizingFromTickTrate = false;
 	Chaos::FReal InitializedDeltaTime = 0.0f;
 
+	// Used to measure how long a resim takes in realtime and then wait that much before allowing to trigger a new resim
+	uint64 ResimStartCycleCounter = 0;
+	uint64 ResimPerformanceCooldownCycleCounter = 0;
+
 	// Used to block rewinding past a physics change we currently don't handle
 	int32 BlockResimFrame = INDEX_NONE;
 
