Extend the cvar 'iax.InvalidUrlChance' to take an optional arg allowing the caller to apply it to either IAS or IAD rather than both.

#rnx

- The new syntax is 'iax.InvalidUrlChance' [num] [IAD/IAS/All]
-- The first arg is a number between 0-100 and the second is the type of request to apply too. Note that the second arg can be skipped entirely.
- To support dynamic args we need to convert the cvar from FAutoConsoleVariableRef to a console command.
- Since this deature significantly expands the amount of code needed I have moved the code for the cvar to it's own class FRequestErrorSimulator to avoid bloating FHttpDispatcher.
- Note that 'FRequestErrorSimulator is likely to be expanded in the future with different error types.
- Renamed UE_ALLOW_INVALID_URL_DEBUGGING to UE_ALLOW_ERROR_SIMULATION.
- Added UE_CVAR_ERROR_LOG as an easier way to log errors when parsing console commands so we don't have to test for NO_LOGGING everywhere.
- Added some missing undefs at the end of the cpp.
- The actual test itself remains unchanged, if we decide that a request should now be invalid we append extra info to the url so it won't be found. This will cause the request to return a 404 http reponse and will exercise all of our lower level error handling paths.

#Synced-CL 51086784

diff --git a/Engine/Source/Runtime/Experimental/IoStore/OnDemand/Private/OnDemandHttpIoDispacher.cpp b/Engine/Source/Runtime/Experimental/IoStore/OnDemand/Private/OnDemandHttpIoDispacher.cpp
--- a/Engine/Source/Runtime/Experimental/IoStore/OnDemand/Private/OnDemandHttpIoDispacher.cpp
+++ b/Engine/Source/Runtime/Experimental/IoStore/OnDemand/Private/OnDemandHttpIoDispacher.cpp
@@ -41,8 +41,8 @@
 #include "Misc/PackageName.h"
 #endif
 
-// When enabled the cvar 'iax.InvalidUrlChance' will be enabled and allow us to simulate invalid urls
-#define UE_ALLOW_INVALID_URL_DEBUGGING !UE_BUILD_SHIPPING
+// When enabled it will be possible to induce errors to requests in FHttpDispatcher for testing.
+#define UE_ALLOW_ERROR_SIMULATION (1 && !NO_CVARS && !UE_BUILD_SHIPPING)
 // When enabled it is possible to poison data loaded from the IAS cache via the cvar 'ias.PoisonCache'
 #define UE_ALLOW_CACHE_POISONING !UE_BUILD_SHIPPING
 // When enabled it is possible to pause the http requests via the cvar 'iax.HttpPause'
@@ -50,6 +50,13 @@
 // When enabled it is possible to disable all http bandwidth limits via the 'iax.HttpNoRateLimits' cvar
 #define UE_ALLOW_DISABLE_HTTP_LIMITS !UE_BUILD_SHIPPING
 
+/** Wrapper around FOutputDevice so we can use LogHttpIoDispatcher category when logging is enabled */
+#if NO_LOGGING
+	#define UE_CVAR_ERROR_LOG(Msg) Ar.Log(Msg);
+#else
+	#define UE_CVAR_ERROR_LOG(Msg) Ar.Log(LogHttpIoDispatcher.GetCategoryName(), ELogVerbosity::Error, Msg);
+#endif // NO_LOGGING
+
 ////////////////////////////////////////////////////////////////////////////////
 const TCHAR* LexToString(EThreadPriority Priority)
 {
@@ -77,12 +84,12 @@ const TCHAR* LexToString(EThreadPriority Priority)
 }
 
 ////////////////////////////////////////////////////////////////////////////////
-#if UE_ALLOW_HTTP_PAUSE
+#if UE_ALLOW_HTTP_PAUSE || UE_ALLOW_ERROR_SIMULATION
 void FConsoleCommandUnregister::operator()(IConsoleCommand* ConsoleCommandPtr) const
 {
 	IConsoleManager::Get().UnregisterConsoleObject(ConsoleCommandPtr);
 };
-#endif // UE_ALLOW_HTTP_PAUSE
+#endif // UE_ALLOW_HTTP_PAUSE || UE_ALLOW_ERROR_SIMULATION
 
 ////////////////////////////////////////////////////////////////////////////////
 namespace UE::IoStore
@@ -328,15 +335,6 @@ static FAutoConsoleVariableRef CVar_IaxHttpEnableInflightCancellation(
 	TEXT("When enabled the system will attempt to cancel inflight http requests")
 );
 
-#if UE_ALLOW_INVALID_URL_DEBUGGING
-float GIaxInvalidUrlChance = 0.0;
-static FAutoConsoleVariableRef CVar_GIaxInvalidUrlChance(
-	TEXT("iax.InvalidUrlChance"),
-	GIaxInvalidUrlChance,
-	TEXT("Chance that a url for a GET request will be invalid (0-100%)")
-);
-#endif // UE_ALLOW_INVALID_URL_DEBUGGING
-
 bool GHttpCacheEnabled = true;
 static FAutoConsoleVariableRef CVar_HttpCacheEnabled(
 	TEXT("iax.HttpCacheEnabled"),
@@ -387,6 +385,114 @@ static constexpr int32 GThreadPriorities[5] =
 ////////////////////////////////////////////////////////////////////////////////
 namespace HttpIoDispatcher
 {
+
+#if UE_ALLOW_ERROR_SIMULATION
+
+/** Used to help induce errors into our system for testing purposes */
+class FRequestErrorSimulator
+{
+public:
+	FRequestErrorSimulator()
+	{
+		InvalidUrlCmdHandle = FConsoleCommandPtr(IConsoleManager::Get().RegisterConsoleCommand(
+			TEXT("iax.InvalidUrlChance"),
+			TEXT("Chance that a url for a GET request will be invalid (0-100%) and optional request type filter [IAD/IAS/All]. By default the request filter will be set to All."),
+			FConsoleCommandWithArgsAndOutputDeviceDelegate::CreateRaw(this, &FRequestErrorSimulator::OnSetInvalidUrlChance),
+			ECVF_Default));
+	}
+
+	~FRequestErrorSimulator() = default;
+
+	bool SimulateUrlError(EHttpRequestType RequestType)
+	{
+		if (InvalidUrlChance <= 0.0f)
+		{
+			return false;
+		}
+
+		if (!EnumHasAnyFlags(InvalidUrlFilter, ToFilter(RequestType)))
+		{
+			return false;
+		}
+
+		return (FMath::FRand() * 100.0f) < InvalidUrlChance;
+	}
+
+private:
+	enum class ERequestFilter : uint8
+	{
+		Installed = 1 << 0,
+		Streaming = 1 << 1,
+
+		All = Installed | Streaming
+	};
+	FRIEND_ENUM_CLASS_FLAGS(ERequestFilter);
+
+	ERequestFilter ToFilter(EHttpRequestType Type)
+	{
+		switch (Type)
+		{
+			case EHttpRequestType::Installed:
+				return ERequestFilter::Installed;
+			case EHttpRequestType::Streaming:
+				return ERequestFilter::Streaming;
+			default:
+				checkNoEntry();
+				return ERequestFilter::All;
+		}
+	}
+
+	void OnSetInvalidUrlChance(const TArray<FString>& Args, FOutputDevice& Ar)
+	{
+		if (Args.Num() < 1 || Args.Num() > 2)
+		{
+			UE_CVAR_ERROR_LOG(TEXT("Invalid args! Expected iax.InvalidUrlChance [chance (float)] [IAD/IAS/All (string)]"));
+			return;
+		}
+
+		float ParsedChance = 0.0f;
+		if (!LexTryParseString(ParsedChance, *Args[0]) || ParsedChance < 0.0f || ParsedChance > 100.0f)
+		{
+			UE_CVAR_ERROR_LOG(TEXT("Invalid arg!, first arg should be a number between 0.0 and 100.0"));
+			return;
+		}
+
+		ERequestFilter ParsedFilter = ERequestFilter::All;
+
+		if (Args.Num() == 2)
+		{
+			if (Args[1] == TEXT("IAD"))
+			{
+				ParsedFilter = ERequestFilter::Installed;
+			}
+			else if (Args[1] == TEXT("IAS"))
+			{
+				ParsedFilter = ERequestFilter::Streaming;
+			}
+			else if (Args[1] == TEXT("All"))
+			{
+				ParsedFilter = ERequestFilter::All;
+			}
+			else
+			{
+				UE_CVAR_ERROR_LOG(TEXT("Unrecognized arg!, second arg should be [IAD/IAS/All]"));
+				return;
+			}
+		}
+
+		InvalidUrlChance = ParsedChance;
+		InvalidUrlFilter = ParsedFilter;
+	}
+private:
+
+	FConsoleCommandPtr InvalidUrlCmdHandle;
+	float InvalidUrlChance = 0.0f;
+	ERequestFilter InvalidUrlFilter = ERequestFilter::All;
+};
+ENUM_CLASS_FLAGS(FRequestErrorSimulator::ERequestFilter);
+
+#endif // UE_ALLOW_ERROR_SIMULATION
+
 /**
  * Tracks the lifespan of a http request for as long as we want to keep it active. When the token is
  * destroyed we will try to cancel the request if it is still in progress.
@@ -1259,6 +1365,10 @@ private:
 #if UE_ALLOW_DROP_CACHE
 	FConsoleCommandPtr						DropCacheCommand;
 #endif // UE_ALLOW_DROP_CACHE
+
+#if UE_ALLOW_ERROR_SIMULATION
+	FRequestErrorSimulator					RequestErrorSimulator;
+#endif // UE_ALLOW_ERROR_SIMULATION
 };
 
 ////////////////////////////////////////////////////////////////////////////////
@@ -1866,13 +1976,12 @@ bool FHttpDispatcher::ProcessHttpRequests()
 				//DebugPrintRange(HttpRequest, TotalRange);
 			}
 
-#if UE_ALLOW_INVALID_URL_DEBUGGING
-			// Avoid the rand call if there is no chance
-			if (GIaxInvalidUrlChance > 0.0 && (FMath::FRand() * 100.0f) < GIaxInvalidUrlChance)
+#if UE_ALLOW_ERROR_SIMULATION
+			if (RequestErrorSimulator.SimulateUrlError(HttpRequest->GetHttpRequestType()))
 			{
 				Url << "-DebugInvalidUrl";
 			}
-#endif // UE_ALLOW_INVALID_URL_DEBUGGING
+#endif // UE_ALLOW_ERROR_SIMULATION
 
 			NumConcurrentRequests++;
 			EMultiEndpointRequestFlags Flags = EnumHasAnyFlags(HttpRequest->Options.GetFlags(), EIoHttpFlags::ResponseHeaders)
@@ -2051,9 +2160,7 @@ void FHttpDispatcher::OnPauseCommand(bool bShouldPause, const TArray<FString>& A
 {
 	if (Args.Num() > 1)
 	{
-#if !NO_LOGGING
-		Ar.Log(LogHttpIoDispatcher.GetCategoryName(), ELogVerbosity::Error, TEXT("Too many args for command, either 0 (all) or 1 (IAD/IAS) expected"));
-#endif // !NO_LOGGING
+		UE_CVAR_ERROR_LOG(TEXT("Too many args for command, either 0 (all) or 1 (IAD/IAS) expected"));
 		return;
 	}
 
@@ -2071,9 +2178,7 @@ void FHttpDispatcher::OnPauseCommand(bool bShouldPause, const TArray<FString>& A
 		}
 		else
 		{
-#if !NO_LOGGING
-			Ar.Log(LogHttpIoDispatcher.GetCategoryName(), ELogVerbosity::Error, TEXT("Invalid arg for command, expecting 'IAD' or 'IAS'"));
-#endif // !NO_LOGGING
+			UE_CVAR_ERROR_LOG(TEXT("Invalid arg for command, expecting 'IAD' or 'IAS'"));
 			return;
 		}
 	}
@@ -2149,8 +2254,13 @@ TSharedPtr<IOnDemandHttpIoDispatcher> MakeOnDemanHttpIoDispatcher(TUniquePtr<cla
 
 } // UE::IoStore
 
+#undef UE_CVAR_ERROR_LOG
+
+#undef UE_ALLOW_DISABLE_HTTP_LIMITS
+#undef UE_ALLOW_HTTP_PAUSE
 #undef UE_ALLOW_CACHE_POISONING
-#undef UE_ALLOW_INVALID_URL_DEBUGGING
+#undef UE_ALLOW_ERROR_SIMULATION
+
 #undef TRACE_INIT
 #undef TRACE_REQUEST_STARTED
 #undef TRACE_REQUEST_COMPLETED
