Remove ISMPollDebugDrawComponent from non debug drawing builds. AISMPoolActor goes from 1024 -> 632 bytes

#rnx

#Synced-CL 51046130

#Synced-CL 51289644

diff --git a/Engine/Source/Runtime/Experimental/ISMPool/Public/ISMPool/ISMPoolActor.h b/Engine/Source/Runtime/Experimental/ISMPool/Public/ISMPool/ISMPoolActor.h
--- a/Engine/Source/Runtime/Experimental/ISMPool/Public/ISMPool/ISMPoolActor.h
+++ b/Engine/Source/Runtime/Experimental/ISMPool/Public/ISMPool/ISMPoolActor.h
@@ -16,8 +16,9 @@ private:
 	UPROPERTY(Category = ISMPoolActor, VisibleAnywhere, BlueprintReadOnly, meta = (AllowPrivateAccess = "true"))
 	TObjectPtr<class UISMPoolComponent> ISMPoolComp;
 
-	UPROPERTY(Category = ISMPoolActor, VisibleAnywhere, BlueprintReadOnly, meta = (AllowPrivateAccess = "true"))
+#if UE_ENABLE_DEBUG_DRAWING
 	TObjectPtr<class UISMPoolDebugDrawComponent> ISMPoolDebugDrawComp;
+#endif
 
 public:
 	/** Returns ISMPoolComp subobject **/
