Chaos API: Propagating usage of SceneId over SceneHandle.
- Updated handles for body, body container, geometry collection, and joints
- Updated visualizers

#Synced-CL 51298445

diff --git a/Engine/Plugins/Experimental/ChaosRigidPhysicsAsync/Source/ChaosRigidPhysicsAsync/Private/ChaosRigidPhysicsAsync/JointConstraint6DOFAsync.cpp b/Engine/Plugins/Experimental/ChaosRigidPhysicsAsync/Source/ChaosRigidPhysicsAsync/Private/ChaosRigidPhysicsAsync/JointConstraint6DOFAsync.cpp
--- a/Engine/Plugins/Experimental/ChaosRigidPhysicsAsync/Source/ChaosRigidPhysicsAsync/Private/ChaosRigidPhysicsAsync/JointConstraint6DOFAsync.cpp
+++ b/Engine/Plugins/Experimental/ChaosRigidPhysicsAsync/Source/ChaosRigidPhysicsAsync/Private/ChaosRigidPhysicsAsync/JointConstraint6DOFAsync.cpp
@@ -50,7 +50,7 @@ namespace Chaos::Rigids::Async
 
 	FJointConstraintHandle FJointConstraint6DOFAsyncGT::GetHandle() const
 	{
-		return FJointConstraintHandle(Scene->GetHandle(), ConstraintId);
+		return FJointConstraintHandle(Scene->GetId(), ConstraintId);
 	}
 
 	void FJointConstraint6DOFAsyncGT::Activate()
diff --git a/Engine/Plugins/Experimental/ChaosRigidPhysicsAsync/Source/ChaosRigidPhysicsAsync/Private/ChaosRigidPhysicsAsync/RigidBodyAsync.cpp b/Engine/Plugins/Experimental/ChaosRigidPhysicsAsync/Source/ChaosRigidPhysicsAsync/Private/ChaosRigidPhysicsAsync/RigidBodyAsync.cpp
--- a/Engine/Plugins/Experimental/ChaosRigidPhysicsAsync/Source/ChaosRigidPhysicsAsync/Private/ChaosRigidPhysicsAsync/RigidBodyAsync.cpp
+++ b/Engine/Plugins/Experimental/ChaosRigidPhysicsAsync/Source/ChaosRigidPhysicsAsync/Private/ChaosRigidPhysicsAsync/RigidBodyAsync.cpp
@@ -77,7 +77,7 @@ namespace Chaos::Rigids::Async
 	{
 		using namespace UE::Physics;
 
-		return FRigidBodyHandle(GetScene()->GetHandle(), BodyId);
+		return FRigidBodyHandle(GetScene()->GetId(), BodyId);
 	}
 
 	void FRigidBodyAsyncGT::SetId(const UE::Physics::FRigidObjectId& InId)
diff --git a/Engine/Plugins/Experimental/ChaosRigidPhysicsAsync/Source/ChaosRigidPhysicsAsync/Private/ChaosRigidPhysicsAsync/RigidGeometryCollectionAsync.cpp b/Engine/Plugins/Experimental/ChaosRigidPhysicsAsync/Source/ChaosRigidPhysicsAsync/Private/ChaosRigidPhysicsAsync/RigidGeometryCollectionAsync.cpp
--- a/Engine/Plugins/Experimental/ChaosRigidPhysicsAsync/Source/ChaosRigidPhysicsAsync/Private/ChaosRigidPhysicsAsync/RigidGeometryCollectionAsync.cpp
+++ b/Engine/Plugins/Experimental/ChaosRigidPhysicsAsync/Source/ChaosRigidPhysicsAsync/Private/ChaosRigidPhysicsAsync/RigidGeometryCollectionAsync.cpp
@@ -162,7 +162,7 @@ namespace Chaos::Rigids::Async
 	{
 		using namespace UE::Physics;
 
-		return FRigidBodyContainerHandle(Scene->GetHandle(), BodyContainerId);
+		return FRigidBodyContainerHandle(Scene->GetId(), BodyContainerId);
 	}
 
 	int32 FRigidGeometryCollectionAsyncGT::GetNumBodies() const
@@ -215,7 +215,7 @@ namespace Chaos::Rigids::Async
 	{
 		using namespace UE::Physics;
 
-		return FRigidBodyContainerHandle(Scene->GetHandle(), BodyContainerId);
+		return FRigidBodyContainerHandle(Scene->GetId(), BodyContainerId);
 	}
 
 	void FRigidGeometryCollectionAsyncPT::SetId(const UE::Physics::FRigidObjectId& InBodyContainerId)
diff --git a/Engine/Source/Runtime/Experimental/RigidPhysics/Private/RigidPhysics/JointConstraint6DOFHandle.cpp b/Engine/Source/Runtime/Experimental/RigidPhysics/Private/RigidPhysics/JointConstraint6DOFHandle.cpp
--- a/Engine/Source/Runtime/Experimental/RigidPhysics/Private/RigidPhysics/JointConstraint6DOFHandle.cpp
+++ b/Engine/Source/Runtime/Experimental/RigidPhysics/Private/RigidPhysics/JointConstraint6DOFHandle.cpp
@@ -7,8 +7,8 @@
 
 namespace UE::Physics
 {
-	FJointConstraint6DOFHandle::FJointConstraint6DOFHandle(const FRigidSceneHandle& InSceneHandle, const FRigidObjectId& InConstraintId)
-		: FJointConstraintHandle(InSceneHandle, InConstraintId)
+	FJointConstraint6DOFHandle::FJointConstraint6DOFHandle(const FRigidSceneId& InSceneId, const FRigidObjectId& InConstraintId)
+		: FJointConstraintHandle(InSceneId, InConstraintId)
 	{
 	}
 
diff --git a/Engine/Source/Runtime/Experimental/RigidPhysics/Private/RigidPhysics/JointConstraintHandle.cpp b/Engine/Source/Runtime/Experimental/RigidPhysics/Private/RigidPhysics/JointConstraintHandle.cpp
--- a/Engine/Source/Runtime/Experimental/RigidPhysics/Private/RigidPhysics/JointConstraintHandle.cpp
+++ b/Engine/Source/Runtime/Experimental/RigidPhysics/Private/RigidPhysics/JointConstraintHandle.cpp
@@ -8,15 +8,15 @@
 
 namespace UE::Physics
 {
-	FJointConstraintHandle::FJointConstraintHandle(const FRigidSceneHandle& InSceneHandle, const FRigidObjectId& InConstraintId)
-		: SceneHandle(InSceneHandle)
+	FJointConstraintHandle::FJointConstraintHandle(const FRigidSceneId& InSceneId, const FRigidObjectId& InConstraintId)
+		: SceneId(InSceneId)
 		, ConstraintId(InConstraintId)
 	{
 	}
 
 	FRigidSceneHandle FJointConstraintHandle::GetSceneHandle() const
 	{
-		return SceneHandle;
+		return FRigidSceneHandle(SceneId);
 	}
 
 	const FRigidObjectId& FJointConstraintHandle::GetId() const
@@ -26,7 +26,7 @@ namespace UE::Physics
 
 	void FJointConstraintHandle::Reset()
 	{
-		SceneHandle.Reset();
+		SceneId.Reset();
 		ConstraintId.Reset();
 	}
 
diff --git a/Engine/Source/Runtime/Experimental/RigidPhysics/Private/RigidPhysics/RigidBodyContainerHandle.cpp b/Engine/Source/Runtime/Experimental/RigidPhysics/Private/RigidPhysics/RigidBodyContainerHandle.cpp
--- a/Engine/Source/Runtime/Experimental/RigidPhysics/Private/RigidPhysics/RigidBodyContainerHandle.cpp
+++ b/Engine/Source/Runtime/Experimental/RigidPhysics/Private/RigidPhysics/RigidBodyContainerHandle.cpp
@@ -7,16 +7,16 @@
 
 namespace UE::Physics
 {
-	FRigidBodyContainerHandle::FRigidBodyContainerHandle(const FRigidSceneHandle& InSceneHandle, const FRigidObjectId& InBodyContainerId)
+	FRigidBodyContainerHandle::FRigidBodyContainerHandle(const FRigidSceneId& InSceneId, const FRigidObjectId& InBodyContainerId)
 		: BodyContainerId(InBodyContainerId)
-		, SceneHandle(InSceneHandle)
+		, SceneId(InSceneId)
 	{
 	}
 
 	void FRigidBodyContainerHandle::Reset()
 	{
 		BodyContainerId.Reset();
-		SceneHandle.Reset();
+		SceneId.Reset();
 	}
 
 	const FRigidObjectId& FRigidBodyContainerHandle::GetId() const
@@ -24,9 +24,14 @@ namespace UE::Physics
 		return BodyContainerId;
 	}
 
+	const FRigidSceneId FRigidBodyContainerHandle::GetSceneId() const
+	{
+		return SceneId;
+	}
+
 	const FRigidSceneHandle FRigidBodyContainerHandle::GetSceneHandle() const
 	{
-		return SceneHandle;
+		return FRigidSceneHandle(SceneId);
 	}
 
 	IRigidBodyContainer* FRigidBodyContainerHandle::PinInternal(const IRigidScene* InScene) const
diff --git a/Engine/Source/Runtime/Experimental/RigidPhysics/Private/RigidPhysics/RigidBodyHandle.cpp b/Engine/Source/Runtime/Experimental/RigidPhysics/Private/RigidPhysics/RigidBodyHandle.cpp
--- a/Engine/Source/Runtime/Experimental/RigidPhysics/Private/RigidPhysics/RigidBodyHandle.cpp
+++ b/Engine/Source/Runtime/Experimental/RigidPhysics/Private/RigidPhysics/RigidBodyHandle.cpp
@@ -7,9 +7,9 @@
 
 namespace UE::Physics
 {
-	FRigidBodyHandle::FRigidBodyHandle(const FRigidSceneHandle& InSceneHandle, const FRigidObjectId& InBodyId)
+	FRigidBodyHandle::FRigidBodyHandle(const FRigidSceneId& InSceneId, const FRigidObjectId& InBodyId)
 		: BodyId(InBodyId)
-		, SceneHandle(InSceneHandle)
+		, SceneId(InSceneId)
 	{
 	}
 
@@ -20,7 +20,7 @@ namespace UE::Physics
 
 	void FRigidBodyHandle::Reset()
 	{
-		SceneHandle.Reset();
+		SceneId.Reset();
 		BodyId.Reset();
 	}
 
diff --git a/Engine/Source/Runtime/Experimental/RigidPhysics/Private/RigidPhysics/RigidGeometryCollectionHandle.cpp b/Engine/Source/Runtime/Experimental/RigidPhysics/Private/RigidPhysics/RigidGeometryCollectionHandle.cpp
--- a/Engine/Source/Runtime/Experimental/RigidPhysics/Private/RigidPhysics/RigidGeometryCollectionHandle.cpp
+++ b/Engine/Source/Runtime/Experimental/RigidPhysics/Private/RigidPhysics/RigidGeometryCollectionHandle.cpp
@@ -8,8 +8,8 @@
 
 namespace UE::Physics
 {
-	FRigidGeometryCollectionHandle::FRigidGeometryCollectionHandle(const FRigidSceneHandle& InSceneHandle, const FRigidObjectId& InBodyContainerId)
-		: BodyContainerHandle(InSceneHandle, InBodyContainerId)
+	FRigidGeometryCollectionHandle::FRigidGeometryCollectionHandle(const FRigidSceneId& InSceneId, const FRigidObjectId& InBodyContainerId)
+		: BodyContainerHandle(InSceneId, InBodyContainerId)
 	{
 	}
 
diff --git a/Engine/Source/Runtime/Experimental/RigidPhysics/Private/RigidPhysics/RigidSceneHandle.cpp b/Engine/Source/Runtime/Experimental/RigidPhysics/Private/RigidPhysics/RigidSceneHandle.cpp
--- a/Engine/Source/Runtime/Experimental/RigidPhysics/Private/RigidPhysics/RigidSceneHandle.cpp
+++ b/Engine/Source/Runtime/Experimental/RigidPhysics/Private/RigidPhysics/RigidSceneHandle.cpp
@@ -10,13 +10,28 @@
 namespace UE::Physics
 {
 	FRigidSceneHandle::FRigidSceneHandle(FRigidScene* InScene)
-		: SceneRaw(InScene)
+		: SceneId(InScene->GetId())
+		, SceneRaw(InScene)
 #if UE_RIGIDPHYSICS_SCENEHANDLE_SAFTEY_ENABLED
 		, SceneWeak(InScene)
 #endif
 	{
 	}
 
+	FRigidSceneHandle::FRigidSceneHandle(const FRigidSceneId& InSceneId)
+		: SceneId(InSceneId)
+	{
+		IRigidScene* BaseScene = InSceneId.Get();
+		UE_RIGIDPHYSICS_CHECK(BaseScene);
+		if (BaseScene)
+		{
+			SceneRaw = BaseScene->AsAChecked<FRigidScene>();
+#if UE_RIGIDPHYSICS_SCENEHANDLE_SAFTEY_ENABLED
+			SceneWeak = SceneRaw;
+#endif // UE_RIGIDPHYSICS_SCENEHANDLE_SAFTEY_ENABLED
+		}
+	}
+
 	FRigidContextGameRO FRigidSceneHandle::LockRO() const
 	{
 		DoSafetyCheck();
@@ -31,6 +46,7 @@ namespace UE::Physics
 
 	void FRigidSceneHandle::Reset()
 	{
+		SceneId.Reset();
 		SceneRaw = nullptr;
 #if UE_RIGIDPHYSICS_SCENEHANDLE_SAFTEY_ENABLED
 		SceneWeak.Reset();
@@ -39,12 +55,12 @@ namespace UE::Physics
 
 	FRigidSceneId FRigidSceneHandle::GetId() const
 	{
-		return GetScene()->GetId();
+		return SceneId;
 	}
 
 	IRigidScene* FRigidSceneHandle::GetScene() const
 	{
-		return SceneRaw;
+		return SceneId.Get();
 	}
 
 	void FRigidSceneHandle::DoSafetyCheck() const
diff --git a/Engine/Source/Runtime/Experimental/RigidPhysics/Public/RigidPhysics/JointConstraint6DOFHandle.h b/Engine/Source/Runtime/Experimental/RigidPhysics/Public/RigidPhysics/JointConstraint6DOFHandle.h
--- a/Engine/Source/Runtime/Experimental/RigidPhysics/Public/RigidPhysics/JointConstraint6DOFHandle.h
+++ b/Engine/Source/Runtime/Experimental/RigidPhysics/Public/RigidPhysics/JointConstraint6DOFHandle.h
@@ -19,7 +19,7 @@ namespace UE::Physics
 	public:
 		FJointConstraint6DOFHandle() = default;
 		FJointConstraint6DOFHandle(const FJointConstraint6DOFHandle&) = default;
-		UE_INTERNAL RIGIDPHYSICS_API FJointConstraint6DOFHandle(const FRigidSceneHandle& InSceneHandle, const FRigidObjectId& InConstraintId);
+		UE_INTERNAL RIGIDPHYSICS_API FJointConstraint6DOFHandle(const FRigidSceneId& InSceneId, const FRigidObjectId& InConstraintId);
 
 		friend bool operator==(const FJointConstraint6DOFHandle&, const FJointConstraint6DOFHandle&) = default;
 		friend bool operator!=(const FJointConstraint6DOFHandle&, const FJointConstraint6DOFHandle&) = default;
diff --git a/Engine/Source/Runtime/Experimental/RigidPhysics/Public/RigidPhysics/JointConstraintHandle.h b/Engine/Source/Runtime/Experimental/RigidPhysics/Public/RigidPhysics/JointConstraintHandle.h
--- a/Engine/Source/Runtime/Experimental/RigidPhysics/Public/RigidPhysics/JointConstraintHandle.h
+++ b/Engine/Source/Runtime/Experimental/RigidPhysics/Public/RigidPhysics/JointConstraintHandle.h
@@ -21,7 +21,7 @@ namespace UE::Physics
 	public:
 		FJointConstraintHandle() = default;
 		FJointConstraintHandle(const FJointConstraintHandle&) = default;
-		UE_INTERNAL RIGIDPHYSICS_API FJointConstraintHandle(const FRigidSceneHandle& InSceneHandle, const FRigidObjectId& InConstraintId);
+		UE_INTERNAL RIGIDPHYSICS_API FJointConstraintHandle(const FRigidSceneId& InSceneId, const FRigidObjectId& InConstraintId);
 
 		friend bool operator==(const FJointConstraintHandle&, const FJointConstraintHandle&) = default;
 		friend bool operator!=(const FJointConstraintHandle&, const FJointConstraintHandle&) = default;
@@ -31,9 +31,8 @@ namespace UE::Physics
 		template <typename ContextType>
 		TJointConstraintPtr<ContextType> Pin(const ContextType& InContext) const
 		{
-			if (InContext.IsValid())
+			if (const IRigidScene* Scene = InContext.GetRawSceneChecked(SceneId))
 			{
-				const IRigidScene* Scene = InContext.GetScene().Get();
 				return TJointConstraintPtr<ContextType>(PinInternal(Scene));
 			}
 			return TJointConstraintPtr<ContextType>();
@@ -45,7 +44,7 @@ namespace UE::Physics
 	private:
 		RIGIDPHYSICS_API IJointConstraint* PinInternal(const IRigidScene* Scene) const;
 
-		FRigidSceneHandle SceneHandle;
+		FRigidSceneId SceneId;
 		FRigidObjectId ConstraintId;
 	};
 } // namespace UE::Physics
diff --git a/Engine/Source/Runtime/Experimental/RigidPhysics/Public/RigidPhysics/RigidBodyContainerHandle.h b/Engine/Source/Runtime/Experimental/RigidPhysics/Public/RigidPhysics/RigidBodyContainerHandle.h
--- a/Engine/Source/Runtime/Experimental/RigidPhysics/Public/RigidPhysics/RigidBodyContainerHandle.h
+++ b/Engine/Source/Runtime/Experimental/RigidPhysics/Public/RigidPhysics/RigidBodyContainerHandle.h
@@ -18,7 +18,7 @@ namespace UE::Physics
 	{
 	public:
 		FRigidBodyContainerHandle() = default;
-		UE_INTERNAL RIGIDPHYSICS_API FRigidBodyContainerHandle(const FRigidSceneHandle& InSceneHandle, const FRigidObjectId& InBodyContainerId);
+		UE_INTERNAL RIGIDPHYSICS_API FRigidBodyContainerHandle(const FRigidSceneId& InSceneId, const FRigidObjectId& InBodyContainerId);
 
 		friend bool operator==(const FRigidBodyContainerHandle&, const FRigidBodyContainerHandle&) = default;
 		friend bool operator!=(const FRigidBodyContainerHandle&, const FRigidBodyContainerHandle&) = default;
@@ -27,11 +27,8 @@ namespace UE::Physics
 		template <typename ContextType>
 		TRigidBodyContainerPtr<ContextType> Pin(const ContextType& InContext) const
 		{
-			// NOTE: we don't use our copy of SceneHandle here (or anywhere!)
-			// TODO @ JoshD: This seems wrong. We need some safety checking here...
-			if (InContext.IsValid())
+			if (const IRigidScene* Scene = InContext.GetRawSceneChecked(SceneId))
 			{
-				const IRigidScene* Scene = InContext.GetScene().Get();
 				return PinInternal(Scene);
 			}
 			return nullptr;
@@ -39,17 +36,14 @@ namespace UE::Physics
 
 		UE_INTERNAL RIGIDPHYSICS_API void Reset();
 		UE_INTERNAL RIGIDPHYSICS_API const FRigidObjectId& GetId() const;
-
-		// NOTE: We don't really need the scene handle because it is provided by the Context
-		// in Pin(). It is only needed for the debugger visualizers (see RigidPhysics.natvis).
-		// TODO_CHAOSAPI: Consider making SceneHandle shipping only
+		UE_INTERNAL RIGIDPHYSICS_API const FRigidSceneId GetSceneId() const;
 		UE_INTERNAL RIGIDPHYSICS_API const FRigidSceneHandle GetSceneHandle() const;
 
 	private:
 		RIGIDPHYSICS_API IRigidBodyContainer* PinInternal(const IRigidScene* InScene) const;
 
 		FRigidObjectId BodyContainerId;
-		FRigidSceneHandle SceneHandle;
+		FRigidSceneId SceneId;
 	};
 } // namespace UE::Physics
 
diff --git a/Engine/Source/Runtime/Experimental/RigidPhysics/Public/RigidPhysics/RigidBodyHandle.h b/Engine/Source/Runtime/Experimental/RigidPhysics/Public/RigidPhysics/RigidBodyHandle.h
--- a/Engine/Source/Runtime/Experimental/RigidPhysics/Public/RigidPhysics/RigidBodyHandle.h
+++ b/Engine/Source/Runtime/Experimental/RigidPhysics/Public/RigidPhysics/RigidBodyHandle.h
@@ -23,7 +23,7 @@ namespace UE::Physics
 		FRigidBodyHandle() = default;
 		FRigidBodyHandle(const FRigidBodyHandle&) = default;
 
-		UE_INTERNAL RIGIDPHYSICS_API FRigidBodyHandle(const FRigidSceneHandle& InSceneHandle, const FRigidObjectId& InBodyId);
+		UE_INTERNAL RIGIDPHYSICS_API FRigidBodyHandle(const FRigidSceneId& InSceneId, const FRigidObjectId& InBodyId);
 
 		friend bool operator==(const FRigidBodyHandle&, const FRigidBodyHandle&) = default;
 		friend bool operator!=(const FRigidBodyHandle&, const FRigidBodyHandle&) = default;
@@ -32,9 +32,8 @@ namespace UE::Physics
 		template <typename ContextType>
 		TRigidBodyPtr<ContextType> Pin(const ContextType& InContext) const
 		{
-			if (InContext.IsValid())
+			if (const IRigidScene* Scene = InContext.GetRawSceneChecked(SceneId))
 			{
-				const IRigidScene* Scene = InContext.GetScene().Get();
 				return PinInternal(Scene);
 			}
 			return TRigidBodyPtr<ContextType>();
@@ -47,7 +46,7 @@ namespace UE::Physics
 		RIGIDPHYSICS_API IRigidBody* PinInternal(const IRigidScene* Scene) const;
 
 		FRigidObjectId BodyId;
-		FRigidSceneHandle SceneHandle;
+		FRigidSceneId SceneId;
 	};
 } // namespace UE::Physics
 
diff --git a/Engine/Source/Runtime/Experimental/RigidPhysics/Public/RigidPhysics/RigidContext.h b/Engine/Source/Runtime/Experimental/RigidPhysics/Public/RigidPhysics/RigidContext.h
--- a/Engine/Source/Runtime/Experimental/RigidPhysics/Public/RigidPhysics/RigidContext.h
+++ b/Engine/Source/Runtime/Experimental/RigidPhysics/Public/RigidPhysics/RigidContext.h
@@ -27,6 +27,7 @@ namespace UE::Physics
 		{
 			if (InScene != nullptr)
 			{
+				SceneId = InScene->GetHandle().GetId();
 				InScene->Lock(LockType);
 				ScenePtr.Init(InScene);
 			}
@@ -66,9 +67,15 @@ namespace UE::Physics
 			return ScenePtr;
 		}
 
+		UE_INTERNAL const IRigidScene* GetRawSceneChecked(const FRigidSceneId InSceneId) const
+		{
+			return (IsValid() && InSceneId == SceneId) ? ScenePtr->Get() : nullptr;
+		}
+
 	private:
 		TRigidScenePtr<FContext> ScenePtr;
+		FRigidSceneId SceneId;
 	};
 }
 
-#endif // UE_RIGIDPHYSICS_API_ENABLED
\ No newline at end of file
+#endif // UE_RIGIDPHYSICS_API_ENABLED
diff --git a/Engine/Source/Runtime/Experimental/RigidPhysics/Public/RigidPhysics/RigidGeometryCollection.h b/Engine/Source/Runtime/Experimental/RigidPhysics/Public/RigidPhysics/RigidGeometryCollection.h
--- a/Engine/Source/Runtime/Experimental/RigidPhysics/Public/RigidPhysics/RigidGeometryCollection.h
+++ b/Engine/Source/Runtime/Experimental/RigidPhysics/Public/RigidPhysics/RigidGeometryCollection.h
@@ -66,7 +66,7 @@ namespace UE::Physics
 			if (GeometryCollectionRaw != nullptr)
 			{
 				FRigidBodyContainerHandle BodyContainerHandle = GeometryCollectionRaw->GetHandle();
-				return FRigidGeometryCollectionHandle(BodyContainerHandle.GetSceneHandle(), BodyContainerHandle.GetId());
+				return FRigidGeometryCollectionHandle(BodyContainerHandle.GetSceneId(), BodyContainerHandle.GetId());
 			}
 			return FRigidGeometryCollectionHandle();
 		}
diff --git a/Engine/Source/Runtime/Experimental/RigidPhysics/Public/RigidPhysics/RigidGeometryCollectionHandle.h b/Engine/Source/Runtime/Experimental/RigidPhysics/Public/RigidPhysics/RigidGeometryCollectionHandle.h
--- a/Engine/Source/Runtime/Experimental/RigidPhysics/Public/RigidPhysics/RigidGeometryCollectionHandle.h
+++ b/Engine/Source/Runtime/Experimental/RigidPhysics/Public/RigidPhysics/RigidGeometryCollectionHandle.h
@@ -35,7 +35,7 @@ namespace UE::Physics
 	{
 	public:
 		FRigidGeometryCollectionHandle() = default;
-		UE_INTERNAL RIGIDPHYSICS_API FRigidGeometryCollectionHandle(const FRigidSceneHandle& InSceneHandle, const FRigidObjectId& InBodyContainerId);
+		UE_INTERNAL RIGIDPHYSICS_API FRigidGeometryCollectionHandle(const FRigidSceneId& InSceneId, const FRigidObjectId& InBodyContainerId);
 
 		friend bool operator==(const FRigidBodyContainerHandle&, const FRigidBodyContainerHandle&) = default;
 		friend bool operator!=(const FRigidBodyContainerHandle&, const FRigidBodyContainerHandle&) = default;
diff --git a/Engine/Source/Runtime/Experimental/RigidPhysics/Public/RigidPhysics/RigidSceneHandle.h b/Engine/Source/Runtime/Experimental/RigidPhysics/Public/RigidPhysics/RigidSceneHandle.h
--- a/Engine/Source/Runtime/Experimental/RigidPhysics/Public/RigidPhysics/RigidSceneHandle.h
+++ b/Engine/Source/Runtime/Experimental/RigidPhysics/Public/RigidPhysics/RigidSceneHandle.h
@@ -19,6 +19,7 @@ namespace UE::Physics
 	public:
 		FRigidSceneHandle() = default;
 		UE_INTERNAL RIGIDPHYSICS_API FRigidSceneHandle(FRigidScene* InScene);
+		UE_INTERNAL RIGIDPHYSICS_API FRigidSceneHandle(const FRigidSceneId& InSceneId);
 
 		friend bool operator==(const FRigidSceneHandle&, const FRigidSceneHandle&) = default;
 		friend bool operator!=(const FRigidSceneHandle&, const FRigidSceneHandle&) = default;
@@ -34,8 +35,9 @@ namespace UE::Physics
 	private:
 		void DoSafetyCheck() const;
 
+		FRigidSceneId SceneId;
 		// NOTE: Always the GT Scene, but never accessed on PT
-		FRigidScene* SceneRaw = nullptr;
+		IRigidScene* SceneRaw = nullptr;
 #if UE_RIGIDPHYSICS_SCENEHANDLE_SAFTEY_ENABLED
 		FRigidSceneWeakPtr SceneWeak;
 #endif
