Chaos API: Fixing an issue where rigid object pointers could be casted to unrelated context types. This allowed casting a RO context to a RW context which would invalidate proper locking.
Also fixed an occurance of this in the ChaosDataFlowPlugin.

#Synced-CL 51298445

diff --git a/Engine/Plugins/Experimental/ChaosDataflowSolver/Source/ChaosDataflowSolver/Private/RigidComponents/PrimitiveRigidComponentInterface.h b/Engine/Plugins/Experimental/ChaosDataflowSolver/Source/ChaosDataflowSolver/Private/RigidComponents/PrimitiveRigidComponentInterface.h
--- a/Engine/Plugins/Experimental/ChaosDataflowSolver/Source/ChaosDataflowSolver/Private/RigidComponents/PrimitiveRigidComponentInterface.h
+++ b/Engine/Plugins/Experimental/ChaosDataflowSolver/Source/ChaosDataflowSolver/Private/RigidComponents/PrimitiveRigidComponentInterface.h
@@ -106,7 +106,7 @@ class PrimitiveRigidComponentInterface
 
 		if (UE::Physics::FRigidContextGameRO Context = SceneHandle.LockRO())
 		{
-			if (UE::Physics::TRigidBodyPtr<UE::Physics::FRigidContextGameRW> BodyPtr = BodyHandle.Pin(Context))
+			if (UE::Physics::TRigidBodyPtr<UE::Physics::FRigidContextGameRO> BodyPtr = BodyHandle.Pin(Context))
 			{
 				if (!BodyPtr->IsStatic())
 				{
diff --git a/Engine/Source/Runtime/Experimental/RigidPhysics/Public/RigidPhysics/RigidContext.h b/Engine/Source/Runtime/Experimental/RigidPhysics/Public/RigidPhysics/RigidContext.h
--- a/Engine/Source/Runtime/Experimental/RigidPhysics/Public/RigidPhysics/RigidContext.h
+++ b/Engine/Source/Runtime/Experimental/RigidPhysics/Public/RigidPhysics/RigidContext.h
@@ -17,9 +17,9 @@ namespace UE::Physics
 	class TRigidContext
 	{
 	public:
-		static const ERigidContextType ContextType = ContextTypeParam;
-		static const ERigidLockType LockType = LockTypeParam;
-		static const bool bWriteEnabled = (LockType == ERigidLockType::ReadWrite);
+		static constexpr ERigidContextType ContextType = ContextTypeParam;
+		static constexpr ERigidLockType LockType = LockTypeParam;
+		static constexpr bool bWriteEnabled = (LockType == ERigidLockType::ReadWrite);
 
 		using FContext = TRigidContext<ContextType, LockType>;
 
diff --git a/Engine/Source/Runtime/Experimental/RigidPhysics/Public/RigidPhysics/RigidObjectPtr.h b/Engine/Source/Runtime/Experimental/RigidPhysics/Public/RigidPhysics/RigidObjectPtr.h
--- a/Engine/Source/Runtime/Experimental/RigidPhysics/Public/RigidPhysics/RigidObjectPtr.h
+++ b/Engine/Source/Runtime/Experimental/RigidPhysics/Public/RigidPhysics/RigidObjectPtr.h
@@ -48,11 +48,12 @@ namespace UE::Physics
 
 		// Downcasting. Will fail to compile for attempts to cast to non-derived types.
 		// Will generate "invalid" pointer in the runtime when InPtr is not an instance of right type.
+		// Note: Also make sure we don't allow conversion to different context types.
 		template<
 			typename OtherPointerImplType
-			UE_REQUIRES(std::is_base_of_v<typename OtherPointerImplType::FInterface, FInterface>)
+			UE_REQUIRES(std::is_base_of_v<typename OtherPointerImplType::FInterface, FInterface>&& std::is_same_v<typename OtherPointerImplType::FContext, FContext>)
 			>
-		UE_INTERNAL TRigidObjectPtr(const TRigidObjectPtr<OtherPointerImplType>& InPtr) 
+		UE_INTERNAL TRigidObjectPtr(const TRigidObjectPtr<OtherPointerImplType>& InPtr)
 			: Impl(InPtr.Get())
 		{
 		}
@@ -94,14 +95,14 @@ namespace UE::Physics
 			return ToHandle();
 		}
 
-		// Const pointer sematics
+		// Const pointer semantics
 		const FPointerImpl* operator->() const
 		{
 			CheckIsValid();
 			return &Impl;
 		}
 
-		// Non-const pointer sematics
+		// Non-const pointer semantics
 		FPointerImpl* operator->() const requires (FContext::bWriteEnabled)
 		{
 			CheckIsValid();
