[NetPhysics] Add new setting (percent of resim cache size) for number of frames to look back when finding valid resim frame.

#Synced-CL 51339479

diff --git a/Engine/Source/Runtime/Experimental/Chaos/Private/RewindData.cpp b/Engine/Source/Runtime/Experimental/Chaos/Private/RewindData.cpp
--- a/Engine/Source/Runtime/Experimental/Chaos/Private/RewindData.cpp
+++ b/Engine/Source/Runtime/Experimental/Chaos/Private/RewindData.cpp
@@ -2117,6 +2117,8 @@ CHAOS_API int32 ResimFrameValidation = static_cast<int32>(EResimFrameValidation:
 FAutoConsoleVariableRef CVarResimFrameValidationLeniency(TEXT("p.Resim.ResimFrameValidation"), ResimFrameValidation, TEXT("0 = no leniency, all dirty particles need a valid target. 1 = Island leniency, all particles in resim islands need a valid target. 2 = Full leniency, only the particle triggering the resim need a valid target."));
 CHAOS_API int32 FindValidFrameNFramesBehind = 2;
 FAutoConsoleVariableRef CVarResimFindValidFrameNFramesBehind(TEXT("p.Resim.FindValidFrameNFramesBehind"), FindValidFrameNFramesBehind, TEXT("Number of frames to allow searching for a valid frame behind the requested frame."));
+CHAOS_API float FindValidFrameNPercentBehind = 0.1f;
+FAutoConsoleVariableRef CVarResimFindValidFrameNPercentBehind(TEXT("p.Resim.FindValidFrameNPercentBehind"), FindValidFrameNPercentBehind, TEXT("Value between 0.0 -> 1.0 as percent of rewind history size to allow searching for a valid frame behind the requested frame. Example: If rewind history is 60 frames long and we allow looking 10% frames behind, we will allow to look 6 frames back from the requested frame."));
 CHAOS_API bool bResimIncompleteHistory = true;
 FAutoConsoleVariableRef CVarResimIncompleteHistory(TEXT("p.Resim.IncompleteHistory"), bResimIncompleteHistory, TEXT("If a valid resim frame can't be found, use the requested resim frame and perform a resimulation with incomplete data."));
 CHAOS_API bool bFindValidInputHistory = true;
@@ -2226,11 +2228,12 @@ int32 FRewindData::FindValidResimFrame(const int32 RequestedFrame)
 		IslandManager.FindParticlesInIslands(IslandValidationIslands, OUT IslandValidationIslandParticles);
 	}
 
-	/** Get earliest frame we are allowed to rewind taking these three into account:
+	/** Get earliest frame we are allowed to rewind to taking these three into account:
 	* - Earliest frame in history
-	* - Blocked resim frame 
-	* - Setting for how many frames behind the requested frame we allow to search (NOTE: -1 always added because the for-loop is excluding the EarliestFrame) */ 
-	const int32 EarliestFrame = FMath::Max3(GetEarliestFrame_Internal(), BlockResimFrame, RequestedFrame - 1 - FMath::Max(FindValidFrameNFramesBehind, 0));
+	* - Blocked resim frame
+	* - Setting for how many frames behind the requested frame we allow to search (NOTE: -1 because the for-loop is excluding the EarliestFrame) */
+	const int32 FramesBehindSettings = FMath::Max(FMath::Max(FindValidFrameNFramesBehind, 0), FMath::CeilToInt32(static_cast<float>(Capacity()) * FMath::Clamp(FindValidFrameNPercentBehind, 0.0f, 1.0f)));
+	const int32 EarliestFrame = FMath::Max3(GetEarliestFrame_Internal(), BlockResimFrame, RequestedFrame - 1 - FramesBehindSettings);
 	bool bHasTargetHistory = false;
 	bool bHasCustomDataHistory = false;
 
