Chaos API: Fixing Check Exports failing due to API tags on inline and pure virtual methods

#Synced-CL 51298445

diff --git a/Engine/Source/Runtime/Experimental/RigidPhysics/Private/RigidPhysics/Internal/IJointConstraint6DOF.cpp b/Engine/Source/Runtime/Experimental/RigidPhysics/Private/RigidPhysics/Internal/IJointConstraint6DOF.cpp
--- a/Engine/Source/Runtime/Experimental/RigidPhysics/Private/RigidPhysics/Internal/IJointConstraint6DOF.cpp
+++ b/Engine/Source/Runtime/Experimental/RigidPhysics/Private/RigidPhysics/Internal/IJointConstraint6DOF.cpp
@@ -13,6 +13,65 @@ namespace UE::Physics
 	{
 		return FJointConstraint6DOFHandle(GetHandle());
 	}
+
+	void IJointConstraint6DOF::Activate()
+	{
+		RIGIDPHYSICS_API_UNSUPPORTED();
+	}
+
+	void IJointConstraint6DOF::SetBodies(IRigidBody* Body0, IRigidBody* Body1)
+	{
+		RIGIDPHYSICS_API_UNSUPPORTED();
+	}
+
+	void IJointConstraint6DOF::SetJointTransforms(const FTransform3d& InTransform0, const FTransform3d& InTransform1)
+	{
+		RIGIDPHYSICS_API_UNSUPPORTED();
+	}
+
+	double IJointConstraint6DOF::GetLinearLimit() const
+	{
+		RIGIDPHYSICS_API_UNSUPPORTED();
+		return 0;
+	}
+
+	void IJointConstraint6DOF::SetLinearLimit(const double InLinearLimit)
+	{
+		RIGIDPHYSICS_API_UNSUPPORTED();
+	}
+
+	EJointMotionType IJointConstraint6DOF::GetLinearMotionTypesX() const
+	{
+		RIGIDPHYSICS_API_UNSUPPORTED();
+		return EJointMotionType::Free;
+	}
+
+	void IJointConstraint6DOF::SetLinearMotionTypesX(EJointMotionType InMotionType)
+	{
+		RIGIDPHYSICS_API_UNSUPPORTED();
+	}
+
+	EJointMotionType IJointConstraint6DOF::GetLinearMotionTypesY() const
+	{
+		RIGIDPHYSICS_API_UNSUPPORTED();
+		return EJointMotionType::Free;
+	}
+
+	void IJointConstraint6DOF::SetLinearMotionTypesY(EJointMotionType InMotionType)
+	{
+		RIGIDPHYSICS_API_UNSUPPORTED();
+	}
+
+	EJointMotionType IJointConstraint6DOF::GetLinearMotionTypesZ() const
+	{
+		RIGIDPHYSICS_API_UNSUPPORTED();
+		return EJointMotionType::Free;
+	}
+
+	void IJointConstraint6DOF::SetLinearMotionTypesZ(EJointMotionType InMotionType)
+	{
+		RIGIDPHYSICS_API_UNSUPPORTED();
+	}
 } // namespace UE::Physics
 
 #endif // UE_RIGIDPHYSICS_API_ENABLED
diff --git a/Engine/Source/Runtime/Experimental/RigidPhysics/Private/RigidPhysics/Internal/IRigidBody.cpp b/Engine/Source/Runtime/Experimental/RigidPhysics/Private/RigidPhysics/Internal/IRigidBody.cpp
--- a/Engine/Source/Runtime/Experimental/RigidPhysics/Private/RigidPhysics/Internal/IRigidBody.cpp
+++ b/Engine/Source/Runtime/Experimental/RigidPhysics/Private/RigidPhysics/Internal/IRigidBody.cpp
@@ -8,6 +8,162 @@
 namespace UE::Physics
 {
 	UE_RIGIDPHYSICS_RIGIDTYPED_IMPL(IRigidBody, IRigidTyped);
-}
+
+	bool IRigidBody::IsActive() const
+	{
+		RIGIDPHYSICS_API_UNSUPPORTED();
+		return false;
+	}
+
+	void IRigidBody::Activate()
+	{
+		RIGIDPHYSICS_API_UNSUPPORTED();
+	}
+
+	void IRigidBody::Deactivate()
+	{
+		RIGIDPHYSICS_API_UNSUPPORTED();
+	}
+
+	void IRigidBody::SetId(const FRigidObjectId& InId)
+	{
+		RIGIDPHYSICS_API_UNSUPPORTED();
+	}
+
+	void IRigidBody::SetIsDynamic(bool bIsDynamic) const
+	{
+		RIGIDPHYSICS_API_UNSUPPORTED();
+	}
+
+	bool IRigidBody::IsSleeping() const
+	{
+		RIGIDPHYSICS_API_UNSUPPORTED();
+		return false;
+	}
+
+	void IRigidBody::AddShape(const FConstRigidShapeInstancePtr& InShape)
+	{
+		RIGIDPHYSICS_API_UNSUPPORTED();
+	}
+
+	void IRigidBody::InitTransform(const FTransform3d& InTransform)
+	{
+		RIGIDPHYSICS_API_UNSUPPORTED();
+	}
+
+	void IRigidBody::UpdateTransform(const FTransform3d& InTransform)
+	{
+		RIGIDPHYSICS_API_UNSUPPORTED();
+	}
+
+	void IRigidBody::SetKinematicTarget(const FTransform3d& InTransform)
+	{
+		RIGIDPHYSICS_API_UNSUPPORTED();
+	}
+
+	FTransform3d IRigidBody::GetTransform() const
+	{
+		RIGIDPHYSICS_API_UNSUPPORTED();
+		return FTransform3d::Identity;
+	}
+
+	FBounds3d IRigidBody::GetBounds() const
+	{
+		RIGIDPHYSICS_API_UNSUPPORTED();
+		return FBounds3d(FVector3d::Zero(), FVector3d::Zero());
+	}
+
+	double IRigidBody::GetMass() const
+	{
+		RIGIDPHYSICS_API_UNSUPPORTED();
+		return 0;
+	}
+
+	void IRigidBody::SetMass(double InM)
+	{
+		RIGIDPHYSICS_API_UNSUPPORTED();
+	}
+
+	FVector3d IRigidBody::GetInertia() const
+	{
+		RIGIDPHYSICS_API_UNSUPPORTED();
+		return FVector3d::Zero();
+	}
+
+	void IRigidBody::SetInertia(const FVector3d& InInertia)
+	{
+		RIGIDPHYSICS_API_UNSUPPORTED();
+	}
+
+	void IRigidBody::SetCCDEnabled(bool bInEnabled)
+	{
+		RIGIDPHYSICS_API_UNSUPPORTED();
+	}
+
+	void IRigidBody::SetMACDEnabled(bool bInEnabled)
+	{
+		RIGIDPHYSICS_API_UNSUPPORTED();
+	}
+
+	void IRigidBody::SetSmoothEdgeCollisionsEnabled(bool bInEnabled)
+	{
+		RIGIDPHYSICS_API_UNSUPPORTED();
+	}
+
+	void IRigidBody::SetMaxLinearVelocity(float InMax)
+	{
+		RIGIDPHYSICS_API_UNSUPPORTED();
+	}
+
+	void IRigidBody::SetMaxAngularVelocity(float InMax)
+	{
+		RIGIDPHYSICS_API_UNSUPPORTED();
+	}
+
+	void IRigidBody::SetInertiaConditioningEnabled(bool bInEnabled)
+	{
+		RIGIDPHYSICS_API_UNSUPPORTED();
+	}
+
+	void IRigidBody::SetGyroscopicTorqueEnabled(bool bInEnabled)
+	{
+		RIGIDPHYSICS_API_UNSUPPORTED();
+	}
+
+	void IRigidBody::SetGravityGroup(int32 InId)
+	{
+		RIGIDPHYSICS_API_UNSUPPORTED();
+	}
+
+	void IRigidBody::SetGravityScale(float InScale)
+	{
+		RIGIDPHYSICS_API_UNSUPPORTED();
+	}
+
+	void IRigidBody::SetWakeEventsEnabled(bool bInEnabled)
+	{
+		RIGIDPHYSICS_API_UNSUPPORTED();
+	}
+
+	void IRigidBody::SetNumPositionIterations(int32 InNum)
+	{
+		RIGIDPHYSICS_API_UNSUPPORTED();
+	}
+
+	void IRigidBody::SetNumVelocityIterations(int32 InNum)
+	{
+		RIGIDPHYSICS_API_UNSUPPORTED();
+	}
+
+	void IRigidBody::SetNumProjectionIterations(int32 InNum)
+	{
+		RIGIDPHYSICS_API_UNSUPPORTED();
+	}
+
+	void IRigidBody::SetIsSleeping(bool bInSleeping)
+	{
+		RIGIDPHYSICS_API_UNSUPPORTED();
+	}
+} // namespace UE::Physics
 
 #endif // UE_RIGIDPHYSICS_API_ENABLED
diff --git a/Engine/Source/Runtime/Experimental/RigidPhysics/Private/RigidPhysics/Internal/IRigidBodyContainer.cpp b/Engine/Source/Runtime/Experimental/RigidPhysics/Private/RigidPhysics/Internal/IRigidBodyContainer.cpp
--- a/Engine/Source/Runtime/Experimental/RigidPhysics/Private/RigidPhysics/Internal/IRigidBodyContainer.cpp
+++ b/Engine/Source/Runtime/Experimental/RigidPhysics/Private/RigidPhysics/Internal/IRigidBodyContainer.cpp
@@ -7,6 +7,22 @@ namespace UE::Physics
 {
 	UE_RIGIDPHYSICS_RIGIDTYPED_IMPL(IRigidBodyContainer, IRigidTyped);
 
+	void IRigidBodyContainer::SetId(const FRigidObjectId& InBodyId)
+	{
+		RIGIDPHYSICS_API_UNSUPPORTED();
+	}
+
+	IRigidBody* IRigidBodyContainer::CreateBody(const FRigidDebugName& InName, ERigidMovementType InMovementType)
+	{
+		RIGIDPHYSICS_API_UNSUPPORTED();
+		return nullptr;
+	}
+
+	void IRigidBodyContainer::DestroyBody(IRigidBody* InBody)
+	{
+		RIGIDPHYSICS_API_UNSUPPORTED();
+	}
+
 	FRigidBodyContainer::FRigidBodyContainer()
 	{
 		// TODO_CHAOSAPI: Make slack a parameter
diff --git a/Engine/Source/Runtime/Experimental/RigidPhysics/Private/RigidPhysics/Internal/IRigidScene.cpp b/Engine/Source/Runtime/Experimental/RigidPhysics/Private/RigidPhysics/Internal/IRigidScene.cpp
--- a/Engine/Source/Runtime/Experimental/RigidPhysics/Private/RigidPhysics/Internal/IRigidScene.cpp
+++ b/Engine/Source/Runtime/Experimental/RigidPhysics/Private/RigidPhysics/Internal/IRigidScene.cpp
@@ -11,6 +11,66 @@ namespace UE::Physics
 {
 	UE_RIGIDPHYSICS_RIGIDTYPED_IMPL(IRigidScene, IRigidLockable);
 
+	bool IRigidScene::Query(const FRigidQueryDesc& InQueryDesc, TArray<FRigidQueryHit>& OutHits)
+	{
+		RIGIDPHYSICS_API_UNSUPPORTED();
+		return false;
+	}
+
+	void IRigidScene::Tick(float InDt)
+	{
+		RIGIDPHYSICS_API_UNSUPPORTED();
+	}
+
+	void IRigidScene::RegisterModifier(IRigidSceneModifier* InModifier)
+	{
+		RIGIDPHYSICS_API_UNSUPPORTED();
+	}
+
+	void IRigidScene::UnregisterModifier(IRigidSceneModifier* InModifier)
+	{
+		RIGIDPHYSICS_API_UNSUPPORTED();
+	}
+
+	IRigidBody* IRigidScene::CreateBody(const FRigidDebugName& InName, ERigidMovementType InMovementType)
+	{
+		RIGIDPHYSICS_API_UNSUPPORTED();
+		return nullptr;
+	}
+
+	void IRigidScene::DestroyBody(IRigidBody* InBody)
+	{
+		RIGIDPHYSICS_API_UNSUPPORTED();
+	}
+
+	IRigidGeometryCollection* IRigidScene::CreateGeometryCollection(const FRigidDebugName& InName, const FSimulationParameters& InParameters)
+	{
+		RIGIDPHYSICS_API_UNSUPPORTED();
+		return nullptr;
+	}
+
+	void IRigidScene::DestroyGeometryCollection(IRigidGeometryCollection* InGC)
+	{
+		RIGIDPHYSICS_API_UNSUPPORTED();
+	}
+
+	IJointConstraint6DOF* IRigidScene::CreateJointConstraint6DOF()
+	{
+		RIGIDPHYSICS_API_UNSUPPORTED();
+		return nullptr;
+	}
+
+	void IRigidScene::DestroyJointConstraint(IJointConstraint* InConstraint)
+	{
+		RIGIDPHYSICS_API_UNSUPPORTED();
+	}
+
+	IJointConstraint* IRigidScene::GetJointConstraint(const FRigidObjectId& InConstraintId) const
+	{
+		RIGIDPHYSICS_API_UNSUPPORTED();
+		return nullptr;
+	}
+
 	FRigidScene::FRigidScene()
 	{
 		// TODO_CHAOSAPI: registry slack should be a parameter
diff --git a/Engine/Source/Runtime/Experimental/RigidPhysics/Private/RigidPhysics/RigidLockable.cpp b/Engine/Source/Runtime/Experimental/RigidPhysics/Private/RigidPhysics/RigidLockable.cpp
--- a/Engine/Source/Runtime/Experimental/RigidPhysics/Private/RigidPhysics/RigidLockable.cpp
+++ b/Engine/Source/Runtime/Experimental/RigidPhysics/Private/RigidPhysics/RigidLockable.cpp
@@ -6,6 +6,6 @@
 namespace UE::Physics
 {
 	UE_RIGIDPHYSICS_RIGIDTYPED_IMPL(IRigidLockable, IRigidTyped);
-}
+} //namespace UE::Physics
 
-#endif // UE_RIGIDPHYSICS_API_ENABLED
\ No newline at end of file
+#endif // UE_RIGIDPHYSICS_API_ENABLED
diff --git a/Engine/Source/Runtime/Experimental/RigidPhysics/Public/RigidPhysics/Internal/IJointConstraint.h b/Engine/Source/Runtime/Experimental/RigidPhysics/Public/RigidPhysics/Internal/IJointConstraint.h
--- a/Engine/Source/Runtime/Experimental/RigidPhysics/Public/RigidPhysics/Internal/IJointConstraint.h
+++ b/Engine/Source/Runtime/Experimental/RigidPhysics/Public/RigidPhysics/Internal/IJointConstraint.h
@@ -17,8 +17,8 @@ namespace UE::Physics
 
 		IJointConstraint() = default;
 		virtual ~IJointConstraint() = default;
-		
-		UE_INTERNAL RIGIDPHYSICS_API virtual FJointConstraintHandle GetHandle() const = 0;
+
+		UE_INTERNAL virtual FJointConstraintHandle GetHandle() const = 0;
 	};
 } // namespace UE::Physics
 
diff --git a/Engine/Source/Runtime/Experimental/RigidPhysics/Public/RigidPhysics/Internal/IJointConstraint6DOF.h b/Engine/Source/Runtime/Experimental/RigidPhysics/Public/RigidPhysics/Internal/IJointConstraint6DOF.h
--- a/Engine/Source/Runtime/Experimental/RigidPhysics/Public/RigidPhysics/Internal/IJointConstraint6DOF.h
+++ b/Engine/Source/Runtime/Experimental/RigidPhysics/Public/RigidPhysics/Internal/IJointConstraint6DOF.h
@@ -19,67 +19,26 @@ namespace UE::Physics
 		IJointConstraint6DOF() = default;
 		virtual ~IJointConstraint6DOF() = default;
 
-		UE_INTERNAL RIGIDPHYSICS_API virtual FJointConstraintHandle GetHandle() const = 0;
+		UE_INTERNAL virtual FJointConstraintHandle GetHandle() const = 0;
 		UE_INTERNAL RIGIDPHYSICS_API FJointConstraint6DOFHandle GetTypedHandle() const;
 
-		UE_INTERNAL RIGIDPHYSICS_API virtual void Activate()
-		{
-			RIGIDPHYSICS_API_UNSUPPORTED();
-		}
+		UE_INTERNAL RIGIDPHYSICS_API virtual void Activate();
 
-		UE_INTERNAL RIGIDPHYSICS_API virtual void SetBodies(IRigidBody* Body0, IRigidBody* Body1)
-		{
-			RIGIDPHYSICS_API_UNSUPPORTED();
-		}
+		UE_INTERNAL RIGIDPHYSICS_API virtual void SetBodies(IRigidBody* Body0, IRigidBody* Body1);
 
-		UE_INTERNAL RIGIDPHYSICS_API virtual void SetJointTransforms(const FTransform3d& InTransform0, const FTransform3d& InTransform1)
-		{
-			RIGIDPHYSICS_API_UNSUPPORTED();
-		}
+		UE_INTERNAL RIGIDPHYSICS_API virtual void SetJointTransforms(const FTransform3d& InTransform0, const FTransform3d& InTransform1);
 
-		UE_INTERNAL RIGIDPHYSICS_API virtual double GetLinearLimit() const
-		{
-			RIGIDPHYSICS_API_UNSUPPORTED();
-			return 0;
-		}
+		UE_INTERNAL RIGIDPHYSICS_API virtual double GetLinearLimit() const;
+		UE_INTERNAL RIGIDPHYSICS_API virtual void SetLinearLimit(const double InLinearLimit);
 
-		UE_INTERNAL RIGIDPHYSICS_API virtual void SetLinearLimit(const double InLinearLimit)
-		{
-			RIGIDPHYSICS_API_UNSUPPORTED();
-		}
+		UE_INTERNAL RIGIDPHYSICS_API virtual EJointMotionType GetLinearMotionTypesX() const;
+		UE_INTERNAL RIGIDPHYSICS_API virtual void SetLinearMotionTypesX(EJointMotionType InMotionType);
 
-		UE_INTERNAL RIGIDPHYSICS_API virtual EJointMotionType GetLinearMotionTypesX() const
-		{
-			RIGIDPHYSICS_API_UNSUPPORTED();
-			return EJointMotionType::Free;
-		}
+		UE_INTERNAL RIGIDPHYSICS_API virtual EJointMotionType GetLinearMotionTypesY() const;
+		UE_INTERNAL RIGIDPHYSICS_API virtual void SetLinearMotionTypesY(EJointMotionType InMotionType);
 
-		UE_INTERNAL RIGIDPHYSICS_API virtual void SetLinearMotionTypesX(EJointMotionType InMotionType)
-		{
-			RIGIDPHYSICS_API_UNSUPPORTED();
-		}
-
-		UE_INTERNAL RIGIDPHYSICS_API virtual EJointMotionType GetLinearMotionTypesY() const
-		{
-			RIGIDPHYSICS_API_UNSUPPORTED();
-			return EJointMotionType::Free;
-		}
-
-		UE_INTERNAL RIGIDPHYSICS_API virtual void SetLinearMotionTypesY(EJointMotionType InMotionType)
-		{
-			RIGIDPHYSICS_API_UNSUPPORTED();
-		}
-
-		UE_INTERNAL RIGIDPHYSICS_API virtual EJointMotionType GetLinearMotionTypesZ() const
-		{
-			RIGIDPHYSICS_API_UNSUPPORTED();
-			return EJointMotionType::Free;
-		}
-
-		UE_INTERNAL RIGIDPHYSICS_API virtual void SetLinearMotionTypesZ(EJointMotionType InMotionType)
-		{
-			RIGIDPHYSICS_API_UNSUPPORTED();
-		}
+		UE_INTERNAL RIGIDPHYSICS_API virtual EJointMotionType GetLinearMotionTypesZ() const;
+		UE_INTERNAL RIGIDPHYSICS_API virtual void SetLinearMotionTypesZ(EJointMotionType InMotionType);
 	};
 } // namespace UE::Physics
 
diff --git a/Engine/Source/Runtime/Experimental/RigidPhysics/Public/RigidPhysics/Internal/IRigidBody.h b/Engine/Source/Runtime/Experimental/RigidPhysics/Public/RigidPhysics/Internal/IRigidBody.h
--- a/Engine/Source/Runtime/Experimental/RigidPhysics/Public/RigidPhysics/Internal/IRigidBody.h
+++ b/Engine/Source/Runtime/Experimental/RigidPhysics/Public/RigidPhysics/Internal/IRigidBody.h
@@ -29,194 +29,94 @@ namespace UE::Physics
 
 		// Is the body active in its scene?
 		// Inactive bodies do not participate in physics in any way.
-		UE_INTERNAL RIGIDPHYSICS_API virtual bool IsActive() const
-		{
-			RIGIDPHYSICS_API_UNSUPPORTED();
-			return false;
-		}
-
-		UE_INTERNAL RIGIDPHYSICS_API virtual void Activate()
-		{
-			RIGIDPHYSICS_API_UNSUPPORTED();
-		}
-
-		UE_INTERNAL RIGIDPHYSICS_API virtual void Deactivate()
-		{
-			RIGIDPHYSICS_API_UNSUPPORTED();
-		}
+		UE_INTERNAL RIGIDPHYSICS_API virtual bool IsActive() const;
+		UE_INTERNAL RIGIDPHYSICS_API virtual void Activate();
+		UE_INTERNAL RIGIDPHYSICS_API virtual void Deactivate();
 
 
 		// The handle to this body
-		UE_INTERNAL RIGIDPHYSICS_API virtual FRigidBodyHandle GetHandle() const = 0;
+		UE_INTERNAL virtual FRigidBodyHandle GetHandle() const = 0;
 
-		UE_INTERNAL RIGIDPHYSICS_API virtual void SetId(const FRigidObjectId& InId)
-		{
-			RIGIDPHYSICS_API_UNSUPPORTED();
-		}
+		UE_INTERNAL RIGIDPHYSICS_API virtual void SetId(const FRigidObjectId& InId);
 
 		// Is this body Static? Static bodies can be moved but have no velocity.
-		UE_INTERNAL RIGIDPHYSICS_API virtual bool IsStatic() const = 0;
+		UE_INTERNAL virtual bool IsStatic() const = 0;
 
 		// Is this body kinematic? Kinematic bodies can be explicitly moved and have velocity, but do
 		// not react to forces, collisions, etc.
-		UE_INTERNAL RIGIDPHYSICS_API virtual bool IsKinematic() const = 0;
+		UE_INTERNAL virtual bool IsKinematic() const = 0;
 
 		// Is this a dynamic body? Dynamic bodies have momentum and react to forces and collisions.
-		UE_INTERNAL RIGIDPHYSICS_API virtual bool IsDynamic() const = 0;
+		UE_INTERNAL virtual bool IsDynamic() const = 0;
 
 		// Make the body dynamic (will only affect Kinematic bodies. Static bodies cannot be made dynamic)
-		UE_INTERNAL RIGIDPHYSICS_API virtual void SetIsDynamic(bool bIsDynamic) const
-		{
-			RIGIDPHYSICS_API_UNSUPPORTED();
-		}
+		UE_INTERNAL RIGIDPHYSICS_API virtual void SetIsDynamic(bool bIsDynamic) const;
 
 		// Is this a sleeping dynamic body? Sleeping bodies are dynamic bodies that are ignored by
 		// the simulation (as an optimization) until the scene changes in a way that might affect them, 
 		// or they are explicitly woken.
-		UE_INTERNAL RIGIDPHYSICS_API virtual bool IsSleeping() const
-		{
-			RIGIDPHYSICS_API_UNSUPPORTED();
-			return false;
-		}
+		UE_INTERNAL RIGIDPHYSICS_API virtual bool IsSleeping() const;
 
-		UE_INTERNAL RIGIDPHYSICS_API virtual int32 GetNumShapes() const = 0;
-		UE_INTERNAL RIGIDPHYSICS_API virtual FConstRigidShapeInstancePtr GetShape(int32 ShapeIndex) const = 0;
+		UE_INTERNAL virtual int32 GetNumShapes() const = 0;
+		UE_INTERNAL virtual FConstRigidShapeInstancePtr GetShape(int32 ShapeIndex) const = 0;
 
 		// Add a Collision Shape
-		UE_INTERNAL RIGIDPHYSICS_API virtual void AddShape(const FConstRigidShapeInstancePtr& InShape)
-		{
-			RIGIDPHYSICS_API_UNSUPPORTED();
-		}
+		UE_INTERNAL RIGIDPHYSICS_API virtual void AddShape(const FConstRigidShapeInstancePtr& InShape);
 
 		// Set the transform and initialize body state (velocity etc.) at the new location.
-		UE_INTERNAL RIGIDPHYSICS_API virtual void InitTransform(const FTransform3d& InTransform)
-		{
-			RIGIDPHYSICS_API_UNSUPPORTED();
-		}
+		UE_INTERNAL RIGIDPHYSICS_API virtual void InitTransform(const FTransform3d& InTransform);
 
 		// Set the transform without changing any other rigid body state.
-		UE_INTERNAL RIGIDPHYSICS_API virtual void UpdateTransform(const FTransform3d& InTransform)
-		{
-			RIGIDPHYSICS_API_UNSUPPORTED();
-		}
+		UE_INTERNAL RIGIDPHYSICS_API virtual void UpdateTransform(const FTransform3d& InTransform);
 
 		// Set the kinematic target. The body will attempt to reach the target at the next
 		// tick. Kinematic bodies will always reach the target bscause they do not respond
 		// to forces or collisions. Dynamic bodies will only reach their target if their
 		// velocity in the absence of interactions (forces, collisions, or other constraints).
-		UE_INTERNAL RIGIDPHYSICS_API virtual void SetKinematicTarget(const FTransform3d& InTransform)
-		{
-			RIGIDPHYSICS_API_UNSUPPORTED();
-		}
+		UE_INTERNAL RIGIDPHYSICS_API virtual void SetKinematicTarget(const FTransform3d& InTransform);
 
 		// Get the current transform. NOTE: This will immediately reflect any recent calls to
 		// to InitTransform or UpdateTransform, but will not reflect the latest KinematicTarget
 		// until after the next physics tick.
-		UE_INTERNAL RIGIDPHYSICS_API virtual FTransform3d GetTransform() const
-		{
-			RIGIDPHYSICS_API_UNSUPPORTED();
-			return FTransform3d::Identity;
-		}
+		UE_INTERNAL RIGIDPHYSICS_API virtual FTransform3d GetTransform() const;
 
 		// Get the world-space bounds at the current transform
-		UE_INTERNAL RIGIDPHYSICS_API virtual FBounds3d GetBounds() const
-		{
-			RIGIDPHYSICS_API_UNSUPPORTED();
-			return FBounds3d(FVector3d::Zero(), FVector3d::Zero());
-		}
+		UE_INTERNAL RIGIDPHYSICS_API virtual FBounds3d GetBounds() const;
 
 		// TODO: Maybe just Set/GetMassProperties
-		UE_INTERNAL RIGIDPHYSICS_API virtual double GetMass() const
-		{
-			RIGIDPHYSICS_API_UNSUPPORTED();
-			return 0;
-		}
-
-		UE_INTERNAL RIGIDPHYSICS_API virtual void SetMass(double InM)
-		{
-			RIGIDPHYSICS_API_UNSUPPORTED();
-		}
-
-		UE_INTERNAL RIGIDPHYSICS_API virtual FVector3d GetInertia() const
-		{
-			RIGIDPHYSICS_API_UNSUPPORTED();
-			return FVector3d::Zero();
-		}
-
-		UE_INTERNAL RIGIDPHYSICS_API virtual void SetInertia(const FVector3d& InInertia)
-		{
-			RIGIDPHYSICS_API_UNSUPPORTED();
-		}
-
-		UE_INTERNAL RIGIDPHYSICS_API virtual void SetCCDEnabled(bool bInEnabled)
-		{
-			RIGIDPHYSICS_API_UNSUPPORTED();
-		}
-
-		UE_INTERNAL RIGIDPHYSICS_API virtual void SetMACDEnabled(bool bInEnabled)
-		{
-			RIGIDPHYSICS_API_UNSUPPORTED();
-		}
-
-		UE_INTERNAL RIGIDPHYSICS_API virtual void SetSmoothEdgeCollisionsEnabled(bool bInEnabled)
-		{
-			RIGIDPHYSICS_API_UNSUPPORTED();
-		}
-
-		UE_INTERNAL RIGIDPHYSICS_API virtual void SetMaxLinearVelocity(float InMax)
-		{
-			RIGIDPHYSICS_API_UNSUPPORTED();
-		}
-
-		UE_INTERNAL RIGIDPHYSICS_API virtual void SetMaxAngularVelocity(float InMax)
-		{
-			RIGIDPHYSICS_API_UNSUPPORTED();
-		}
-
-		UE_INTERNAL RIGIDPHYSICS_API virtual void SetInertiaConditioningEnabled(bool bInEnabled)
-		{
-			RIGIDPHYSICS_API_UNSUPPORTED();
-		}
-
-		UE_INTERNAL RIGIDPHYSICS_API virtual void SetGyroscopicTorqueEnabled(bool bInEnabled)
-		{
-			RIGIDPHYSICS_API_UNSUPPORTED();
-		}
-
-		UE_INTERNAL RIGIDPHYSICS_API virtual void SetGravityGroup(int32 InId)
-		{
-			RIGIDPHYSICS_API_UNSUPPORTED();
-		}
-
-		UE_INTERNAL RIGIDPHYSICS_API virtual void SetGravityScale(float InScale)
-		{
-			RIGIDPHYSICS_API_UNSUPPORTED();
-		}
-
-		UE_INTERNAL RIGIDPHYSICS_API virtual void SetWakeEventsEnabled(bool bInEnabled)
-		{
-			RIGIDPHYSICS_API_UNSUPPORTED();
-		}
-
-		UE_INTERNAL RIGIDPHYSICS_API virtual void SetNumPositionIterations(int32 InNum)
-		{
-			RIGIDPHYSICS_API_UNSUPPORTED();
-		}
-
-		UE_INTERNAL RIGIDPHYSICS_API virtual void SetNumVelocityIterations(int32 InNum)
-		{
-			RIGIDPHYSICS_API_UNSUPPORTED();
-		}
-
-		UE_INTERNAL RIGIDPHYSICS_API virtual void SetNumProjectionIterations(int32 InNum)
-		{
-			RIGIDPHYSICS_API_UNSUPPORTED();
-		}
-
-		UE_INTERNAL RIGIDPHYSICS_API virtual void SetIsSleeping(bool bInSleeping)
-		{
-			RIGIDPHYSICS_API_UNSUPPORTED();
-		}
+		UE_INTERNAL RIGIDPHYSICS_API virtual double GetMass() const;
+		UE_INTERNAL RIGIDPHYSICS_API virtual void SetMass(double InM);
+
+		UE_INTERNAL RIGIDPHYSICS_API virtual FVector3d GetInertia() const;
+		UE_INTERNAL RIGIDPHYSICS_API virtual void SetInertia(const FVector3d& InInertia);
+
+		UE_INTERNAL RIGIDPHYSICS_API virtual void SetCCDEnabled(bool bInEnabled);
+
+		UE_INTERNAL RIGIDPHYSICS_API virtual void SetMACDEnabled(bool bInEnabled);
+
+		UE_INTERNAL RIGIDPHYSICS_API virtual void SetSmoothEdgeCollisionsEnabled(bool bInEnabled);
+
+		UE_INTERNAL RIGIDPHYSICS_API virtual void SetMaxLinearVelocity(float InMax);
+
+		UE_INTERNAL RIGIDPHYSICS_API virtual void SetMaxAngularVelocity(float InMax);
+
+		UE_INTERNAL RIGIDPHYSICS_API virtual void SetInertiaConditioningEnabled(bool bInEnabled);
+
+		UE_INTERNAL RIGIDPHYSICS_API virtual void SetGyroscopicTorqueEnabled(bool bInEnabled);
+
+		UE_INTERNAL RIGIDPHYSICS_API virtual void SetGravityGroup(int32 InId);
+
+		UE_INTERNAL RIGIDPHYSICS_API virtual void SetGravityScale(float InScale);
+
+		UE_INTERNAL RIGIDPHYSICS_API virtual void SetWakeEventsEnabled(bool bInEnabled);
+
+		UE_INTERNAL RIGIDPHYSICS_API virtual void SetNumPositionIterations(int32 InNum);
+
+		UE_INTERNAL RIGIDPHYSICS_API virtual void SetNumVelocityIterations(int32 InNum);
+
+		UE_INTERNAL RIGIDPHYSICS_API virtual void SetNumProjectionIterations(int32 InNum);
+
+		UE_INTERNAL RIGIDPHYSICS_API virtual void SetIsSleeping(bool bInSleeping);
 	};
 } // namespace UE::Physics
 
diff --git a/Engine/Source/Runtime/Experimental/RigidPhysics/Public/RigidPhysics/Internal/IRigidBodyContainer.h b/Engine/Source/Runtime/Experimental/RigidPhysics/Public/RigidPhysics/Internal/IRigidBodyContainer.h
--- a/Engine/Source/Runtime/Experimental/RigidPhysics/Public/RigidPhysics/Internal/IRigidBodyContainer.h
+++ b/Engine/Source/Runtime/Experimental/RigidPhysics/Public/RigidPhysics/Internal/IRigidBodyContainer.h
@@ -23,31 +23,20 @@ namespace UE::Physics
 		IRigidBodyContainer() = default;
 		virtual ~IRigidBodyContainer() = default;
 
-		UE_INTERNAL RIGIDPHYSICS_API virtual FRigidBodyContainerHandle GetHandle() const = 0;
+		UE_INTERNAL virtual FRigidBodyContainerHandle GetHandle() const = 0;
 
-		UE_INTERNAL RIGIDPHYSICS_API virtual int32 GetNumBodies() const = 0;
+		UE_INTERNAL virtual int32 GetNumBodies() const = 0;
 
-		UE_INTERNAL RIGIDPHYSICS_API virtual IRigidBody* GetBody(const FRigidObjectId& InBodyId) const = 0;
+		UE_INTERNAL virtual IRigidBody* GetBody(const FRigidObjectId& InBodyId) const = 0;
 
 		//
 		// Below here is GT only. See if we can separate the Game and SimCallback APIs
 		//
 
-		UE_INTERNAL RIGIDPHYSICS_API virtual void SetId(const FRigidObjectId& InBodyId)
-		{
-			RIGIDPHYSICS_API_UNSUPPORTED();
-		}
-
-		UE_INTERNAL RIGIDPHYSICS_API virtual IRigidBody* CreateBody(const FRigidDebugName& InName, ERigidMovementType InMovementType)
-		{
-			RIGIDPHYSICS_API_UNSUPPORTED();
-			return nullptr;
-		}
-
-		UE_INTERNAL RIGIDPHYSICS_API virtual void DestroyBody(IRigidBody* InBody)
-		{
-			RIGIDPHYSICS_API_UNSUPPORTED();
-		}
+		UE_INTERNAL RIGIDPHYSICS_API virtual void SetId(const FRigidObjectId& InBodyId);
+
+		UE_INTERNAL RIGIDPHYSICS_API virtual IRigidBody* CreateBody(const FRigidDebugName& InName, ERigidMovementType InMovementType);
+		UE_INTERNAL RIGIDPHYSICS_API virtual void DestroyBody(IRigidBody* InBody);
 	};
 
 	class UE_INTERNAL FRigidBodyContainer : public IRigidBodyContainer
diff --git a/Engine/Source/Runtime/Experimental/RigidPhysics/Public/RigidPhysics/Internal/IRigidScene.h b/Engine/Source/Runtime/Experimental/RigidPhysics/Public/RigidPhysics/Internal/IRigidScene.h
--- a/Engine/Source/Runtime/Experimental/RigidPhysics/Public/RigidPhysics/Internal/IRigidScene.h
+++ b/Engine/Source/Runtime/Experimental/RigidPhysics/Public/RigidPhysics/Internal/IRigidScene.h
@@ -23,76 +23,33 @@ namespace UE::Physics
 		IRigidScene() = default;
 		virtual ~IRigidScene() = default;
 
-		UE_INTERNAL RIGIDPHYSICS_API virtual FRigidSceneId GetId() const = 0;
-		UE_INTERNAL RIGIDPHYSICS_API virtual FRigidSceneHandle GetHandle() const = 0;
+		UE_INTERNAL virtual FRigidSceneId GetId() const = 0;
+		UE_INTERNAL virtual FRigidSceneHandle GetHandle() const = 0;
 
-		UE_INTERNAL RIGIDPHYSICS_API virtual IRigidBody* GetBody(const FRigidObjectId& InBodyId) const = 0;
+		UE_INTERNAL virtual IRigidBody* GetBody(const FRigidObjectId& InBodyId) const = 0;
 
-		UE_INTERNAL RIGIDPHYSICS_API virtual IRigidBodyContainer* GetBodyContainer(const FRigidObjectId& InBodyContainerId) const = 0;
+		UE_INTERNAL virtual IRigidBodyContainer* GetBodyContainer(const FRigidObjectId& InBodyContainerId) const = 0;
 
-		UE_INTERNAL RIGIDPHYSICS_API virtual IRigidGeometryCollection* GetGeometryCollection(const FRigidObjectId& InGeometryCollectionId) const = 0;
+		UE_INTERNAL virtual IRigidGeometryCollection* GetGeometryCollection(const FRigidObjectId& InGeometryCollectionId) const = 0;
 
-		UE_INTERNAL RIGIDPHYSICS_API virtual bool Query(const FRigidQueryDesc& InQueryDesc, TArray<FRigidQueryHit>& OutHits)
-		{
-			RIGIDPHYSICS_API_UNSUPPORTED();
-			return false;
-		}
+		UE_INTERNAL RIGIDPHYSICS_API virtual bool Query(const FRigidQueryDesc& InQueryDesc, TArray<FRigidQueryHit>& OutHits);
 
 		// Advance the simulation by Dt. If the implementation is asynchronous, this will return immediately
-		UE_INTERNAL RIGIDPHYSICS_API virtual void Tick(float InDt)
-		{
-			RIGIDPHYSICS_API_UNSUPPORTED();
-		}
+		UE_INTERNAL RIGIDPHYSICS_API virtual void Tick(float InDt);
 
 		// GT Only API - Move this somehow
-		UE_INTERNAL RIGIDPHYSICS_API virtual void RegisterModifier(IRigidSceneModifier* InModifier)
-		{
-			RIGIDPHYSICS_API_UNSUPPORTED();
-		}
-
-		UE_INTERNAL RIGIDPHYSICS_API virtual void UnregisterModifier(IRigidSceneModifier* InModifier)
-		{
-			RIGIDPHYSICS_API_UNSUPPORTED();
-		}
-
-		UE_INTERNAL RIGIDPHYSICS_API virtual IRigidBody* CreateBody(const FRigidDebugName& InName, ERigidMovementType InMovementType)
-		{
-			RIGIDPHYSICS_API_UNSUPPORTED();
-			return nullptr;
-		}
-
-		UE_INTERNAL RIGIDPHYSICS_API virtual void DestroyBody(IRigidBody* InBody)
-		{
-			RIGIDPHYSICS_API_UNSUPPORTED();
-		}
-
-		UE_INTERNAL RIGIDPHYSICS_API virtual IRigidGeometryCollection* CreateGeometryCollection(const FRigidDebugName& InName, const FSimulationParameters& InParameters)
-		{
-			RIGIDPHYSICS_API_UNSUPPORTED();
-			return nullptr;
-		}
-
-		UE_INTERNAL RIGIDPHYSICS_API virtual void DestroyGeometryCollection(IRigidGeometryCollection* InGC)
-		{
-			RIGIDPHYSICS_API_UNSUPPORTED();
-		}
-		
-		UE_INTERNAL RIGIDPHYSICS_API virtual IJointConstraint6DOF* CreateJointConstraint6DOF()
-		{
-			RIGIDPHYSICS_API_UNSUPPORTED();
-			return nullptr;
-		}
-
-		UE_INTERNAL RIGIDPHYSICS_API virtual void DestroyJointConstraint(IJointConstraint* InConstraint)
-		{
-			RIGIDPHYSICS_API_UNSUPPORTED();
-		}
-		
-		UE_INTERNAL RIGIDPHYSICS_API virtual IJointConstraint* GetJointConstraint(const FRigidObjectId& InConstraintId) const
-		{
-			RIGIDPHYSICS_API_UNSUPPORTED();
-			return nullptr;
-		}
+		UE_INTERNAL RIGIDPHYSICS_API virtual void RegisterModifier(IRigidSceneModifier* InModifier);
+		UE_INTERNAL RIGIDPHYSICS_API virtual void UnregisterModifier(IRigidSceneModifier* InModifier);
+
+		UE_INTERNAL RIGIDPHYSICS_API virtual IRigidBody* CreateBody(const FRigidDebugName& InName, ERigidMovementType InMovementType);
+		UE_INTERNAL RIGIDPHYSICS_API virtual void DestroyBody(IRigidBody* InBody);
+
+		UE_INTERNAL RIGIDPHYSICS_API virtual IRigidGeometryCollection* CreateGeometryCollection(const FRigidDebugName& InName, const FSimulationParameters& InParameters);
+		UE_INTERNAL RIGIDPHYSICS_API virtual void DestroyGeometryCollection(IRigidGeometryCollection* InGC);
+
+		UE_INTERNAL RIGIDPHYSICS_API virtual IJointConstraint6DOF* CreateJointConstraint6DOF();
+		UE_INTERNAL RIGIDPHYSICS_API virtual void DestroyJointConstraint(IJointConstraint* InConstraint);
+		UE_INTERNAL RIGIDPHYSICS_API virtual IJointConstraint* GetJointConstraint(const FRigidObjectId& InConstraintId) const;
 	};
 
 	//
diff --git a/Engine/Source/Runtime/Experimental/RigidPhysics/Public/RigidPhysics/RigidLockable.h b/Engine/Source/Runtime/Experimental/RigidPhysics/Public/RigidPhysics/RigidLockable.h
--- a/Engine/Source/Runtime/Experimental/RigidPhysics/Public/RigidPhysics/RigidLockable.h
+++ b/Engine/Source/Runtime/Experimental/RigidPhysics/Public/RigidPhysics/RigidLockable.h
@@ -13,17 +13,12 @@ namespace UE::Physics
 	public:
 		UE_RIGIDPHYSICS_RIGIDTYPED_DECL(RIGIDPHYSICS_API, IRigidLockable);
 
-		UE_INTERNAL IRigidLockable()
-		{
-		}
+		UE_INTERNAL IRigidLockable() = default;
+		UE_INTERNAL virtual ~IRigidLockable() = default;
 
-		UE_INTERNAL virtual ~IRigidLockable()
-		{
-		}
-
-		UE_INTERNAL RIGIDPHYSICS_API virtual void Lock(ERigidLockType InLockType) const = 0;
-		UE_INTERNAL RIGIDPHYSICS_API virtual void Unlock(ERigidLockType InLockType) const = 0;
+		UE_INTERNAL virtual void Lock(ERigidLockType InLockType) const = 0;
+		UE_INTERNAL virtual void Unlock(ERigidLockType InLockType) const = 0;
 	};
-}
+} // namespace UE::Physics
 
-#endif // UE_RIGIDPHYSICS_API_ENABLED
\ No newline at end of file
+#endif // UE_RIGIDPHYSICS_API_ENABLED
