GeometryCollection: Extend transform triple buffering to manage 4 buffers to take into account the need to always have a previous frame transform buffer.
Todo: Investigate whether to disable by default or deprecate explicit triple buffering, and just let the RHI handle buffer renaming.
#Synced-CL 51342769

diff --git a/Engine/Source/Runtime/Experimental/GeometryCollectionEngine/Private/GeometryCollection/GeometryCollectionSceneProxy.cpp b/Engine/Source/Runtime/Experimental/GeometryCollectionEngine/Private/GeometryCollection/GeometryCollectionSceneProxy.cpp
--- a/Engine/Source/Runtime/Experimental/GeometryCollectionEngine/Private/GeometryCollection/GeometryCollectionSceneProxy.cpp
+++ b/Engine/Source/Runtime/Experimental/GeometryCollectionEngine/Private/GeometryCollection/GeometryCollectionSceneProxy.cpp
@@ -340,8 +340,8 @@ void FGeometryCollectionSceneProxyBase::SetDynamicData_RenderThread(FRHICommandL
 
 			if (bLocalGeometryCollectionTripleBufferUploads && TransformBuffers.Num() == 1)
 			{
-				TransformBuffers.AddDefaulted(2);
-				check(TransformBuffers.Num() == 3);
+				TransformBuffers.AddDefaulted(3);
+				check(TransformBuffers.Num() == 4);
 
 				for (int32 i = 1; i < TransformBuffers.Num(); i++)
 				{
diff --git a/Engine/Source/Runtime/Experimental/GeometryCollectionEngine/Private/GeometryCollection/GeometryCollectionSceneProxy.h b/Engine/Source/Runtime/Experimental/GeometryCollectionEngine/Private/GeometryCollection/GeometryCollectionSceneProxy.h
--- a/Engine/Source/Runtime/Experimental/GeometryCollectionEngine/Private/GeometryCollection/GeometryCollectionSceneProxy.h
+++ b/Engine/Source/Runtime/Experimental/GeometryCollectionEngine/Private/GeometryCollection/GeometryCollectionSceneProxy.h
@@ -239,7 +239,7 @@ private:
 	FPositionVertexBuffer SkinnedPositionVertexBuffer;
 
 	int32 CurrentTransformBufferIndex = 0;
-	TArray<FGeometryCollectionTransformBuffer, TInlineAllocator<3>> TransformBuffers;
+	TArray<FGeometryCollectionTransformBuffer, TInlineAllocator<4>> TransformBuffers;
 
 	const bool bHasRayTracingRepresentation = false;
 
