Added Cvar and define for default behavior of IAX development mode.

#rnx

#Synced-CL 51343971

diff --git a/Engine/Source/Runtime/Experimental/IoStore/OnDemand/Private/OnDemandIoStore.cpp b/Engine/Source/Runtime/Experimental/IoStore/OnDemand/Private/OnDemandIoStore.cpp
--- a/Engine/Source/Runtime/Experimental/IoStore/OnDemand/Private/OnDemandIoStore.cpp
+++ b/Engine/Source/Runtime/Experimental/IoStore/OnDemand/Private/OnDemandIoStore.cpp
@@ -50,6 +50,10 @@
 #include "String/LexFromString.h"
 #endif
 
+#ifndef UE_IOSTORE_ONDEMAND_DEVMODE_ENABLED
+#define UE_IOSTORE_ONDEMAND_DEVMODE_ENABLED 0
+#endif
+
 ///////////////////////////////////////////////////////////////////////////////
 namespace UE::IoStore
 {
@@ -68,14 +72,23 @@ static FAutoConsoleVariableRef CVar_IaxMemoryMappedTocEnabled(
 	TEXT("Whether container TOC memory-mapping is enabled.")
 );
 
+bool GIaxDevModeEnabled = UE_IOSTORE_ONDEMAND_DEVMODE_ENABLED;
+static FAutoConsoleVariableRef CVar_IaxDevModeEnabled (
+	TEXT("iax.DevModeEnabled"),
+	GIaxDevModeEnabled,
+	TEXT("Whether IAX development mode is enabled or not.")
+);
+
 ///////////////////////////////////////////////////////////////////////////////
 bool IsDevModeEnabled()
 {
 #if UE_BUILD_SHIPPING
 	return false;
 #else
-	static bool bIsDevModeEnabled = FParse::Param(FCommandLine::Get(), TEXT("iax.DevMode"));
-	return bIsDevModeEnabled;
+	static bool bDevModeEnabled = FParse::Param(FCommandLine::Get(), TEXT("iax.DevMode"));
+	static bool bDevModeDisabled = FParse::Param(FCommandLine::Get(), TEXT("iax.DisableDevMode"));
+
+	return !bDevModeDisabled && (GIaxDevModeEnabled || bDevModeEnabled);
 #endif
 }
 
@@ -2017,6 +2030,11 @@ void FOnDemandIoStore::SetDevelopmentExtension(TSharedPtr<IOnDemandDevelopmentEx
 	DevExt = Ext;
 }
 
+bool FOnDemandIoStore::IsDevelopmentModeEnabled() const
+{
+	return IsDevModeEnabled();
+}
+
 } // namespace UE::IoStore
 
 ////////////////////////////////////////////////////////////////////////////////
diff --git a/Engine/Source/Runtime/Experimental/IoStore/OnDemand/Private/OnDemandIoStore.h b/Engine/Source/Runtime/Experimental/IoStore/OnDemand/Private/OnDemandIoStore.h
--- a/Engine/Source/Runtime/Experimental/IoStore/OnDemand/Private/OnDemandIoStore.h
+++ b/Engine/Source/Runtime/Experimental/IoStore/OnDemand/Private/OnDemandIoStore.h
@@ -373,6 +373,7 @@ public:
 	virtual TUniquePtr<IAnalyticsRecording> StartAnalyticsRecording() const override;
 	virtual void						OnImmediateAnalytic(FOnDemandImmediateAnalyticHandler EventHandler) override;
 	virtual void						SetDevelopmentExtension(FSharedOnDemandDevelopmentExtension Ext) override;
+	virtual bool						IsDevelopmentModeEnabled() const override;
 	virtual void						CancelInstallRequest(FSharedInternalInstallRequest InstallRequest) override;
 	virtual void						UpdateInstallRequestPriority(FSharedInternalInstallRequest InstallRequest, int32 NewPriority) override;
 	virtual void						ReleaseContent(FOnDemandInternalContentHandle& ContentHandle) override;
diff --git a/Engine/Source/Runtime/Experimental/IoStore/OnDemand/Private/Statistics.cpp b/Engine/Source/Runtime/Experimental/IoStore/OnDemand/Private/Statistics.cpp
--- a/Engine/Source/Runtime/Experimental/IoStore/OnDemand/Private/Statistics.cpp
+++ b/Engine/Source/Runtime/Experimental/IoStore/OnDemand/Private/Statistics.cpp
@@ -28,6 +28,8 @@ LLM_DEFINE_TAG(Ias);
 namespace UE::IoStore
 {
 
+bool IsDevModeEnabled();
+
 extern int32 GHttpIoDispatcherMaxRangeKiB;
 
 int32 CalculateRequestQueueLengthMax();
@@ -880,6 +882,9 @@ void FOnDemandIoBackendStats::PrintOnScreenStatistics(FCoreDelegates::FSeverityM
 			{
 				BackendStatusText.Appendf(TEXT("%d KiB"), MaxRange);
 			}
+			
+			BackendStatusText.Append(" | DevMode: ");
+			BackendStatusText.Append(IsDevModeEnabled() ? TEXT("Enabled") : TEXT("Disabled"));
 		}
 
 		OutMessages.Add(FCoreDelegates::EOnScreenMessageSeverity::Info, FText::FromStringView(BackendStatusText.ToView()));
diff --git a/Engine/Source/Runtime/Experimental/IoStore/OnDemandCore/Public/IO/IoStoreOnDemand.h b/Engine/Source/Runtime/Experimental/IoStore/OnDemandCore/Public/IO/IoStoreOnDemand.h
--- a/Engine/Source/Runtime/Experimental/IoStore/OnDemandCore/Public/IO/IoStoreOnDemand.h
+++ b/Engine/Source/Runtime/Experimental/IoStore/OnDemandCore/Public/IO/IoStoreOnDemand.h
@@ -711,6 +711,8 @@ public:
 	virtual void OnImmediateAnalytic(FOnDemandImmediateAnalyticHandler EventHandler) = 0;
 	/** Set the development extension. */
 	virtual void SetDevelopmentExtension(FSharedOnDemandDevelopmentExtension Ctx) { }
+	/** Whether development mode is enabled or not. */
+	virtual bool IsDevelopmentModeEnabled() const { return false; }
 
 private:
 	/** Initialize the I/O store. Called after the I/O store factory module feature has been registered. */
