Added CSV metadata whether IAX development mode is enabled or not.

#rnx

#Synced-CL 51343971

diff --git a/Engine/Source/Runtime/Experimental/IoStore/OnDemand/Private/Statistics.cpp b/Engine/Source/Runtime/Experimental/IoStore/OnDemand/Private/Statistics.cpp
--- a/Engine/Source/Runtime/Experimental/IoStore/OnDemand/Private/Statistics.cpp
+++ b/Engine/Source/Runtime/Experimental/IoStore/OnDemand/Private/Statistics.cpp
@@ -674,6 +674,9 @@ FOnDemandIoBackendStats::FOnDemandIoBackendStats(FBackendStatus& InStatus)
 			}),
 		ECVF_Default));
 #endif // UE_ENABLE_ONSCREEN_STATISTICS
+#if CSV_PROFILER_STATS
+	CSVProfileEndRequestedHandle = FCsvProfiler::Get()->OnCSVProfileEndRequested().AddRaw(this, &FOnDemandIoBackendStats::OnCSVProfileEndRequested);
+#endif
 }
 
 FOnDemandIoBackendStats::~FOnDemandIoBackendStats()
@@ -681,6 +684,9 @@ FOnDemandIoBackendStats::~FOnDemandIoBackendStats()
 	FCoreDelegates::ApplicationHasEnteredForegroundDelegate.Remove(OnApplicationResumeHandle);
 	FCoreDelegates::OnEndFrame.Remove(EndFrameDelegateHandle);
 	FCoreDelegates::OnGetOnScreenMessages.Remove(OnScreenDelegateHandle);
+#if CSV_PROFILER_STATS
+	FCsvProfiler::Get()->OnCSVProfileEndRequested().Remove(CSVProfileEndRequestedHandle);
+#endif
 
 	GStatistics = nullptr;
 }
@@ -1556,6 +1562,13 @@ void FOnDemandIoBackendStats::OnApplicationResume()
 	++GAppResumeCount;
 }
 
+void FOnDemandIoBackendStats::OnCSVProfileEndRequested()
+{
+#if CSV_PROFILER_STATS
+	CSV_METADATA(TEXT("IaxDevMode"), IsDevModeEnabled() ? TEXT("1") : TEXT("0"));
+#endif
+}
+
 ////////////////////////////////////////////////////////////////////////////////
 void FOnDemandContentInstallerStats::OnRequestEnqueued()
 {
diff --git a/Engine/Source/Runtime/Experimental/IoStore/OnDemand/Private/Statistics.h b/Engine/Source/Runtime/Experimental/IoStore/OnDemand/Private/Statistics.h
--- a/Engine/Source/Runtime/Experimental/IoStore/OnDemand/Private/Statistics.h
+++ b/Engine/Source/Runtime/Experimental/IoStore/OnDemand/Private/Statistics.h
@@ -112,6 +112,7 @@ private:
 #if IAS_WITH_STATISTICS
 
 	void OnApplicationResume();
+	void OnCSVProfileEndRequested();
 
 	void UpdateCSVValues();
 	void PrintPeriodicLogging();
@@ -122,6 +123,7 @@ private:
 	FDelegateHandle OnApplicationResumeHandle;
 	FDelegateHandle EndFrameDelegateHandle;
 	FDelegateHandle OnScreenDelegateHandle;
+	FDelegateHandle CSVProfileEndRequestedHandle;
 
 #if UE_ENABLE_ONSCREEN_STATISTICS
 	FConsoleCommandPtr ResetStatisticsCommand;
