LaneTrace updates following integration into Zen

#rnx

#Synced-CL 50969227

diff --git a/Engine/Source/Runtime/Experimental/IoStore/HttpClient/Internal/IO/Http/LaneTrace.h b/Engine/Source/Runtime/Experimental/IoStore/HttpClient/Internal/IO/Http/LaneTrace.h
--- a/Engine/Source/Runtime/Experimental/IoStore/HttpClient/Internal/IO/Http/LaneTrace.h
+++ b/Engine/Source/Runtime/Experimental/IoStore/HttpClient/Internal/IO/Http/LaneTrace.h
@@ -2,11 +2,14 @@
 
 #pragma once
 
-#if !defined(NO_UE_INCLUDES)
-#include <Trace/Trace.h>
-#include <Containers/StringView.h>
-#include <ProfilingDebugging/CpuProfilerTrace.h>
-#include <Templates/UnrealTemplate.h>
+#if __has_include(<Containers/StringView.h>)
+#	include <Trace/Trace.h>
+#	include <Containers/StringView.h>
+#	include <ProfilingDebugging/CpuProfilerTrace.h>
+#	include <Templates/UnrealTemplate.h>
+#else
+#	define UE_LANETRACE_ENABLED		2
+#	define IOSTOREHTTPCLIENT_API
 #endif
 
 ////////////////////////////////////////////////////////////////////////////////
@@ -22,6 +25,12 @@
 #	define UE_API					inline
 #endif
 
+
+////////////////////////////////////////////////////////////////////////////////
+#if UE_LANETRACE_ENABLED == 2
+namespace trace {
+#endif
+
 ////////////////////////////////////////////////////////////////////////////////
 struct FLaneTraceSpec
 {
@@ -126,3 +135,12 @@ public:
 };
 
 #endif // UE_LANETRACE_ENABLED
+
+////////////////////////////////////////////////////////////////////////////////
+#if UE_LANETRACE_ENABLED == 2
+} // namespace trace
+#endif
+
+#if defined(TRACE_IMPLEMENT) && TRACE_IMPLEMENT
+#	include "lane_trace.inl"
+#endif
diff --git a/Engine/Source/Runtime/Experimental/IoStore/HttpClient/Private/LaneTrace.cpp b/Engine/Source/Runtime/Experimental/IoStore/HttpClient/Private/LaneTrace.cpp
--- a/Engine/Source/Runtime/Experimental/IoStore/HttpClient/Private/LaneTrace.cpp
+++ b/Engine/Source/Runtime/Experimental/IoStore/HttpClient/Private/LaneTrace.cpp
@@ -1,19 +1,47 @@
 // Copyright Epic Games, Inc. All Rights Reserved.
 
-#include <IO/Http/LaneTrace.h>
+#if __has_include(<IO/Http/LaneTrace.h>)
+#	include <IO/Http/LaneTrace.h>
+#	if UE_LANETRACE_ENABLED
+#		include <HAL/PlatformTime.h>
+#		include <HAL/UnrealMemory.h>
+#		include <Trace/Trace.inl>
+#		include <Misc/ScopeLock.h>
+#	endif
+#endif
 
 #if UE_LANETRACE_ENABLED
 
-#define LANETRACE_UNTESTED 0
+////////////////////////////////////////////////////////////////////////////////
+#if UE_LANETRACE_ENABLED == 2
+
+#include <mutex>
+
+#define check(e) do { if ((e) == false) *(int volatile*)0 = 0; } while (false)
+
+struct FCriticalSection
+{
+	void		Lock()		{ Inner.lock(); }
+	void		Unlock()	{ Inner.unlock(); }
+	std::mutex	Inner;
+};
+
+struct FScopeLock
+{
+	FScopeLock(FCriticalSection* Lock) : Inner(Lock->Inner) {}
+	std::unique_lock<std::mutex> Inner;
+};
+
+namespace trace {
+#endif
 
-#include <HAL/PlatformTime.h>
-#include <HAL/UnrealMemory.h>
-#include <Trace/Trace.inl>
-#include <Misc/ScopeLock.h>
+// For this to work Insights would need to know it there may be a ThreadId field
+#define LANETRACE_USE_V2_EVENTS 0
 
 namespace LaneTraceDetail
 {
 
+////////////////////////////////////////////////////////////////////////////////
 UE_TRACE_EVENT_BEGIN($Trace, ThreadInfo, NoSync|Important)
 	UE_TRACE_EVENT_FIELD(uint32, ThreadId)
 	UE_TRACE_EVENT_FIELD(int32, SortHint)
@@ -25,7 +53,7 @@ UE_TRACE_EVENT_BEGIN(CpuProfiler, EventBatch)
 	UE_TRACE_EVENT_FIELD(uint16, ThreadId)
 UE_TRACE_EVENT_END()
 
-#if LANETRACE_UNTESTED
+#if LANETRACE_USE_V2_EVENTS
 UE_TRACE_EVENT_BEGIN(CpuProfiler, EventBatchV2)
 	UE_TRACE_EVENT_FIELD(uint8[], Data)
 	UE_TRACE_EVENT_FIELD(uint16, ThreadId)
@@ -89,7 +117,11 @@ static int32 Encode64_7bit(int64 Value, void* __restrict Out)
 ////////////////////////////////////////////////////////////////////////////////
 static uint64 TimeGetTimestamp()
 {
+#if UE_LANETRACE_ENABLED != 2
 	return FPlatformTime::Cycles64();
+#else
+	return trace::detail::TimeGetTimestamp();
+#endif
 }
 
 
@@ -113,7 +145,7 @@ private:
 		Overflow			= 24,
 		EnterLsb			= 1,
 		LeaveLsb			= 0,
-		TraceEventBatchVer	= 1,
+		TraceEventBatchVer	= 1 + LANETRACE_USE_V2_EVENTS,
 	};
 	uint64					LastTimestamp = 0;
 	uint64					PrevTimestamp = 0;
@@ -167,12 +199,12 @@ void FScopeBuffer::Flush(bool Force)
 			<< EventBatch.ThreadId(uint16(ThreadIdOverride))
 			<< EventBatch.Data(Buffer, uint32(ptrdiff_t(Cursor - Buffer)));
 	}
-#if LANETRACE_UNTESTED
+#if LANETRACE_USE_V2_EVENTS
 	else
 	{
 		UE_TRACE_LOG(CpuProfiler, EventBatchV2, Channel)
-			<< EventBatch.ThreadId(uint16(ThreadIdOverride))
-			<< EventBatch.Data(Buffer, uint32(ptrdiff_t(Cursor - Buffer)));
+			<< EventBatchV2.ThreadId(uint16(ThreadIdOverride))
+			<< EventBatchV2.Data(Buffer, uint32(ptrdiff_t(Cursor - Buffer)));
 
 		// Both protocols should really do this rebase but it make analysis go
 		// bonkers and I'm not looking into that right now
@@ -355,7 +387,11 @@ FLane::~FLane()
 ////////////////////////////////////////////////////////////////////////////////
 uint32 FLane::NewScope(const FAnsiStringView& Name)
 {
+#if UE_LANETRACE_ENABLED != 2
 	return FCpuProfilerTrace::OutputEventType(Name.GetData(), "", 0u);
+#else
+	return uint32(ScopeNew(Name));
+#endif
 }
 
 ////////////////////////////////////////////////////////////////////////////////
@@ -465,8 +501,8 @@ public:
 private:
 	struct FEntry
 	{
-		UPTRINT			Postcode;
-		FLaneTrace*		Lane;
+		UPTRINT			Postcode = 0;
+		FLaneTrace*		Lane = nullptr;
 	};
 
 	enum { GROWTH_SIZE = 4 };
@@ -520,7 +556,7 @@ FLaneTrace* FEstate::Build(UPTRINT Postcode)
 	}
 
 	int32 NextSize = Directory.Num() + GROWTH_SIZE;
-	Directory.SetNumZeroed(NextSize);
+	Directory.SetNum(NextSize);
 	return UseEstate(Postcode, Directory[NextSize - GROWTH_SIZE]);
 }
 
@@ -602,4 +638,9 @@ void LaneEstate_Demolish(FLaneEstate* Estate, FLanePostcode Postcode)
 }
 
 #undef LANETRACE_UNTESTED
+
+#if UE_LANETRACE_ENABLED == 2
+} // namespace trace
+#endif
+
 #endif // UE_LANETRACE_ENABLED
