Dataflow - Remove ensure on non existing nodes outside of regular editor operations (e.g. cooking).

#rnx

#Synced-CL 51292854

diff --git a/Engine/Source/Runtime/Experimental/Dataflow/Core/Private/Dataflow/DataflowGraph.cpp b/Engine/Source/Runtime/Experimental/Dataflow/Core/Private/Dataflow/DataflowGraph.cpp
--- a/Engine/Source/Runtime/Experimental/Dataflow/Core/Private/Dataflow/DataflowGraph.cpp
+++ b/Engine/Source/Runtime/Experimental/Dataflow/Core/Private/Dataflow/DataflowGraph.cpp
@@ -740,9 +740,18 @@ namespace UE::Dataflow
 			DATAFLOW_OPTIONAL_BLOCK_READ_ELSE()
 			{
 				InGraph->DisabledNodes.Add(ArNodeName);
-				ensureMsgf(false,
-					TEXT("Error: Missing registered node type (%s) will be removed from graph on load. Graph will fail to evaluate due to missing node (%s).")
-					, *ArType.ToString(), *ArName.ToString());
+				if (!IsRunningCommandlet() && GIsEditor)  // Don't ensure during cooking or during non editor operations
+				{
+					ensureMsgf(false,
+						TEXT("Error: Missing registered node type (%s) will be removed from graph on load. Graph will fail to evaluate due to missing node (%s).")
+						, *ArType.ToString(), *ArNodeName.ToString());
+				}
+				else
+				{
+					UE_LOGF(LogChaos, Display, 
+						"Error: Missing registered node type (%ls) will be removed from graph on load. Graph will fail to evaluate due to missing node (%ls)."
+						, *ArType.ToString(), *ArNodeName.ToString());
+				}
 			}
 			DATAFLOW_OPTIONAL_BLOCK_READ_END();
 		}
