Refactor IAD cache to support multiple CAS

#Synced-CL 51256130

diff --git a/Engine/Source/Runtime/Experimental/IndexedCacheStorage/Private/IndexedCacheStorageManager.cpp b/Engine/Source/Runtime/Experimental/IndexedCacheStorage/Private/IndexedCacheStorageManager.cpp
--- a/Engine/Source/Runtime/Experimental/IndexedCacheStorage/Private/IndexedCacheStorageManager.cpp
+++ b/Engine/Source/Runtime/Experimental/IndexedCacheStorage/Private/IndexedCacheStorageManager.cpp
@@ -577,6 +577,7 @@ void FIndexedCacheStorageManager::Initialize()
 					FString IniKey;
 					if (GetIniParameters(IniName, IniSection, IniKey, EarlyStartupIniPath))
 					{
+						// Fetch the value directly from EarlyStartupIniPath
 						FString EarlyStartupSizeStr = GConfig->GetStr(*IniSection, *IniKey, *IniName);
 						EarlyStartupSize = ParseSizeParam(EarlyStartupSizeStr);
 						if (EarlyStartupSize == 0)
@@ -585,7 +586,36 @@ void FIndexedCacheStorageManager::Initialize()
 								   TEXT("%s%s specified, but no %s:%s exists in the %s ini file"), PropertyEarlyStartup, *EarlyStartupIniPath, *IniSection, *IniKey, *IniName);
 						}
 					}
-				}			
+					else if(TArray<FString> EarlyStartupIniPathArray; 
+							GConfig->GetArray(TEXT("IndexedCacheStorage"), *EarlyStartupIniPath, EarlyStartupIniPathArray, GEngineIni))
+					{
+						// EarlyStartupIniPath is the key for an array, sum all the sections in the array
+						for (const FString& IniPath : EarlyStartupIniPathArray)
+						{
+							if (GetIniParameters(IniName, IniSection, IniKey, IniPath))
+							{
+								FString EarlyStartupSizeStr = GConfig->GetStr(*IniSection, *IniKey, *IniName);
+								const uint64 ParsedSize = ParseSizeParam(EarlyStartupSizeStr);
+								if (ParsedSize == 0)
+								{
+									UE_LOG(LogIndexedCacheStorageManager, Error,
+										TEXT("%s%s specified, but no %s:%s exists in the %s ini file"), PropertyEarlyStartup, *IniPath, *IniSection, *IniKey, *IniName);
+								}
+								EarlyStartupSize += ParsedSize;
+							}
+							else
+							{
+								const uint64 ParsedSize = ParseSizeParam(IniPath);
+								if (ParsedSize == 0)
+								{
+									UE_LOG(LogIndexedCacheStorageManager, Error,
+										TEXT("%s%s specified, but no does not specify a usable value"), PropertyEarlyStartup, *IniPath);
+								}
+								EarlyStartupSize += ParsedSize;
+							}
+						}
+					}
+				}
 			}
 
 			if (IndexedCacheEntriesMetaData[StorageIndex] == nullptr)
diff --git a/Engine/Source/Runtime/Experimental/IoStore/OnDemand/Private/OnDemandConfig.cpp b/Engine/Source/Runtime/Experimental/IoStore/OnDemand/Private/OnDemandConfig.cpp
--- a/Engine/Source/Runtime/Experimental/IoStore/OnDemand/Private/OnDemandConfig.cpp
+++ b/Engine/Source/Runtime/Experimental/IoStore/OnDemand/Private/OnDemandConfig.cpp
@@ -403,31 +403,33 @@ TIoStatusOr<FOnDemandInstallCacheConfig> TryParseInstallCacheConfig(const TCHAR*
 
 	FOnDemandInstallCacheConfig OutConfig;
 
-	const int64 DiskQuota = [CommandLine]
+	OutConfig.RootDirectory = Config::GetInstallCacheDirectory(CommandLine);
+
+	const FString ICSName = [CommandLine]
 	{
-		if (FString ParamValue; FParse::Value(CommandLine, TEXT("-Iad.FileCache.DiskQuota="), ParamValue))
+		if (FString ParamValue; FParse::Value(CommandLine, TEXT("-Iad.IndexedCacheStorage.IndexName="), ParamValue))
 		{
-			return Config::ParseSizeParam(ParamValue);
+			return ParamValue;
 		}
-		else if (FString ValueStr; GConfig->GetString(TEXT("OnDemandInstall"), TEXT("FileCache.DiskQuota"), ValueStr, GEngineIni))
+		else if (FString ValueStr; GConfig->GetString(TEXT("OnDemandInstall"), TEXT("IndexedCacheStorage.IndexName"), ValueStr, GEngineIni))
 		{
-			return Config::ParseSizeParam(ValueStr);
+			return ValueStr;
 		}
-		return int64(-1);
+		return FString();
 	}();
 
-	if (DiskQuota > 0)
+	if (ICSName.IsEmpty() == false)
 	{
-		OutConfig.DiskQuota = DiskQuota;
+		OutConfig.IndexedCacheStorageName = ICSName;
 	}
 
 	const int64 JournalMaxSize = [CommandLine]
 	{
-		if (FString ParamValue; FParse::Value(CommandLine, TEXT("-Iad.FileCache.JournalMaxSize="), ParamValue))
+		if (FString ParamValue; FParse::Value(CommandLine, TEXT("-Iad.Journal.MaxSize="), ParamValue))
 		{
 			return Config::ParseSizeParam(ParamValue);
 		}
-		else if (FString ValueStr; GConfig->GetString(TEXT("OnDemandInstall"), TEXT("FileCache.JournalMaxSize"), ValueStr, GEngineIni))
+		else if (FString ValueStr; GConfig->GetString(TEXT("OnDemandInstall"), TEXT("Journal.MaxSize"), ValueStr, GEngineIni))
 		{
 			return Config::ParseSizeParam(ValueStr);
 		}
@@ -439,64 +441,121 @@ TIoStatusOr<FOnDemandInstallCacheConfig> TryParseInstallCacheConfig(const TCHAR*
 		OutConfig.JournalMaxSize = JournalMaxSize;
 	}
 
-	const int64 MinBlockSize = [CommandLine]
+#if !UE_BUILD_SHIPPING
+	OutConfig.bDropCache = FParse::Param(CommandLine, TEXT("Iad.DropCache"));
+#endif
+
+	const TCHAR* ConfigPrefixByType[EOnDemandInstallCasType::Count];
+	ConfigPrefixByType[EOnDemandInstallCasType::General] = TEXT("FileCache");
+	ConfigPrefixByType[EOnDemandInstallCasType::MMap] = TEXT("MMapCache");
+	
+	auto TryParseCasConfig = [&ConfigPrefixByType, CommandLine](FOnDemandInstallCasConfig& OutConfig, EOnDemandInstallCasType CasType)
 	{
-		if (FString ParamValue; FParse::Value(CommandLine, TEXT("-Iad.FileCache.MinBlockSize="), ParamValue))
+		const TCHAR* ConfigPrefix = ConfigPrefixByType[CasType];
+
+		const int64 DiskQuota = [ConfigPrefix, CommandLine]
 		{
-			return Config::ParseSizeParam(ParamValue);
-		}
-		else if (FString ValueStr; GConfig->GetString(TEXT("OnDemandInstall"), TEXT("FileCache.MinBlockSize"), ValueStr, GEngineIni))
+			const FString Param = FString::Printf(TEXT("-Iad.%s.DiskQuota="), ConfigPrefix);
+			const FString ConfigKey = FString::Printf(TEXT("%s.DiskQuota"), ConfigPrefix);
+
+			if (FString ParamValue; FParse::Value(CommandLine, *Param, ParamValue))
+			{
+				return Config::ParseSizeParam(ParamValue);
+			}
+			else if (FString ValueStr; GConfig->GetString(TEXT("OnDemandInstall"), *ConfigKey, ValueStr, GEngineIni))
+			{
+				return Config::ParseSizeParam(ValueStr);
+			}
+			return int64(-1);
+		}();
+
+		if (DiskQuota > 0)
 		{
-			return Config::ParseSizeParam(ValueStr);
+			OutConfig.DiskQuota = DiskQuota;
 		}
-		return int64(-1);
-	}();
 
-	if (MinBlockSize > 0)
-	{
-		OutConfig.MinBlockSize = MinBlockSize;
-	}
+		const int64 MinBlockSize = [ConfigPrefix, CommandLine]
+		{
+			const FString Param = FString::Printf(TEXT("-Iad.%s.MinBlockSize="), ConfigPrefix);
+			const FString ConfigKey = FString::Printf(TEXT("%s.MinBlockSize"), ConfigPrefix);
 
-	const int64 MaxBlockSize = [CommandLine]
-	{
-		if (FString ParamValue; FParse::Value(CommandLine, TEXT("-Iad.FileCache.MaxBlockSize="), ParamValue))
+			if (FString ParamValue; FParse::Value(CommandLine, *Param, ParamValue))
+			{
+				return Config::ParseSizeParam(ParamValue);
+			}
+			else if (FString ValueStr; GConfig->GetString(TEXT("OnDemandInstall"), *ConfigKey, ValueStr, GEngineIni))
+			{
+				return Config::ParseSizeParam(ValueStr);
+			}
+			return int64(-1);
+		}();
+
+		if (MinBlockSize > 0)
 		{
-			return Config::ParseSizeParam(ParamValue);
+			OutConfig.MinBlockSize = MinBlockSize;
 		}
-		else if (FString ValueStr; GConfig->GetString(TEXT("OnDemandInstall"), TEXT("FileCache.MaxBlockSize"), ValueStr, GEngineIni))
+
+		const int64 MaxBlockSize = [ConfigPrefix, CommandLine]
 		{
-			return Config::ParseSizeParam(ValueStr);
-		}
-		return int64(-1);
-	}();
+			const FString Param = FString::Printf(TEXT("-Iad.%s.MaxBlockSize="), ConfigPrefix);
+			const FString ConfigKey = FString::Printf(TEXT("%s.MaxBlockSize"), ConfigPrefix);
 
-	if (MaxBlockSize > 0)
-	{
-		OutConfig.MaxBlockSize = MaxBlockSize;
-	}
+			if (FString ParamValue; FParse::Value(CommandLine, *Param, ParamValue))
+			{
+				return Config::ParseSizeParam(ParamValue);
+			}
+			else if (FString ValueStr; GConfig->GetString(TEXT("OnDemandInstall"), *ConfigKey, ValueStr, GEngineIni))
+			{
+				return Config::ParseSizeParam(ValueStr);
+			}
+			return int64(-1);
+		}();
 
-	const double LastAccessGranularitySeconds = [CommandLine]
-	{
-		if (double ParamValue; FParse::Value(CommandLine, TEXT("-Iad.FileCache.LastAccessGranularitySeconds="), ParamValue))
+		if (MaxBlockSize > 0)
 		{
-			return ParamValue;
+			OutConfig.MaxBlockSize = MaxBlockSize;
 		}
-		else if (GConfig->GetDouble(TEXT("OnDemandInstall"), TEXT("FileCache.LastAccessGranularitySeconds"), ParamValue, GEngineIni))
+
+		const double LastAccessGranularitySeconds = [ConfigPrefix, CommandLine]
 		{
-			return ParamValue;
+			const FString Param = FString::Printf(TEXT("-Iad.%s.LastAccessGranularitySeconds="), ConfigPrefix);
+			const FString ConfigKey = FString::Printf(TEXT("%s.LastAccessGranularitySeconds"), ConfigPrefix);
+
+			if (double ParamValue; FParse::Value(CommandLine, *Param, ParamValue))
+			{
+				return ParamValue;
+			}
+			else if (GConfig->GetDouble(TEXT("OnDemandInstall"), *ConfigKey, ParamValue, GEngineIni))
+			{
+				return ParamValue;
+			}
+			return double(-1.0);
+		}();
+
+		if (LastAccessGranularitySeconds >= 0)
+		{
+			OutConfig.LastAccessGranularitySeconds = LastAccessGranularitySeconds;
 		}
-		return double(-1.0);
-	}();
+	};
 
-	if (LastAccessGranularitySeconds >= 0)
+	// Always make the general CAS
+	FOnDemandInstallCasConfig& GeneralCasConfig = OutConfig.CasConfig.Emplace_GetRef();
+	GeneralCasConfig.Type = EOnDemandInstallCasType::General;
+	GeneralCasConfig.CasSubdirectory = TEXT("blocks");
+
+	TryParseCasConfig(GeneralCasConfig, GeneralCasConfig.Type);
+
+	// Conditionally make MMap CAS
+	bool bUseMMapCache = false;
+	GConfig->GetBool(TEXT("OnDemandInstall"), TEXT("MMapCache.Enabled"), bUseMMapCache, GEngineIni);
+	if (bUseMMapCache)
 	{
-		OutConfig.LastAccessGranularitySeconds = LastAccessGranularitySeconds;
-	}
+		FOnDemandInstallCasConfig& CasConfig = OutConfig.CasConfig.Emplace_GetRef();
+		CasConfig.Type = EOnDemandInstallCasType::MMap;
+		CasConfig.CasSubdirectory = TEXT("mmap");
 
-	OutConfig.RootDirectory = Config::GetInstallCacheDirectory(CommandLine);
-#if !UE_BUILD_SHIPPING
-	OutConfig.bDropCache = FParse::Param(CommandLine, TEXT("Iad.DropCache"));
-#endif
+		TryParseCasConfig(CasConfig, CasConfig.Type);
+	}
 
 	return OutConfig;
 }
diff --git a/Engine/Source/Runtime/Experimental/IoStore/OnDemand/Private/OnDemandContentInstaller.cpp b/Engine/Source/Runtime/Experimental/IoStore/OnDemand/Private/OnDemandContentInstaller.cpp
--- a/Engine/Source/Runtime/Experimental/IoStore/OnDemand/Private/OnDemandContentInstaller.cpp
+++ b/Engine/Source/Runtime/Experimental/IoStore/OnDemand/Private/OnDemandContentInstaller.cpp
@@ -3,6 +3,7 @@
 #include "OnDemandContentInstaller.h"
 
 #include "Algo/Accumulate.h"
+#include "Algo/AnyOf.h"
 #include "Async/UniqueLock.h"
 #include "Containers/RingBuffer.h"
 #include "Experimental/UnifiedError/CoreErrorTypes.h"
@@ -272,7 +273,7 @@ void ResolveChunksToInstall(
 		OnDemandPackageStore,
 		PackageIds,
 		bIncludeSoftReferences,
-		[&ResolvedPackageIds, &ShaderMapChunkIds, &OutMissing, &OnDemandPackageStore](
+		[&ResolvedPackageIds, &ShaderMapChunkIds, &OutMissing](
 			FPackageId PackageId,
 			EPackageStoreEntryStatus EntryStatus,
 			const FPackageStoreEntry& PackageStoreEntry,
@@ -726,7 +727,8 @@ void FOnDemandContentInstaller::TryExecuteNextRequest()
 
 					if (bExecuteFlush)
 					{
-						FResult Result = IoStore.InstallCache->Flush();
+						// TODO: IadMultiCAS: content installer needs to determine which chunks go to which CAS
+						FResult Result = IoStore.InstallCache->Flush(EOnDemandInstallCasType::General);
 						if (Result.HasError())
 						{
 							UE_LOG(LogIoStoreOnDemand, Error, TEXT("Failed to flush pending chunks, Reason %s"), *LexToString(Result.GetError()));
@@ -1099,7 +1101,8 @@ void FOnDemandContentInstaller::ExecuteInstallRequest(FRequest& Request)
 			}
 		}
 
-		if (Request.Result = IoStore.InstallCache->Purge(BytesToInstall); Request.Result.HasError())
+		// TODO: IadMultiCAS: content installer needs to determine which chunks go to which CAS
+		if (Request.Result = IoStore.InstallCache->Purge(EOnDemandInstallCasType::General, BytesToInstall); Request.Result.HasError())
 		{
 			return CompleteInstallRequest(Request);
 		}
@@ -1187,16 +1190,17 @@ void FOnDemandContentInstaller::ExecutePurgeRequest(FRequest& Request)
 	check(Request.IsPurge());
 	check(RunningRequests.Num() == 1 && RunningRequests.Contains(&Request));
 
-	FRequest::FPurge& PurgeRequest	= Request.AsPurge();
-	const bool bDefrag				= EnumHasAnyFlags(PurgeRequest.Args.Options, EOnDemandPurgeOptions::Defrag);
-	const uint64* BytesToPurge		= PurgeRequest.Args.BytesToPurge.GetPtrOrNull();
+	FRequest::FPurge& PurgeRequest			= Request.AsPurge();
+	const bool bDefrag						= EnumHasAnyFlags(PurgeRequest.Args.Options, EOnDemandPurgeOptions::Defrag);
+	const uint64* BytesToPurge				= PurgeRequest.Args.BytesToPurge.GetPtrOrNull();
+	const EOnDemandInstallCasType CasType	= PurgeRequest.Args.CasType;
 
-	UE_LOG(LogIoStoreOnDemand, Log, TEXT("Executing purge request, BytesToPurge=%llu, Defrag='%s'"),
-		BytesToPurge != nullptr ? *BytesToPurge : -1, bDefrag ? TEXT("True") : TEXT("False"));
+	UE_LOG(LogIoStoreOnDemand, Log, TEXT("Executing purge request, CasType=%s, BytesToPurge=%llu, Defrag='%s'"),
+		LexToString(CasType), BytesToPurge != nullptr ? *BytesToPurge : -1, bDefrag ? TEXT("True") : TEXT("False"));
 
 	if (Request.IsCancelled() == false)
 	{
-		Request.Result = IoStore.InstallCache->PurgeAllUnreferenced(bDefrag, BytesToPurge);
+		Request.Result = IoStore.InstallCache->PurgeAllUnreferenced(CasType, bDefrag, BytesToPurge);
 	}
 
 	CompletePurgeRequest(Request);
@@ -1207,15 +1211,16 @@ void FOnDemandContentInstaller::ExecuteDefragRequest(FRequest& Request)
 	check(Request.IsDefrag());
 	check(RunningRequests.Num() == 1 && RunningRequests.Contains(&Request));
 
-	FRequest::FDefrag& DefragRequest	= Request.AsDefrag();
-	const uint64* BytesToFree			= DefragRequest.Args.BytesToFree.GetPtrOrNull();
+	FRequest::FDefrag& DefragRequest		= Request.AsDefrag();
+	const uint64* BytesToFree				= DefragRequest.Args.BytesToFree.GetPtrOrNull();
+	const EOnDemandInstallCasType CasType	= DefragRequest.Args.CasType;
 
-	UE_LOG(LogIoStoreOnDemand, Log, TEXT("Executing defrag request, BytesToFree=%llu"),
-		BytesToFree != nullptr ? *BytesToFree : -1); 
+	UE_LOG(LogIoStoreOnDemand, Log, TEXT("Executing defrag request, CasType=%s, BytesToFree=%llu"),
+		LexToString(CasType), BytesToFree != nullptr ? *BytesToFree : -1);
 
 	if (Request.IsCancelled() == false)
 	{
-		Request.Result = IoStore.InstallCache->DefragAll(BytesToFree);
+		Request.Result = IoStore.InstallCache->DefragAll(CasType, BytesToFree);
 	}
 
 	CompleteDefragRequest(Request);
@@ -1269,14 +1274,17 @@ void FOnDemandContentInstaller::ExectuteFlushLastAccessRequest(FRequest& Request
 		}
 	}
 
-	bool bLastAccessDirty = false;
+	TStaticArray<bool, EOnDemandInstallCasType::Count> bLastAccessDirty(EInPlace::InPlace, false);
 	IoStore.InstallCache->UpdateLastAccess(TouchedChunks, bLastAccessDirty);
-	bLastAccessDirty = bLastAccessDirty || FlushLastAccessReqest.Args.bForceLastAccessDirty;
-
-	if (bLastAccessDirty)
+	for (bool& bCasLastAccessDirty : bLastAccessDirty)
+	{
+		bCasLastAccessDirty = bCasLastAccessDirty || FlushLastAccessReqest.Args.bForceLastAccessDirty;
+	}
+	
+	if (Algo::AnyOf(bLastAccessDirty))
 	{
 		UE_LOG(LogIoStoreOnDemand, Log, TEXT("Flushing last access to journal"));
-		Request.Result = IoStore.InstallCache->FlushLastAccess();
+		Request.Result = IoStore.InstallCache->FlushLastAccess(bLastAccessDirty);
 	}
 
 	CompleteFlushLastAccessRequest(Request);
@@ -1334,7 +1342,8 @@ void FOnDemandContentInstaller::ProcessDownloadedChunk(
 		{
 			TRACE_COUNTER_INCREMENT(OnDemandContentInstaller_HttpRequestsSuccess);
 
-			TResult<bool> PutResult = IoStore.InstallCache->PutChunk(MoveTemp(Chunk), ChunkHash);
+			// TODO: IadMultiCAS: content installer needs to determine which chunks go to which CAS
+			TResult<bool> PutResult = IoStore.InstallCache->PutChunk(EOnDemandInstallCasType::General, MoveTemp(Chunk), ChunkHash);
 			if (PutResult.HasError())
 			{
 				Result = MakeError(PutResult.StealError());
@@ -1368,10 +1377,11 @@ void FOnDemandContentInstaller::ProcessDownloadedChunk(
 						PendingDownloadSize += ChunkEntry.GetDiskSize();
 					}
 
-					Result = IoStore.InstallCache->Purge(PendingDownloadSize);
+					// TODO: IadMultiCAS: content installer needs to determine which chunks go to which CAS
+					Result = IoStore.InstallCache->Purge(EOnDemandInstallCasType::General, PendingDownloadSize);
 					if (Result.HasValue())
 					{
-						PutResult = IoStore.InstallCache->PutChunk(MoveTemp(Chunk), ChunkHash);
+						PutResult = IoStore.InstallCache->PutChunk(EOnDemandInstallCasType::General, MoveTemp(Chunk), ChunkHash);
 						if (PutResult.HasError())
 						{
 							Result = MakeError(PutResult.StealError());
diff --git a/Engine/Source/Runtime/Experimental/IoStore/OnDemand/Private/OnDemandInstallCache.cpp b/Engine/Source/Runtime/Experimental/IoStore/OnDemand/Private/OnDemandInstallCache.cpp
--- a/Engine/Source/Runtime/Experimental/IoStore/OnDemand/Private/OnDemandInstallCache.cpp
+++ b/Engine/Source/Runtime/Experimental/IoStore/OnDemand/Private/OnDemandInstallCache.cpp
@@ -267,29 +267,32 @@
 
 		struct FChunkLocation
 		{
-			EType			Type = EType::ChunkLocation;
-			uint8			Pad[3] = { 0 };
-			FCasLocation	CasLocation;
-			FCasAddr		CasAddr;
+			EType					Type = EType::ChunkLocation;
+			EOnDemandInstallCasType	CasType = EOnDemandInstallCasType::General;
+			ANSICHAR				Pad[2] = { 0 };
+			FCasLocation			CasLocation;
+			FCasAddr				CasAddr;
 		};
 		static_assert(sizeof(FChunkLocation) == 24);
 
 		struct FBlockOperation
 		{
-			EType		Type = EType::None;
-			uint8		Pad[3] = { 0 };
-			FCasBlockId	BlockId;
-			int64		UtcTicks = 0;
-			uint8		Pad1[8] = { 0 };
+			EType					Type = EType::None;
+			EOnDemandInstallCasType	CasType = EOnDemandInstallCasType::General;
+			uint8					Pad[2] = { 0 };
+			FCasBlockId				BlockId;
+			int64					UtcTicks = 0;
+			uint8					Pad1[8] = { 0 };
 		};
 		static_assert(sizeof(FBlockOperation) == 24);
 
 		struct FCriticalError
 		{
-			EType		Type = EType::CriticalError;
-			uint8		Pad[3] = { 0 };
-			EErrorCode	ErrorCode = EErrorCode::None;
-			uint8		Pad1[16] = { 0 };
+			EType					Type = EType::CriticalError;
+			EOnDemandInstallCasType	CasType = EOnDemandInstallCasType::General;
+			uint8					Pad[2] = { 0 };
+			EErrorCode				ErrorCode = EErrorCode::None;
+			uint8					Pad1[16] = { 0 };
 		};
 		static_assert(sizeof(FCriticalError) == 24);
 
@@ -301,6 +304,11 @@
 		};
 
 		EType			Type() const { return *reinterpret_cast<const EType*>(this); }
+		EOnDemandInstallCasType CasType() const
+		{
+			const uint8* pThis = reinterpret_cast<const uint8*>(this);
+			return *reinterpret_cast<const EOnDemandInstallCasType*>(pThis + sizeof(EType));
+		}
 		static int64	Size() { return sizeof(FEntry); }
 
 		FEntry() {}
@@ -316,16 +324,17 @@
 		void			BlockModification(FCasBlockId BlockId, int64 UtcTicks);
 		void			CriticalError(FCasJournal::EErrorCode ErrorCode);
 
-		FString			JournalFile;
-		TArray<FEntry>	Entries;
+		FString					JournalFile;
+		TArray<FEntry>			Entries;
+		EOnDemandInstallCasType	CasType = EOnDemandInstallCasType::General;
 	};
 
 	using FEntryHandler = TFunction<FResult(const FEntry&)>;
 
 	static FResult			Replay(const FString& JournalFile, FEntryHandler&& Handler);
 	static FResult			Create(const FString& JournalFile);
-	static FTransaction		Begin(FString&& JournalFile);
-	static FTransaction		Begin(const FString& JournalFile) { return Begin(FString(JournalFile)); }
+	static FTransaction		Begin(EOnDemandInstallCasType InCasType, FString&& JournalFile);
+	static FTransaction		Begin(EOnDemandInstallCasType InCasType, const FString& JournalFile) { return Begin(InCasType, FString(JournalFile)); }
 	static FResult			Commit(FTransaction&& Transaction);
 	static FResult			Commit(FTransaction&& Transaction, uint64& OutByteCount, uint32& OutOpCount);
 };
@@ -511,11 +520,12 @@
 	return MakeValue();
 }
 
-FCasJournal::FTransaction FCasJournal::Begin(FString&& JournalFile)
+FCasJournal::FTransaction FCasJournal::Begin(EOnDemandInstallCasType InCasType, FString&& JournalFile)
 {
 	return FTransaction
 	{
-		.JournalFile = MoveTemp(JournalFile)
+		.JournalFile = MoveTemp(JournalFile),
+		.CasType = InCasType
 	};
 }
 
@@ -664,6 +674,7 @@
 {
 	new(&Entries.AddDefaulted_GetRef().ChunkLocation) FEntry::FChunkLocation
 	{
+		.CasType = CasType,
 		.CasLocation = Location,
 		.CasAddr = Addr
 	};
@@ -674,6 +685,7 @@
 	new(&Entries.AddDefaulted_GetRef().BlockOperation) FEntry::FBlockOperation
 	{
 		.Type = FEntry::EType::BlockCreated,
+		.CasType = CasType,
 		.BlockId = BlockId,
 		.UtcTicks = FDateTime::UtcNow().GetTicks()
 	};
@@ -684,6 +696,7 @@
 	new(&Entries.AddDefaulted_GetRef().BlockOperation) FEntry::FBlockOperation
 	{
 		.Type = FEntry::EType::BlockDeleted,
+		.CasType = CasType,
 		.BlockId = BlockId,
 		.UtcTicks = FDateTime::UtcNow().GetTicks()
 	};
@@ -694,6 +707,7 @@
 	new(&Entries.AddDefaulted_GetRef().BlockOperation) FEntry::FBlockOperation
 	{
 		.Type = FEntry::EType::BlockAccess,
+		.CasType = CasType,
 		.BlockId = BlockId,
 		.UtcTicks = UtcTicks
 	};
@@ -704,6 +718,7 @@
 	new(&Entries.AddDefaulted_GetRef().BlockOperation) FEntry::FBlockOperation
 	{
 		.Type = FEntry::EType::BlockModification,
+		.CasType = CasType,
 		.BlockId = BlockId,
 		.UtcTicks = UtcTicks
 	};
@@ -714,10 +729,158 @@
 	new(&Entries.AddDefaulted_GetRef().CriticalError) FEntry::FCriticalError
 	{
 		.Type = FEntry::EType::CriticalError,
+		.CasType = CasType,
 		.ErrorCode = ErrorCode
 	};
 }
 
 ///////////////////////////////////////////////////////////////////////////////
+struct FCasSnapshot
+{
+	enum class EVersion : uint32
+	{
+		Invalid = 0,
+		Initial,
+
+		AddModificationTime,
+		FixSwitchBuild,
+		FixLostSnapshots,
+		AddCasType,
+
+		LatestPlusOne,
+		Latest = LatestPlusOne - 1
+	};
+
+	static const FGuid CustomVersionGuid;
+
+	struct FHeader
+	{
+		static const inline uint8 MagicSequence[16] = {'+', 'S', 'N', 'A', 'P', 'S', 'H', 'O', 'T', 'H', 'E', 'A', 'D', 'E', 'R', '+'};
+
+		bool			IsValid() const;
+		static int64	Size() { return sizeof(FHeader); }
+
+		uint8		Magic[16] = {0};
+		EVersion	Version = EVersion::Invalid;
+		uint8		Pad[12] = {0};
+	};
+	static_assert(sizeof(FHeader) == 32);
+
+	struct FFooter
+	{
+		static const inline uint8 MagicSequence[16] = {'+', 'S', 'N', 'A', 'P', 'S', 'H', 'O', 'T', 'F', 'O', 'O', 'T', 'E', 'R', '+'};
+
+		static int64	Size() { return sizeof(FFooter); }
+		bool			IsValid() const;
+
+		uint8 Magic[16] = {0};
+	};
+	static_assert(sizeof(FFooter) == 16);
+
+	struct FBlock
+	{
+		friend FArchive& operator<<(FArchive& Ar, FBlock& Block)
+		{
+			Ar << Block.BlockId;
+			Ar << Block.LastAccess;
+
+			if (Ar.CustomVer(FCasSnapshot::CustomVersionGuid) >= static_cast<int32>(EVersion::AddModificationTime))
+			{
+				Ar << Block.ModTime;
+			}
+
+			return Ar;
+		}
+
+		FCasBlockId BlockId;
+		int64		LastAccess = 0;
+		int64		ModTime = 0;
+	};
+
+	using FChunkLocation = TPair<FCasAddr, FCasLocation>;
+
+	struct FCasState
+	{
+		friend FArchive& operator<<(FArchive& Ar, FCasState& InCasState)
+		{
+			Ar << InCasState.Blocks;
+			Ar << InCasState.ChunkLocations;
+			Ar << InCasState.CurrentBlockId;
+
+			if (Ar.CustomVer(FCasSnapshot::CustomVersionGuid) >= static_cast<int32>(EVersion::AddCasType))
+			{
+				using UnderType = std::underlying_type<EOnDemandInstallCasType>::type;
+				UnderType Val = static_cast<UnderType>(InCasState.CasType);
+				if (Ar.IsLoading())
+				{
+					Ar << Val;
+					InCasState.CasType = static_cast<EOnDemandInstallCasType>(Val);
+				}
+				else
+				{
+					Ar << Val;
+				}
+			}
+
+			return Ar;
+		}
+
+		TArray<FBlock>						Blocks;
+		TArray<FChunkLocation>				ChunkLocations;
+		FCasBlockId							CurrentBlockId;
+		EOnDemandInstallCasType				CasType{ EOnDemandInstallCasType::None };
+	};
+
+	static TResult<FCasSnapshot>		FromJournal(const FString& JournalFile, FCasSnapshot&& BaseSnapshot);
+	static TResult<FCasSnapshot>		Load(const FString& SnapshotFile, int64* OutFileSize = nullptr);
+	static TResult<int64>				Save(const FCasSnapshot& Snapshot, const FString& SnapshotFile);
+	/** 
+	 * Append the journal file to the snapshot and reset the journal file.
+	 *
+	 * @param SnapshotFile	The snapshot file to appened the journal to. May not exist yet.
+	 * @param JournalFile	The journal file to append to the snapshot, and then reset.
+	 *
+	 * Note: If this operation fails, it will leave behind temporary files, call TempFilesExist() check for them.
+	 * There is no provided method to cleanup temporary files. If they are discovered, the
+	 * cache is in an inconsistent state and should be reset.
+	 */
+	static TResult<int64>				AppendAndResetJournal(
+											const FString& SnapshotFile, 
+											const FString& JournalFile);
+
+	/**
+	 * There is no provided method to cleanup temporary files. If they are discovered, the
+	 * cacche is in an inconsistent state and should be reset.
+	 */
+	static bool TempFilesExist(const FString& SnapshotFile, const FString& JournalFile);
+
+	static FString MakeTempFileName(const FString& FileName)
+	{
+		const bool bIncludeDot = true;
+		const FString Extension = FPaths::GetExtension(FileName, bIncludeDot) + TEXT("tmp");
+		return FPaths::ChangeExtension(FileName, Extension);
+	}
+
+	TArray<FCasState> CasState;
+
+private:
+	FCasState& FindOrAddCasState(EOnDemandInstallCasType CasType)
+	{
+		FCasState* Found = Algo::FindBy(CasState, CasType, &FCasState::CasType);
+		if (Found)
+		{
+			return *Found;
+		}
+
+		CasState.Emplace_GetRef().CasType = CasType;
+
+		Algo::SortBy(CasState, &FCasState::CasType);
+
+		Found = Algo::FindBy(CasState, CasType, &FCasState::CasType);
+		return *Found;
+	}
+};
+
+///////////////////////////////////////////////////////////////////////////////
 struct FCasBlockState
 {
@@ -727,7 +890,6 @@
 	int64 LastModificationUtcTicks = 0;
 };
 
-struct FCasSnapshot;
 struct FCas
 {
 	static constexpr uint32		DeleteBlockMaxWaitTimeMs = 10000;
@@ -744,7 +906,7 @@
 	void						Lock()		{ Mutex.Lock(); }
 	void						Unlock()	{ Mutex.Unlock(); }
 
-	FResult						Configure(const FOnDemandInstallCacheConfig& Config);
+	FResult						Configure(const FOnDemandInstallCasConfig& Config);
 	FResult						Initialize(FStringView Directory, bool bDeleteExisting = false);
 	FCasLocation				FindChunk(const FIoHash& Hash) const;
 	FCasBlockId					CreateBlock();
@@ -778,7 +940,7 @@
 	FBlockFileSizes				FindAllBlockFiles() const;
 	void						Compact();
 	FResult						Verify(TArray<FCasAddr>& OutRemovedChunksAddrs, TMap<FCasBlockId, uint64>& OutOverBudgetBlocks);
-	void						LoadSnapshot(FCasSnapshot&& Snapshot);
+	void						LoadSnapshot(FCasSnapshot::FCasState&& Snapshot);
 	FResult						ApplyJournalEntry(const FCasJournal::FEntry& JournalEntry);
 
 	// Returns the timestamps that are "dirty" and need to be flushed to disk and clears dirty flag
@@ -799,10 +961,10 @@
 		return MaxBlockSize;
 	}
 
-	FString						GetRootDirectory() const { return FString(RootDirectory); }
+	const FString&				GetRootDirectory() const { return RootDirectory; }
 
 private:
-	FStringView				RootDirectory;
+	FString					RootDirectory;
 public: // TODO: FIXME: encapsulate these members
 	FLookup					Lookup;
 private:
@@ -815,12 +977,401 @@
 	mutable UE::FMutex		Mutex;
 };
 
+const FGuid FCasSnapshot::CustomVersionGuid(0x217bee7d, 0x23a54a23, 0x9ef080d2, 0x3df768b1);
+
 ///////////////////////////////////////////////////////////////////////////////
+bool FCasSnapshot::FHeader::IsValid() const
+{
+	if (FMemory::Memcmp(&Magic, &FHeader::MagicSequence, sizeof(FHeader::MagicSequence)) != 0)
+	{
+		return false;
+	}
+
+	if (static_cast<uint32>(Version) > static_cast<uint32>(EVersion::Latest))
+	{
+		return false;
+	}
+
+	// Drop all snapshots from before we fixed snapshots not being propogated
+	if (static_cast<uint32>(Version) < static_cast<uint32>(EVersion::FixLostSnapshots))
+	{
+		return false;
+	}
+
+	// Versions that can no longer be loaded. The cache must be dropped instead.
+	const EVersion BadVersions[] = { UE_ONDEMANDINSTALLCACHE_BAD_SNAPSHOT_VERSIONS };
+	for (EVersion BadVersion : BadVersions)
+	{
+		if (Version == BadVersion)
+		{
+			return false;
+		}
+	}
+
+	return true;
+}
+
+bool FCasSnapshot::FFooter::IsValid() const
+{
+	return FMemory::Memcmp(Magic, FFooter::MagicSequence, sizeof(FFooter::MagicSequence)) == 0;
+}
+
+TResult<FCasSnapshot> FCasSnapshot::FromJournal(const FString& JournalFile, FCasSnapshot&& BaseSnapshot)
+{
+	TOptional<FCas::FLookup>				CasLookupByType[EOnDemandInstallCasType::Count];
+	TOptional<TMap<FCasBlockId, FBlock>>	BlockMapByType[EOnDemandInstallCasType::Count];
+	TOptional<FCasBlockId>					CurrentBlockIdByType[EOnDemandInstallCasType::Count];
+
+	for (FCasState& BaseState : BaseSnapshot.CasState)
+	{
+		FCas::FLookup&				CasLookup = CasLookupByType[BaseState.CasType].Emplace();
+		TMap<FCasBlockId, FBlock>&	BlockMap = BlockMapByType[BaseState.CasType].Emplace();
+		FCasBlockId&				CurrentBlockId = CurrentBlockIdByType[BaseState.CasType].Emplace();
+
+		CasLookup.Reserve(BaseState.ChunkLocations.Num());
+		for (TPair<FCasAddr, FCasLocation>& Kv : BaseState.ChunkLocations)
+		{
+			CasLookup.Add(MoveTemp(Kv));
+		}
+
+		BlockMap.Reserve(BaseState.Blocks.Num());
+		for (const FCasSnapshot::FBlock& Block : BaseState.Blocks)
+		{
+			BlockMap.Add(Block.BlockId, Block);
+		}
+
+		CurrentBlockId = BaseState.CurrentBlockId;
+	}
+
+	FResult ReplayResult = FCasJournal::Replay(
+		JournalFile,
+		[&CasLookupByType, &BlockMapByType, &CurrentBlockIdByType](const FCasJournal::FEntry& JournalEntry) -> FResult
+		{
+			EOnDemandInstallCasType CasType = JournalEntry.CasType();
+			FCas::FLookup&				CasLookup = CasLookupByType[CasType].IsSet() ? 
+										CasLookupByType[CasType].GetValue() : 
+										CasLookupByType[CasType].Emplace();
+			TMap<FCasBlockId, FBlock>&	BlockMap = BlockMapByType[CasType].IsSet() ? 
+										BlockMapByType[CasType].GetValue() : 
+										BlockMapByType[CasType].Emplace();
+			FCasBlockId&				CurrentBlockId = CurrentBlockIdByType[CasType].IsSet() ?
+										CurrentBlockIdByType[CasType].GetValue() : 
+										CurrentBlockIdByType[CasType].Emplace();
+
+			switch(JournalEntry.Type())
+			{
+			case FCasJournal::FEntry::EType::ChunkLocation:
+			{
+				const FCasJournal::FEntry::FChunkLocation& ChunkLocation = JournalEntry.ChunkLocation;
+				if (ChunkLocation.CasLocation.IsValid())
+				{
+					FCasLocation& Loc = CasLookup.FindOrAdd(ChunkLocation.CasAddr);
+					Loc = ChunkLocation.CasLocation;
+				}
+				else
+				{
+					CasLookup.Remove(ChunkLocation.CasAddr);
+				}
+				break;
+			}
+			case FCasJournal::FEntry::EType::BlockCreated:
+			{
+				const FCasJournal::FEntry::FBlockOperation& Op = JournalEntry.BlockOperation;
+				CurrentBlockId = Op.BlockId;
+				BlockMap.Emplace(Op.BlockId, FBlock{ .BlockId = Op.BlockId });
+				break;
+			}
+			case FCasJournal::FEntry::EType::BlockDeleted:
+			{
+				const FCasJournal::FEntry::FBlockOperation& Op = JournalEntry.BlockOperation;
+				BlockMap.Remove(Op.BlockId);
+				if (CurrentBlockId == Op.BlockId)
+				{
+					CurrentBlockId = FCasBlockId::Invalid; 
+				}
+				break;
+			}
+			case FCasJournal::FEntry::EType::BlockAccess:
+			{
+				const FCasJournal::FEntry::FBlockOperation& Op = JournalEntry.BlockOperation;
+				if (FBlock* Block = BlockMap.Find(Op.BlockId))
+				{
+					Block->LastAccess = Op.UtcTicks;
+				}
+				else
+				{
+					using namespace UE::UnifiedError::IoStoreOnDemand;
+					return MakeJournalError(ECasErrorCode::ReplayJournalFailed, EIoErrorCode::FileCorrupt, TEXT("Journal BlockAccess Op references an invalid block"));
+				}
+				break;
+			}
+			case FCasJournal::FEntry::EType::BlockModification:
+			{
+				const FCasJournal::FEntry::FBlockOperation& Op = JournalEntry.BlockOperation;
+				if (FBlock* Block = BlockMap.Find(Op.BlockId))
+				{
+					Block->ModTime = Op.UtcTicks;
+				}
+				else
+				{
+					using namespace UE::UnifiedError::IoStoreOnDemand;
+					return MakeJournalError(ECasErrorCode::ReplayJournalFailed, EIoErrorCode::FileCorrupt, TEXT("Journal BlockModification Op references an invalid block"));
+				}
+				break;
+			}
+			};
+
+			return MakeValue();
+		});
+
+	if (ReplayResult.HasError())
+	{
+		return MakeError(ReplayResult.StealError());
+	}
+
+	FCasSnapshot Snapshot;
+
+	for (EOnDemandInstallCasType CasType : TEnumRange<EOnDemandInstallCasType>())
+	{
+		if (BlockMapByType[CasType].IsSet())
+		{
+			BlockMapByType[CasType].GetValue().GenerateValueArray(Snapshot.FindOrAddCasState(CasType).Blocks);
+		}
+
+		if (CasLookupByType[CasType].IsSet())
+		{
+			Snapshot.FindOrAddCasState(CasType).ChunkLocations = CasLookupByType[CasType].GetValue().Array();
+		}
+
+		if (CurrentBlockIdByType[CasType].IsSet())
+		{
+			Snapshot.FindOrAddCasState(CasType).CurrentBlockId = CurrentBlockIdByType[CasType].GetValue();
+		}
+	}
+
+	return MakeValue(Snapshot);
+}
+
+TResult<int64> FCasSnapshot::Save(const FCasSnapshot& Snapshot, const FString& SnapshotFile)
+{
+	using namespace UE::UnifiedError::IoStoreOnDemand;
+
+	IFileManager& Ifm = IFileManager::Get();
+
+	TUniquePtr<FArchive> Ar(Ifm.CreateFileWriter(*SnapshotFile));
+	if (Ar.IsValid() == false)
+	{
+		return MakeSnapshotError<int64>(ECasErrorCode::SaveSnapshotFailed, EIoErrorCode::FileOpenFailed, FString::Printf(TEXT("Failed to open file '%s' for writing"), *SnapshotFile));
+	}
+
+	FHeader Header;
+	FMemory::Memcpy(&Header.Magic, &FHeader::MagicSequence, sizeof(FHeader::MagicSequence));
+	Header.Version = EVersion::Latest;
+
+	Ar->SetCustomVersion(CustomVersionGuid, static_cast<int32>(Header.Version), "CasSnapshotVersion");
+
+	Ar->Serialize(reinterpret_cast<uint8*>(&Header), FHeader::Size());
+	if (Ar->IsError())
+	{
+		const uint32 LastError = FPlatformMisc::GetLastError();
+		Ar.Reset();
+		return MakeSnapshotError<int64>(ECasErrorCode::SaveSnapshotFailed, EIoErrorCode::WriteError, FString::Printf(TEXT("Failed to write snapshot header to '%s'"), *SnapshotFile), LastError);
+	}
+
+	FCasSnapshot& NonConst = *const_cast<FCasSnapshot*>(&Snapshot);
+	*Ar << NonConst.CasState;
+
+	if (Ar->IsError())
+	{
+		const uint32 LastError = FPlatformMisc::GetLastError();
+		Ar.Reset();
+		return MakeSnapshotError<int64>(ECasErrorCode::SaveSnapshotFailed, EIoErrorCode::WriteError, FString::Printf(TEXT("Failed to write snapshot to '%s'"), *SnapshotFile), LastError);
+	}
+
+	FFooter Footer;
+	FMemory::Memcpy(&Footer.Magic, &FFooter::MagicSequence, sizeof(FFooter::MagicSequence));
+	Ar->Serialize(reinterpret_cast<uint8*>(&Footer), FFooter::Size());
+	if (Ar->IsError())
+	{
+		const uint32 LastError = FPlatformMisc::GetLastError();
+		Ar.Reset();
+		return MakeSnapshotError<int64>(ECasErrorCode::SaveSnapshotFailed, EIoErrorCode::WriteError, FString::Printf(TEXT("Failed to write snapshot footer to '%s'"), *SnapshotFile), LastError);
+	}
+
+	const int64 FileSize = Ar->TotalSize();
+	if (Ar->Close() == false)
+	{
+		const uint32 LastError = FPlatformMisc::GetLastError();
+		Ar.Reset();
+		return MakeSnapshotError<int64>(ECasErrorCode::SaveSnapshotFailed, EIoErrorCode::WriteError, FString::Printf(TEXT("Failed to close snapshot footer to '%s'"), *SnapshotFile), LastError);
+	}
+
+	return MakeValue(FileSize);
+}
+
+TResult<FCasSnapshot> FCasSnapshot::Load(const FString& SnapshotFile, int64* OutFileSize)
+{
+	using namespace UE::UnifiedError::IoStoreOnDemand;
+
+	IFileManager& Ifm = IFileManager::Get();
+
+	if (Ifm.FileExists(*SnapshotFile) == false)
+	{
+		return MakeSnapshotError<FCasSnapshot>(ECasErrorCode::LoadSnapshotFailed, EIoErrorCode::NotFound, FString::Printf(TEXT("Failed to find snapshot '%s'"), *SnapshotFile));
+	}
+
+	TUniquePtr<FArchive> Ar(Ifm.CreateFileReader(*SnapshotFile));
+	if (Ar.IsValid() == false)
+	{
+		return MakeSnapshotError<FCasSnapshot>(ECasErrorCode::LoadSnapshotFailed, EIoErrorCode::FileOpenFailed, FString::Printf(TEXT("Failed open snapshot '%s'"), *SnapshotFile));
+	}
+
+	FHeader Header;
+	Ar->Serialize(reinterpret_cast<uint8*>(&Header), FHeader::Size());
+	if (Ar->IsError() || Header.IsValid() == false)
+	{
+		return MakeSnapshotError<FCasSnapshot>(ECasErrorCode::LoadSnapshotFailed, EIoErrorCode::ReadError, FString::Printf(TEXT("Failed to validate snapshot header in '%s'"), *SnapshotFile));
+	}
+
+	Ar->SetCustomVersion(CustomVersionGuid, static_cast<int32>(Header.Version), "CasSnapshotVersion");
+
+	FCasSnapshot Snapshot;
+
+	if (Ar->CustomVer(FCasSnapshot::CustomVersionGuid) < static_cast<int32>(EVersion::AddCasType))
+	{
+		*Ar << Snapshot.FindOrAddCasState(EOnDemandInstallCasType::General);
+	}
+	else
+	{
+		*Ar << Snapshot.CasState;
+	}
+
+	FFooter Footer;
+	Ar->Serialize(reinterpret_cast<uint8*>(&Footer), FFooter::Size());
+	if (Ar->IsError() || Footer.IsValid() == false)
+	{
+		return MakeSnapshotError<FCasSnapshot>(ECasErrorCode::LoadSnapshotFailed, EIoErrorCode::ReadError, FString::Printf(TEXT("Failed to validate snapshot footer in '%s'"), *SnapshotFile));
+	}
+
+	if (OutFileSize != nullptr)
+	{
+		*OutFileSize = Ar->Tell();
+	}
+
+	return MakeValue(MoveTemp(Snapshot));
+}
+
+TResult<int64> FCasSnapshot::AppendAndResetJournal(const FString& SnapshotFile, const FString& JournalFile)
+{
+	using namespace UE::UnifiedError::IoStoreOnDemand;
+
+	IFileManager& Ifm = IFileManager::Get();
+
+	if (Ifm.FileExists(*JournalFile) == false)
+	{
+		return MakeSnapshotError<int64>(
+			ECasErrorCode::CreateSnapshotFailed,
+			EIoErrorCode::NotFound,
+			FString::Printf(TEXT("Cannot create snapshot, failed to find journal file '%s'"), *JournalFile));
+	}
+
+	// Load the previous snapshot if it exists
+	TResult<FCasSnapshot> SnapshotResult = FCasSnapshot::Load(SnapshotFile);
+	if (SnapshotResult.HasError())
+	{
+		const FCasErrorContext* Ctx = SnapshotResult.GetError().GetErrorContext<FCasErrorContext>();
+		if (Ctx == nullptr || Ctx->IoErrorCode != EIoErrorCode::NotFound)
+		{
+			return MakeError(SnapshotResult.StealError());
+		}
+	}
+
+	if (SnapshotResult.HasValue() == false)
+	{
+		SnapshotResult = MakeValue(FCasSnapshot());
+	}
+
+	// Load the snapshot from the journal
+	FCasSnapshot BaseSnapshot = SnapshotResult.StealValue();
+	SnapshotResult = FCasSnapshot::FromJournal(JournalFile, MoveTemp(BaseSnapshot));
+	if (SnapshotResult.HasError())
+	{
+		return MakeError(SnapshotResult.StealError());
+	}
+
+	const FString TmpSnapFile		= FCasSnapshot::MakeTempFileName(SnapshotFile);
+	const FString TmpJournalFile	= FCasSnapshot::MakeTempFileName(JournalFile);
+	
+	// Save the snapshot
+	int64 SnapshotSize			= -1;
+	FCasSnapshot Snapshot		= SnapshotResult.StealValue();
+	TResult<int64> SnapshotSaveResult	= FCasSnapshot::Save(Snapshot, TmpSnapFile);
+	if (SnapshotSaveResult.HasValue())
+	{
+		SnapshotSize = SnapshotSaveResult.GetValue();
+	}
+	else
+	{
+		return SnapshotSaveResult;
+	}
+
+	// Try create a new empty journal 
+	if (FResult Result = FCasJournal::Create(TmpJournalFile); Result.HasError())
+	{
+		return MakeError(Result.StealError());
+	}
+
+	if (Ifm.Move(*SnapshotFile, *TmpSnapFile) == false)
+	{
+		const uint32 LastError = FPlatformMisc::GetLastError();
+
+		return MakeSnapshotError<int64>(
+			ECasErrorCode::CreateSnapshotFailed,
+			EIoErrorCode::FileMoveFailed,
+			FString::Printf(TEXT("Failed to move tmp snapshot file '%s' -> '%s'"), *TmpSnapFile, *SnapshotFile),
+			LastError);
+	}
+
+	if (Ifm.Move(*JournalFile, *TmpJournalFile) == false)
+	{
+		const uint32 LastError = FPlatformMisc::GetLastError();
+
+		return MakeSnapshotError<int64>(
+			ECasErrorCode::CreateSnapshotFailed,
+			EIoErrorCode::FileMoveFailed,
+			FString::Printf(TEXT("Failed to move tmp journal file '%s' -> '%s'"), *TmpJournalFile, *JournalFile),
+			LastError);
+	}
+
+	return MakeValue(SnapshotSize);
+}
+
+bool FCasSnapshot::TempFilesExist(const FString& SnapshotFile, const FString& JournalFile)
+{
+	IFileManager& Ifm = IFileManager::Get();
+
+	const FString TmpSnapFile = MakeTempFileName(SnapshotFile);
+	const FString TmpJournalFile = MakeTempFileName(JournalFile);
+
+	if (Ifm.FileExists(*TmpSnapFile))
+	{
+		return true;
+	}
+
+	if (Ifm.FileExists(*TmpJournalFile))
+	{
+		return true;
+	}
+
+	return false;
+}
+
+///////////////////////////////////////////////////////////////////////////////
 FCas::FCas()
 {
 }
 
-FResult FCas::Configure(const FOnDemandInstallCacheConfig& Config)
+FResult FCas::Configure(const FOnDemandInstallCasConfig& Config)
 {
 	using namespace UE::UnifiedError::IoStoreOnDemand;
 
@@ -861,9 +1412,6 @@
 	Lookup.Empty();
 	BlockIds.Empty();
 
-	TStringBuilder<256> Path;
-	FPathViews::Append(Path, RootDirectory, TEXT("blocks"));
-
 	IFileManager& Ifm = IFileManager::Get();
 
 	if (bDeleteExisting)
@@ -871,20 +1419,20 @@
 		bool bRequireExists = false;
 		const bool bTree	= true;
 
-		if (Ifm.DeleteDirectory(Path.ToString(), bRequireExists, bTree) == false)
+		if (Ifm.DeleteDirectory(*RootDirectory, bRequireExists, bTree) == false)
 		{
 			return MakeCasError<void>(
-				ECasErrorCode::InitializeFailed, EIoErrorCode::DeleteError, FString::Printf(TEXT("Failed to delete CAS blocks directory '%s'"), Path.ToString()));
+				ECasErrorCode::InitializeFailed, EIoErrorCode::DeleteError, FString::Printf(TEXT("Failed to delete CAS blocks directory '%s'"), *RootDirectory));
 		}
 	}
 
-	if (Ifm.DirectoryExists(Path.ToString()) == false)
+	if (Ifm.DirectoryExists(*RootDirectory) == false)
 	{
 		const bool bTree = true;
-		if (Ifm.MakeDirectory(Path.ToString(), bTree) == false)
+		if (Ifm.MakeDirectory(*RootDirectory, bTree) == false)
 		{
 			return MakeCasError<void>(
-				ECasErrorCode::InitializeFailed, EIoErrorCode::WriteError, FString::Printf(TEXT("Failed to create CAS blocks directory '%s'"), Path.ToString()));
+				ECasErrorCode::InitializeFailed, EIoErrorCode::WriteError, FString::Printf(TEXT("Failed to create CAS blocks directory '%s'"), *RootDirectory));
 		}
 	}
 
@@ -1007,7 +1555,7 @@
 	FString Hex;
 	BytesToHexLower(reinterpret_cast<const uint8*>(&Id), sizeof(int32), Hex);
 	TStringBuilder<256> Path;
-	FPathViews::Append(Path, RootDirectory, TEXT("blocks"), Hex);
+	FPathViews::Append(Path, RootDirectory, Hex);
 	Path << TEXT(".ucas");
 
 	return FString(Path.ToView());
@@ -1358,14 +1906,11 @@
 		FBlockFileSizes&	BlockFileSizes;
 	};
 
-	TStringBuilder<256> Path;
-	FPathViews::Append(Path, RootDirectory, TEXT("blocks"));
-
 	FBlockFileSizes OutBlockFileSizes;
 
 	IPlatformFile& Ipf = FPlatformFileManager::Get().GetPlatformFile();
 	FDirectoryVisitor Visitor(OutBlockFileSizes);
-	Ipf.IterateDirectoryStat(Path.ToString(), Visitor);
+	Ipf.IterateDirectoryStat(*RootDirectory, Visitor);
 
 	return OutBlockFileSizes;
 }
@@ -1587,457 +2132,7 @@
 }
 
 ///////////////////////////////////////////////////////////////////////////////
-struct FCasSnapshot
-{
-	enum class EVersion : uint32
-	{
-		Invalid = 0,
-		Initial,
-
-		AddModificationTime,
-		FixSwitchBuild,
-		FixLostSnapshots,
-
-		LatestPlusOne,
-		Latest = LatestPlusOne - 1
-	};
-
-	static const FGuid CustomVersionGuid;
-
-	struct FHeader
-	{
-		static const inline uint8 MagicSequence[16] = {'+', 'S', 'N', 'A', 'P', 'S', 'H', 'O', 'T', 'H', 'E', 'A', 'D', 'E', 'R', '+'};
-
-		bool			IsValid() const;
-		static int64	Size() { return sizeof(FHeader); }
-
-		uint8		Magic[16] = {0};
-		EVersion	Version = EVersion::Invalid;
-		uint8		Pad[12] = {0};
-	};
-	static_assert(sizeof(FHeader) == 32);
-
-	struct FFooter
-	{
-		static const inline uint8 MagicSequence[16] = {'+', 'S', 'N', 'A', 'P', 'S', 'H', 'O', 'T', 'F', 'O', 'O', 'T', 'E', 'R', '+'};
-
-		static int64	Size() { return sizeof(FFooter); }
-		bool			IsValid() const;
-
-		uint8 Magic[16] = {0};
-	};
-	static_assert(sizeof(FFooter) == 16);
-
-	struct FBlock
-	{
-		friend FArchive& operator<<(FArchive& Ar, FBlock& Block)
-		{
-			Ar << Block.BlockId;
-			Ar << Block.LastAccess;
-
-			if (Ar.CustomVer(FCasSnapshot::CustomVersionGuid) >= static_cast<int32>(EVersion::AddModificationTime))
-			{
-				Ar << Block.ModTime;
-			}
-
-			return Ar;
-		}
-
-		FCasBlockId BlockId;
-		int64		LastAccess = 0;
-		int64		ModTime = 0;
-	};
-
-	using FChunkLocation = TPair<FCasAddr, FCasLocation>;
-
-	static TResult<FCasSnapshot>		FromJournal(const FString& JournalFile, FCasSnapshot&& BaseSnapshot);
-	static TResult<FCasSnapshot>		Load(const FString& SnapshotFile, int64* OutFileSize = nullptr);
-	static TResult<int64>				Save(const FCasSnapshot& Snapshot, const FString& SnapshotFile);
-	/** 
-	 * Append the journal file to the snapshot and reset the journal file.
-	 *
-	 * @param SnapshotFile	The snapshot file to appened the journal to. May not exist yet.
-	 * @param JournalFile	The journal file to append to the snapshot, and then reset.
-	 *
-	 * Note: If this operation fails, it will leave behind temporary files, call TempFilesExist() check for them.
-	 * There is no provided method to cleanup temporary files. If they are discovered, the
-	 * cache is in an inconsistent state and should be reset.
-	 */
-	static TResult<int64>				AppendAndResetJournal(
-											const FString& SnapshotFile, 
-											const FString& JournalFile);
-
-	/**
-	 * There is no provided method to cleanup temporary files. If they are discovered, the
-	 * cacche is in an inconsistent state and should be reset.
-	 */
-	static bool TempFilesExist(const FString& SnapshotFile, const FString& JournalFile);
-
-	static FString MakeTempFileName(const FString& FileName)
-	{
-		const bool bIncludeDot = true;
-		const FString Extension = FPaths::GetExtension(FileName, bIncludeDot) + TEXT("tmp");
-		return FPaths::ChangeExtension(FileName, Extension);
-	}
-
-	TArray<FBlock>						Blocks;
-	TArray<FChunkLocation>				ChunkLocations;
-	FCasBlockId							CurrentBlockId;
-};
-
-const FGuid FCasSnapshot::CustomVersionGuid(0x217bee7d, 0x23a54a23, 0x9ef080d2, 0x3df768b1);
-
-///////////////////////////////////////////////////////////////////////////////
-bool FCasSnapshot::FHeader::IsValid() const
-{
-	if (FMemory::Memcmp(&Magic, &FHeader::MagicSequence, sizeof(FHeader::MagicSequence)) != 0)
-	{
-		return false;
-	}
-
-	if (static_cast<uint32>(Version) > static_cast<uint32>(EVersion::Latest))
-	{
-		return false;
-	}
-
-	// Drop all snapshots from before we fixed snapshots not being propogated
-	if (static_cast<uint32>(Version) < static_cast<uint32>(EVersion::FixLostSnapshots))
-	{
-		return false;
-	}
-
-	// Versions that can no longer be loaded. The cache must be dropped instead.
-	const EVersion BadVersions[] = { UE_ONDEMANDINSTALLCACHE_BAD_SNAPSHOT_VERSIONS };
-	for (EVersion BadVersion : BadVersions)
-	{
-		if (Version == BadVersion)
-		{
-			return false;
-		}
-	}
-
-	return true;
-}
-
-bool FCasSnapshot::FFooter::IsValid() const
-{
-	return FMemory::Memcmp(Magic, FFooter::MagicSequence, sizeof(FFooter::MagicSequence)) == 0;
-}
-
-TResult<FCasSnapshot> FCasSnapshot::FromJournal(const FString& JournalFile, FCasSnapshot&& BaseSnapshot)
-{
-	FCas::FLookup CasLookup;
-	CasLookup.Reserve(BaseSnapshot.ChunkLocations.Num());
-	for (TPair<FCasAddr, FCasLocation>& Kv : BaseSnapshot.ChunkLocations)
-	{
-		CasLookup.Add(MoveTemp(Kv));
-	}
-
-	TMap<FCasBlockId, FBlock> BlockMap;
-	BlockMap.Reserve(BaseSnapshot.Blocks.Num());
-	for (const FCasSnapshot::FBlock& Block : BaseSnapshot.Blocks)
-	{
-		BlockMap.Add(Block.BlockId, Block);
-	}
-
-	FCasBlockId CurrentBlockId = BaseSnapshot.CurrentBlockId;
-
-	FResult ReplayResult = FCasJournal::Replay(
-		JournalFile,
-		[&CasLookup, &BlockMap, &CurrentBlockId](const FCasJournal::FEntry& JournalEntry) -> FResult
-		{
-			switch(JournalEntry.Type())
-			{
-			case FCasJournal::FEntry::EType::ChunkLocation:
-			{
-				const FCasJournal::FEntry::FChunkLocation& ChunkLocation = JournalEntry.ChunkLocation;
-				if (ChunkLocation.CasLocation.IsValid())
-				{
-					FCasLocation& Loc = CasLookup.FindOrAdd(ChunkLocation.CasAddr);
-					Loc = ChunkLocation.CasLocation;
-				}
-				else
-				{
-					CasLookup.Remove(ChunkLocation.CasAddr);
-				}
-				break;
-			}
-			case FCasJournal::FEntry::EType::BlockCreated:
-			{
-				const FCasJournal::FEntry::FBlockOperation& Op = JournalEntry.BlockOperation;
-				CurrentBlockId = Op.BlockId;
-				BlockMap.Emplace(Op.BlockId, FBlock{ .BlockId = Op.BlockId });
-				break;
-			}
-			case FCasJournal::FEntry::EType::BlockDeleted:
-			{
-				const FCasJournal::FEntry::FBlockOperation& Op = JournalEntry.BlockOperation;
-				BlockMap.Remove(Op.BlockId);
-				if (CurrentBlockId == Op.BlockId)
-				{
-					CurrentBlockId = FCasBlockId::Invalid; 
-				}
-				break;
-			}
-			case FCasJournal::FEntry::EType::BlockAccess:
-			{
-				const FCasJournal::FEntry::FBlockOperation& Op = JournalEntry.BlockOperation;
-				if (FBlock* Block = BlockMap.Find(Op.BlockId))
-				{
-					Block->LastAccess = Op.UtcTicks;
-				}
-				else
-				{
-					using namespace UE::UnifiedError::IoStoreOnDemand;
-					return MakeJournalError(ECasErrorCode::ReplayJournalFailed, EIoErrorCode::FileCorrupt, TEXT("Journal BlockAccess Op references an invalid block"));
-				}
-				break;
-			}
-			case FCasJournal::FEntry::EType::BlockModification:
-			{
-				const FCasJournal::FEntry::FBlockOperation& Op = JournalEntry.BlockOperation;
-				if (FBlock* Block = BlockMap.Find(Op.BlockId))
-				{
-					Block->ModTime = Op.UtcTicks;
-				}
-				else
-				{
-					using namespace UE::UnifiedError::IoStoreOnDemand;
-					return MakeJournalError(ECasErrorCode::ReplayJournalFailed, EIoErrorCode::FileCorrupt, TEXT("Journal BlockModification Op references an invalid block"));
-				}
-				break;
-			}
-			};
-
-			return MakeValue();
-		});
-
-	if (ReplayResult.HasError())
-	{
-		return MakeError(ReplayResult.StealError());
-	}
-
-	FCasSnapshot Snapshot;
-	BlockMap.GenerateValueArray(Snapshot.Blocks);
-
-	Snapshot.ChunkLocations = CasLookup.Array();
-	Snapshot.CurrentBlockId = CurrentBlockId;
-
-	return MakeValue(Snapshot);
-}
-
-TResult<int64> FCasSnapshot::Save(const FCasSnapshot& Snapshot, const FString& SnapshotFile)
-{
-	using namespace UE::UnifiedError::IoStoreOnDemand;
-
-	IFileManager& Ifm = IFileManager::Get();
-
-	TUniquePtr<FArchive> Ar(Ifm.CreateFileWriter(*SnapshotFile));
-	if (Ar.IsValid() == false)
-	{
-		return MakeSnapshotError<int64>(ECasErrorCode::SaveSnapshotFailed, EIoErrorCode::FileOpenFailed, FString::Printf(TEXT("Failed to open file '%s' for writing"), *SnapshotFile));
-	}
-
-	FHeader Header;
-	FMemory::Memcpy(&Header.Magic, &FHeader::MagicSequence, sizeof(FHeader::MagicSequence));
-	Header.Version = EVersion::Latest;
-
-	Ar->SetCustomVersion(CustomVersionGuid, static_cast<int32>(Header.Version), "CasSnapshotVersion");
-
-	Ar->Serialize(reinterpret_cast<uint8*>(&Header), FHeader::Size());
-	if (Ar->IsError())
-	{
-		const uint32 LastError = FPlatformMisc::GetLastError();
-		Ar.Reset();
-		return MakeSnapshotError<int64>(ECasErrorCode::SaveSnapshotFailed, EIoErrorCode::WriteError, FString::Printf(TEXT("Failed to write snapshot header to '%s'"), *SnapshotFile), LastError);
-	}
-
-	FCasSnapshot& NonConst = *const_cast<FCasSnapshot*>(&Snapshot);
-	*Ar << NonConst.Blocks;
-	*Ar << NonConst.ChunkLocations;
-	*Ar << NonConst.CurrentBlockId;
-
-	if (Ar->IsError())
-	{
-		const uint32 LastError = FPlatformMisc::GetLastError();
-		Ar.Reset();
-		return MakeSnapshotError<int64>(ECasErrorCode::SaveSnapshotFailed, EIoErrorCode::WriteError, FString::Printf(TEXT("Failed to write snapshot to '%s'"), *SnapshotFile), LastError);
-	}
-
-	FFooter Footer;
-	FMemory::Memcpy(&Footer.Magic, &FFooter::MagicSequence, sizeof(FFooter::MagicSequence));
-	Ar->Serialize(reinterpret_cast<uint8*>(&Footer), FFooter::Size());
-	if (Ar->IsError())
-	{
-		const uint32 LastError = FPlatformMisc::GetLastError();
-		Ar.Reset();
-		return MakeSnapshotError<int64>(ECasErrorCode::SaveSnapshotFailed, EIoErrorCode::WriteError, FString::Printf(TEXT("Failed to write snapshot footer to '%s'"), *SnapshotFile), LastError);
-	}
-
-	const int64 FileSize = Ar->TotalSize();
-	if (Ar->Close() == false)
-	{
-		const uint32 LastError = FPlatformMisc::GetLastError();
-		Ar.Reset();
-		return MakeSnapshotError<int64>(ECasErrorCode::SaveSnapshotFailed, EIoErrorCode::WriteError, FString::Printf(TEXT("Failed to close snapshot footer to '%s'"), *SnapshotFile), LastError);
-	}
-
-	return MakeValue(FileSize);
-}
-
-TResult<FCasSnapshot> FCasSnapshot::Load(const FString& SnapshotFile, int64* OutFileSize)
-{
-	using namespace UE::UnifiedError::IoStoreOnDemand;
-
-	IFileManager& Ifm = IFileManager::Get();
-
-	if (Ifm.FileExists(*SnapshotFile) == false)
-	{
-		return MakeSnapshotError<FCasSnapshot>(ECasErrorCode::LoadSnapshotFailed, EIoErrorCode::NotFound, FString::Printf(TEXT("Failed to find snapshot '%s'"), *SnapshotFile));
-	}
-
-	TUniquePtr<FArchive> Ar(Ifm.CreateFileReader(*SnapshotFile));
-	if (Ar.IsValid() == false)
-	{
-		return MakeSnapshotError<FCasSnapshot>(ECasErrorCode::LoadSnapshotFailed, EIoErrorCode::FileOpenFailed, FString::Printf(TEXT("Failed open snapshot '%s'"), *SnapshotFile));
-	}
-
-	FHeader Header;
-	Ar->Serialize(reinterpret_cast<uint8*>(&Header), FHeader::Size());
-	if (Ar->IsError() || Header.IsValid() == false)
-	{
-		return MakeSnapshotError<FCasSnapshot>(ECasErrorCode::LoadSnapshotFailed, EIoErrorCode::ReadError, FString::Printf(TEXT("Failed to validate snapshot header in '%s'"), *SnapshotFile));
-	}
-
-	Ar->SetCustomVersion(CustomVersionGuid, static_cast<int32>(Header.Version), "CasSnapshotVersion");
-
-	FCasSnapshot Snapshot;
-	*Ar << Snapshot.Blocks;
-	*Ar << Snapshot.ChunkLocations;
-	*Ar << Snapshot.CurrentBlockId;
-
-	FFooter Footer;
-	Ar->Serialize(reinterpret_cast<uint8*>(&Footer), FFooter::Size());
-	if (Ar->IsError() || Footer.IsValid() == false)
-	{
-		return MakeSnapshotError<FCasSnapshot>(ECasErrorCode::LoadSnapshotFailed, EIoErrorCode::ReadError, FString::Printf(TEXT("Failed to validate snapshot footer in '%s'"), *SnapshotFile));
-	}
-
-	if (OutFileSize != nullptr)
-	{
-		*OutFileSize = Ar->Tell();
-	}
-
-	return MakeValue(MoveTemp(Snapshot));
-}
-
-TResult<int64> FCasSnapshot::AppendAndResetJournal(const FString& SnapshotFile, const FString& JournalFile)
-{
-	using namespace UE::UnifiedError::IoStoreOnDemand;
-
-	IFileManager& Ifm = IFileManager::Get();
-
-	if (Ifm.FileExists(*JournalFile) == false)
-	{
-		return MakeSnapshotError<int64>(
-			ECasErrorCode::CreateSnapshotFailed,
-			EIoErrorCode::NotFound,
-			FString::Printf(TEXT("Cannot create snapshot, failed to find journal file '%s'"), *JournalFile));
-	}
-
-	// Load the previous snapshot if it exists
-	TResult<FCasSnapshot> SnapshotResult = FCasSnapshot::Load(SnapshotFile);
-	if (SnapshotResult.HasError())
-	{
-		const FCasErrorContext* Ctx = SnapshotResult.GetError().GetErrorContext<FCasErrorContext>();
-		if (Ctx == nullptr || Ctx->IoErrorCode != EIoErrorCode::NotFound)
-		{
-			return MakeError(SnapshotResult.StealError());
-		}
-	}
-
-	if (SnapshotResult.HasValue() == false)
-	{
-		SnapshotResult = MakeValue(FCasSnapshot());
-	}
-
-	// Load the snapshot from the journal
-	FCasSnapshot BaseSnapshot = SnapshotResult.StealValue();
-	SnapshotResult = FCasSnapshot::FromJournal(JournalFile, MoveTemp(BaseSnapshot));
-	if (SnapshotResult.HasError())
-	{
-		return MakeError(SnapshotResult.StealError());
-	}
-
-	const FString TmpSnapFile		= FCasSnapshot::MakeTempFileName(SnapshotFile);
-	const FString TmpJournalFile	= FCasSnapshot::MakeTempFileName(JournalFile);
-	
-	// Save the snapshot
-	int64 SnapshotSize			= -1;
-	FCasSnapshot Snapshot		= SnapshotResult.StealValue();
-	TResult<int64> SnapshotSaveResult	= FCasSnapshot::Save(Snapshot, TmpSnapFile);
-	if (SnapshotSaveResult.HasValue())
-	{
-		SnapshotSize = SnapshotSaveResult.GetValue();
-	}
-	else
-	{
-		return SnapshotSaveResult;
-	}
-
-	// Try create a new empty journal 
-	if (FResult Result = FCasJournal::Create(TmpJournalFile); Result.HasError())
-	{
-		return MakeError(Result.StealError());
-	}
-
-	if (Ifm.Move(*SnapshotFile, *TmpSnapFile) == false)
-	{
-		const uint32 LastError = FPlatformMisc::GetLastError();
-
-		return MakeSnapshotError<int64>(
-			ECasErrorCode::CreateSnapshotFailed,
-			EIoErrorCode::FileMoveFailed,
-			FString::Printf(TEXT("Failed to move tmp snapshot file '%s' -> '%s'"), *TmpSnapFile, *SnapshotFile),
-			LastError);
-	}
-
-	if (Ifm.Move(*JournalFile, *TmpJournalFile) == false)
-	{
-		const uint32 LastError = FPlatformMisc::GetLastError();
-
-		return MakeSnapshotError<int64>(
-			ECasErrorCode::CreateSnapshotFailed,
-			EIoErrorCode::FileMoveFailed,
-			FString::Printf(TEXT("Failed to move tmp journal file '%s' -> '%s'"), *TmpJournalFile, *JournalFile),
-			LastError);
-	}
-
-	return MakeValue(SnapshotSize);
-}
-
-bool FCasSnapshot::TempFilesExist(const FString& SnapshotFile, const FString& JournalFile)
-{
-	IFileManager& Ifm = IFileManager::Get();
-
-	const FString TmpSnapFile = MakeTempFileName(SnapshotFile);
-	const FString TmpJournalFile = MakeTempFileName(JournalFile);
-
-	if (Ifm.FileExists(*TmpSnapFile))
-	{
-		return true;
-	}
-
-	if (Ifm.FileExists(*TmpJournalFile))
-	{
-		return true;
-	}
-
-	return false;
-}
-
-///////////////////////////////////////////////////////////////////////////////
-void FCas::LoadSnapshot(FCasSnapshot&& Snapshot)
+void FCas::LoadSnapshot(FCasSnapshot::FCasState&& Snapshot)
 {
 	TUniqueLock Lock(Mutex);
 
@@ -2123,13 +2218,11 @@
 	struct FChunkRequest
 	{
 		explicit FChunkRequest(
-			FSharedAsyncFileHandle FileHandle,
 			FIoRequestImpl* Request,
 			FOnDemandChunkInfo&& Info,
 			FIoOffsetAndLength Range,
 			uint64 RequestedRawSize)
-				: SharedFileHandle(FileHandle)
-				, DispatcherRequest(Request)
+				: DispatcherRequest(Request)
 				, ChunkInfo(MoveTemp(Info))
 				, ChunkRange(Range)
 				, EncodedChunk(ChunkRange.GetLength())
@@ -2218,6 +2311,19 @@
 		mutable FSharedMutex	SharedMutex;
 	};
 
+	struct FCacheState
+	{
+		EOnDemandInstallCasType		Type{ EOnDemandInstallCasType::None };
+		FCas						Cas;
+		FPendingChunks				PendingChunks;
+		std::atomic<FCasBlockId>	CurrentBlock{ FCasBlockId::Invalid };
+		FSharedMutex				PurgeDefragMutex;
+		uint64						MaxCacheSize{ 0 };
+#if UE_ONDEMANDINSTALLCACHE_EXCLUSIVE_WRITE
+		UE::Tasks::FPipe			ExclusivePipe{ UE_SOURCE_LOCATION };
+#endif // UE_ONDEMANDINSTALLCACHE_EXCLUSIVE_WRITE
+	};
+
 public:
 	FOnDemandInstallCache(const FOnDemandInstallCacheConfig& Config, FOnDemandIoStore& IoStore, FDiskCacheGovernor& Governor);
 	virtual ~FOnDemandInstallCache();
@@ -2236,56 +2342,60 @@
 
 	// IOnDemandInstallCache
 	virtual bool								IsChunkCached(const FIoHash& ChunkHash) override;
-	virtual TResult<bool>						PutChunk(FIoBuffer&& Chunk, const FIoHash& ChunkHash) override;
-	virtual FResult								Purge(uint64 BytesToInstall) override;
-	virtual FResult								PurgeAllUnreferenced(bool bDefrag, const uint64* BytesToPurge = nullptr) override;
-	virtual FResult								DefragAll(const uint64* BytesToFree = nullptr) override;
+	virtual TResult<bool>						PutChunk(EOnDemandInstallCasType CasType, FIoBuffer&& Chunk, const FIoHash& ChunkHash) override;
+	virtual FResult								Purge(EOnDemandInstallCasType CasType, uint64 BytesToInstall) override;
+	virtual FResult								PurgeAllUnreferenced(EOnDemandInstallCasType CasType, bool bDefrag, const uint64* BytesToPurge = nullptr) override;
+	virtual FResult								DefragAll(EOnDemandInstallCasType CasType, const uint64* BytesToFree = nullptr) override;
 	virtual FResult								Verify() override;
-	virtual FResult								Flush() override;
-	virtual FResult								FlushLastAccess() override;
-	virtual void								UpdateLastAccess(TConstArrayView<FIoHash> ChunkHashes, bool& bInOutLastAccessDirty) override;
-	virtual FOnDemandInstallCacheUsage			GetCacheUsage() override;
+	virtual FResult								Flush(EOnDemandInstallCasType CasType) override;
+	virtual FResult								FlushLastAccess(const TStaticArray<bool, EOnDemandInstallCasType::Count>& bFlushCasLastAccess) override;
+	virtual void								UpdateLastAccess(
+													TConstArrayView<FIoHash> ChunkHashes, 
+													TStaticArray<bool, EOnDemandInstallCasType::Count>& bInOutLastAccessDirty) override;
+	virtual FOnDemandInstallCacheUsage			GetCacheUsage(EOnDemandInstallCasType CasType) override;
 
 private:
 	void										RegisterConsoleCommands();
 	FResult										Reset();
 	FResult										InitialVerify();
 	void										AddReferencesToBlocks(
+													const FCacheState& CacheState,
 													const TArray<FSharedOnDemandContainer>& Containers, 
 													const TArray<TBitArray<>>& ChunkEntryIndices,
 													FCasBlockInfoMap& BlockInfoMap,
 													uint64& OutTotalReferencedBytes) const;
-	FResult										PurgeInternal(FCasBlockInfoMap& BlockInfo, uint64 TotalBytesToPurge, uint64& OutTotalPurgedBytes);
+	FResult										PurgeInternal(
+													FCacheState& CacheState, 
+													FCasBlockInfoMap& BlockInfo, 
+													uint64 TotalBytesToPurge, 
+													uint64& OutTotalPurgedBytes);
 	FResult										Defrag(
+													FCacheState& CacheState,
 													const TArray<FSharedOnDemandContainer>& Containers,
 													const TArray<TBitArray<>>& ChunkEntryIndices,
 													FCasBlockInfoMap& BlockInfo, 
 													const uint64* TotalBytesToFree = nullptr,
 													uint64* OutTotalFreedBytes = nullptr);
 	bool										Resolve(FIoRequestImpl* Request);
+	bool										ResolveFromCas(FCacheState& CacheState, FIoRequestImpl* Request);
 	void										CompleteRequest(FIoRequestImpl* Request, EIoErrorCode Status);
-	FResult										FlushPendingChunks(FPendingChunks& Chunks, int64 UtcAccessTicks = 0, int64 UtcModificationTicks = 0);
-	FResult										FlushPendingChunksImpl(const FPendingChunks& Chunks, int64 UtcAccessTicks = 0, int64 UtcModificationTicks = 0);
+	FResult										FlushPendingChunks(FCacheState& CacheState, FPendingChunks& Chunks, int64 UtcAccessTicks = 0, int64 UtcModificationTicks = 0);
+	FResult										FlushPendingChunksImpl(FCacheState& CacheState, const FPendingChunks& Chunks, int64 UtcAccessTicks = 0, int64 UtcModificationTicks = 0);
 	FString										GetJournalFilename() const { return CacheDirectory / TEXT("cas.jrn"); }
 	FString										GetSnapshotFilename() const { return CacheDirectory / TEXT("cas.snp"); }
-	UE::UnifiedError::IoStoreOnDemand::FInstallCacheErrorContext MakeInstallCacheErrorContext(uint64 TotalCachedBytes = 0, uint64 RequestedBytes = 0, uint32 LineNo = __builtin_LINE());
+	FCacheState*								GetCacheState(EOnDemandInstallCasType CacheType);
+	const FCacheState*							GetCacheState(EOnDemandInstallCasType CacheType) const;
+	UE::UnifiedError::IoStoreOnDemand::FInstallCacheErrorContext MakeInstallCacheErrorContext(FCacheState& CacheState, uint64 TotalCachedBytes = 0, uint64 RequestedBytes = 0, uint32 LineNo = __builtin_LINE());
 
 	FOnDemandIoStore&			IoStore;
 	FDiskCacheGovernor&			Governor;
 	FString						CacheDirectory;
-	FCas						Cas;
-	std::atomic<FCasBlockId>	CurrentBlock{ FCasBlockId::Invalid };
-	FPendingChunks				PendingChunks;
 	FSharedBackendContext		BackendContext;
 	FIoRequestList				CompletedRequests;
 	FMutex						Mutex;
-	FSharedMutex				PurgeDefragMutex;
-	uint64						MaxCacheSize{ 0 };
 	uint64						MaxJournalSize;
 
-#if UE_ONDEMANDINSTALLCACHE_EXCLUSIVE_WRITE
-	UE::Tasks::FPipe			ExclusivePipe{ UE_SOURCE_LOCATION };
-#endif // UE_ONDEMANDINSTALLCACHE_EXCLUSIVE_WRITE
+	TArray<FCacheState>			CacheStates;
 
 #if UE_IAD_DEBUG_CONSOLE_CMDS
 	TArray<IConsoleCommand*> ConsoleCommands;
@@ -2297,37 +2407,60 @@
 	: IoStore(InIoStore)
 	, Governor(InGovernor)
 	, CacheDirectory(Config.RootDirectory)
-	, MaxCacheSize(Config.DiskQuota)
 	, MaxJournalSize(Config.JournalMaxSize)
 {
 	using namespace UE::UnifiedError::IoStoreOnDemand;
 
-	FResult InitResult = Cas.Configure(Config);
-	if (InitResult.HasError())
+	FResult InitResult = MakeValue();
+
+	CacheStates.Empty(Config.CasConfig.Num());
+	for (const FOnDemandInstallCasConfig& CasConfig : Config.CasConfig)
 	{
-		UE_LOGFMT(LogIoStoreOnDemand, Fatal, "Failed to configure install cache, error: {Error}", InitResult.GetError());
-		return;
+		if (ensureAlways(CasConfig.Type != EOnDemandInstallCasType::None) == false)
+		{
+			continue;
+		}
+
+		FCacheState& CacheState = CacheStates.Emplace_GetRef();
+		CacheState.Type = CasConfig.Type;
+
+		CacheState.MaxCacheSize = CasConfig.DiskQuota;
+
+		InitResult = CacheState.Cas.Configure(CasConfig);
+		if (InitResult.HasError())
+		{
+			UE_LOGFMT(LogIoStoreOnDemand, Fatal, "Failed to configure install cache CAS {CasType}, error: {Error}", LexToString(CasConfig.Type), InitResult.GetError());
+			return;
+		}
+
+		UE_LOG(LogIoStoreOnDemand, Log, TEXT("Initializing install cache CAS %s, MaxCacheSize=%.2lf MiB, MaxJournalSize=%.2lf KiB, MaxBlockSize=%.2lf MiB, MinBlockSize=%.2lf MiB"),
+			LexToString(CasConfig.Type), ToMiB(CacheState.MaxCacheSize), ToKiB(MaxJournalSize), ToMiB(CacheState.Cas.GetMaxBlockSize()), ToMiB(CacheState.Cas.GetMinBlockSize()));
+
+		const uint64 MinDiskQuota = 2 * CacheState.Cas.GetMaxBlockSize();
+		if (CacheState.MaxCacheSize < MinDiskQuota)
+		{
+			UE_LOG(LogIoStoreOnDemand, Fatal, TEXT("Failed to initialize install cache CAS %s - disk quota must be at least %.2lf MiB"), LexToString(CasConfig.Type), ToMiB(MinDiskQuota));
+			return;
+		}
+
+		// Reserve one block of space for defragmentation overhead
+		CacheState.MaxCacheSize -= CacheState.Cas.GetMaxBlockSize();
+		UE_LOG(LogIoStoreOnDemand, Log, TEXT("Effective MaxCacheSize for CAS %s without defragmentation space is MaxCacheSize=%.2lf MiB"), LexToString(CasConfig.Type), ToMiB(CacheState.MaxCacheSize));
+
+		TStringBuilder<256> CasDirectory;
+		FPathViews::Append(CasDirectory, CacheDirectory, CasConfig.CasSubdirectory);
+
+		InitResult = CacheState.Cas.Initialize(CasDirectory);
+		if (InitResult.HasError())
+		{
+			UE_LOGFMT(LogIoStoreOnDemand, Fatal, "Failed to initialize install cache CAS {CasType}, error: {Error}", LexToString(CasConfig.Type), InitResult.GetError());
+			return;
+		}
 	}
 
-	UE_LOG(LogIoStoreOnDemand, Log, TEXT("Initializing install cache, MaxCacheSize=%.2lf MiB, MaxJournalSize=%.2lf KiB, MaxBlockSize=%.2lf MiB, MinBlockSize=%.2lf MiB"),
-		ToMiB(MaxCacheSize), ToKiB(MaxJournalSize), ToMiB(Cas.GetMaxBlockSize()), ToMiB(Cas.GetMinBlockSize()));
-
-	const uint64 MinDiskQuota = 2 * Cas.GetMaxBlockSize();
-	if (MaxCacheSize < MinDiskQuota)
+	if (CacheStates.IsEmpty())
 	{
-		UE_LOG(LogIoStoreOnDemand, Fatal, TEXT("Failed to initialize install cache - disk quota must be at least %.2lf MiB"), ToMiB(MinDiskQuota));
-		return;
+		UE_LOGF(LogIoStoreOnDemand, Fatal, UTF8TEXT("Failed to configure install cache CAS! No valid configurations!"));
 	}
 
-	// Reserve one block of space for defragmentation overhead
-	MaxCacheSize -= Cas.GetMaxBlockSize();
-	UE_LOG(LogIoStoreOnDemand, Log, TEXT("Effective MaxCacheSize without defragmentation space is MaxCacheSize=%.2lf MiB"), ToMiB(MaxCacheSize));
-
-	InitResult = Cas.Initialize(CacheDirectory);
-	if (InitResult.HasError())
-	{
-		UE_LOGFMT(LogIoStoreOnDemand, Fatal, "Failed to initialize install cache, error: {Error}", InitResult.GetError());
-		return;
-	}
-
 	const FString SnapshotFilename = GetSnapshotFilename();
@@ -2373,10 +2506,20 @@
 		if (SnapshotResult.HasValue())
 		{
 			FCasSnapshot Snapshot = SnapshotResult.StealValue();
-			UE_LOG(LogIoStoreOnDemand, Log, TEXT("Loaded CAS snapshot '%s' %.2lf KiB with %d blocks and %d chunk locations"),
-					*SnapshotFilename, ToKiB(SnapshotSize), Snapshot.Blocks.Num(), Snapshot.ChunkLocations.Num());
+			
+			UE_LOG(LogIoStoreOnDemand, Log, TEXT("Loaded CAS snapshot '%s' %.2lf KiB"), *SnapshotFilename, ToKiB(SnapshotSize));
 
-			Cas.LoadSnapshot(MoveTemp(Snapshot));
+			for (FCacheState& CacheState : CacheStates)
+			{
+				FCasSnapshot::FCasState* SnapshotCasState = Algo::FindBy(Snapshot.CasState, CacheState.Type, &FCasSnapshot::FCasState::CasType);
+				if (SnapshotCasState)
+				{
+					UE_LOG(LogIoStoreOnDemand, Log, TEXT("Loaded CAS state from snapshot with %d blocks and %d chunk locations"),
+						SnapshotCasState->Blocks.Num(), SnapshotCasState->ChunkLocations.Num());
+
+					CacheState.Cas.LoadSnapshot(MoveTemp(*SnapshotCasState));
+				}
+			}
 		}
 		else if (const FCasErrorContext* Ctx = SnapshotResult.GetError().GetErrorContext<FCasErrorContext>())
 		{
@@ -2393,7 +2536,20 @@
 		// Replay the journal 
 		InitResult = FCasJournal::Replay(JournalFilename, [this](const FCasJournal::FEntry& JournalEntry)
 		{
-			FResult Result = Cas.ApplyJournalEntry(JournalEntry);
+			FResult Result = MakeValue();
+			FCacheState* CacheState = Algo::FindBy(CacheStates, JournalEntry.CasType(), &FCacheState::Type);
+			if (CacheState == nullptr)
+			{
+				// This cache state is not currently configured and there's no way to partially cleanup the cache right now.
+				// For now, just issue an error and wipe the cache.
+				// TODO: IadMultiCAS: Need a way to reset a cache type if we remove it.
+				//		TODO: Scan root dir for cache subdirs that are unknown and delete them.
+				//		TODO: When making a snapshot, strip any cache types that aren't currently configured.
+				return MakeJournalError(ECasErrorCode::ReplayJournalFailed, EIoErrorCode::InvalidParameter, 
+					FString::Printf(TEXT("Journal Entry references an invalid cache type %s"), LexToString(JournalEntry.CasType())));
+			}
+
+			Result = CacheState->Cas.ApplyJournalEntry(JournalEntry);
 			if (Result.HasError())
 			{
 				return Result;
@@ -2404,14 +2560,14 @@
 			case FCasJournal::FEntry::EType::BlockCreated:
 			{
 				const FCasJournal::FEntry::FBlockOperation& Op = JournalEntry.BlockOperation;
-				CurrentBlock = Op.BlockId;
+				CacheState->CurrentBlock = Op.BlockId;
 				break;
 			}
 			case FCasJournal::FEntry::EType::BlockDeleted:
 			{
 				const FCasJournal::FEntry::FBlockOperation& Op = JournalEntry.BlockOperation;
 				FCasBlockId MaybeCurrentBlock = Op.BlockId;
-				CurrentBlock.compare_exchange_strong(MaybeCurrentBlock, FCasBlockId::Invalid);
+				CacheState->CurrentBlock.compare_exchange_strong(MaybeCurrentBlock, FCasBlockId::Invalid);
 				break;
 			}
 			};
@@ -2432,9 +2588,17 @@
 				{
 					UE_LOG(LogIoStoreOnDemand, Log, TEXT("Created CAS journal '%s'"), *JournalFilename);
 
-					// Make sure that there are no existing blocks when starting from an empty cache
-					const bool bDeleteExisting = true;
-					InitResult = Cas.Initialize(CacheDirectory, bDeleteExisting);
+					for (FCacheState& CacheState : CacheStates)
+					{
+						// Make sure that there are no existing blocks when starting from an empty cache
+						FString CasRootDir = CacheState.Cas.GetRootDirectory();
+						const bool bDeleteExisting = true;
+						InitResult = CacheState.Cas.Initialize(CasRootDir, bDeleteExisting);
+						if (InitResult.HasError())
+						{
+							break;
+						}
+					}
 				}
 			}
 		}
@@ -2459,7 +2623,10 @@
 	{
 		UE_LOG(LogIoStoreOnDemand, Log, TEXT("Install cache Ok!"));
 		RegisterConsoleCommands();
-		Cas.Compact();
+		for (FCacheState& CacheState : CacheStates)
+		{
+			CacheState.Cas.Compact();
+		}
 	}
 	else
 	{
@@ -2478,21 +2645,28 @@
 
 void FOnDemandInstallCache::Shutdown()
 {
-	if (FResult Result = FlushPendingChunksImpl(PendingChunks); Result.HasError())
+	for (FCacheState& CacheState : CacheStates)
 	{
-		UE_LOG(LogIoStoreOnDemand, Error, TEXT("Failed to flush pending chunks on shutdown, reason '%s'"), *Result.GetError().GetErrorMessage(true).ToString());
+		if (FResult Result = FlushPendingChunksImpl(CacheState, CacheState.PendingChunks); Result.HasError())
+		{
+			UE_LOG(LogIoStoreOnDemand, Error, TEXT("Failed to flush pending chunks on shutdown for %s, reason '%s'"), 
+				LexToString(CacheState.Type), *Result.GetError().GetErrorMessage(true).ToString());
+		}
 	}
 
-	FCas::FLastAccess LastAccess = Cas.GetAndClearDirtyLastAccess();
-
 	const FString JournalFile = GetJournalFilename();
-	FCasJournal::FTransaction Transaction = FCasJournal::Begin(JournalFile);
-	JournalLastAccess(Transaction, LastAccess);
+	for (FCacheState& CacheState : CacheStates)
+	{
+		FCasJournal::FTransaction Transaction = FCasJournal::Begin(CacheState.Type, JournalFile);
 
-	if (FResult Result = FCasJournal::Commit(MoveTemp(Transaction)); Result.HasError())
-	{
-		UE_LOG(LogIoStoreOnDemand, Error, TEXT("Failed to update CAS journal '%s' with block timestamp(s), error '%s'"),
-			*JournalFile, *LexToString(Result.GetError()));
+		FCas::FLastAccess LastAccess = CacheState.Cas.GetAndClearDirtyLastAccess();
+		JournalLastAccess(Transaction, LastAccess);
+
+		if (FResult Result = FCasJournal::Commit(MoveTemp(Transaction)); Result.HasError())
+		{
+			UE_LOG(LogIoStoreOnDemand, Error, TEXT("Failed to update CAS journal '%s' with block timestamp(s), error '%s'"),
+				*JournalFile, *LexToString(Result.GetError()));
+		}
 	}
 
 #if UE_IAD_DEBUG_CONSOLE_CMDS
@@ -2552,17 +2726,21 @@
 	EIoErrorCode ErrorCode = EIoErrorCode::UnknownChunkID;
 	if (FOnDemandChunkInfo ChunkInfo = IoStore.GetInstalledChunkInfo(ChunkId, ErrorCode))
 	{
+		for (const FCacheState& CacheState : CacheStates)
 		{
-			TSharedLock SharedLock(PendingChunks.SharedMutex);
 
-			if (PendingChunks.ContainsChunk(ChunkInfo.Hash()))
 			{
-				return true;
+				TSharedLock SharedLock(CacheState.PendingChunks.SharedMutex);
+
+				if (CacheState.PendingChunks.ContainsChunk(ChunkInfo.Hash()))
+				{
+					return true;
+				}
 			}
+
+			const FCasLocation CasLoc = CacheState.Cas.FindChunk(ChunkInfo.Hash());
+			return CasLoc.IsValid();
 		}
-
-		const FCasLocation CasLoc = Cas.FindChunk(ChunkInfo.Hash());
-		return CasLoc.IsValid();
 	}
 
 	return false;
@@ -2638,22 +2816,48 @@
 		return true;
 	}
 
+	// The internal request parameters are attached/owned by the I/O request via
+	// the backend data parameter. The chunk request is deleted in GetCompletedRequests
+	FChunkRequest::Attach(*Request, new FChunkRequest(
+		Request,
+		MoveTemp(ChunkInfo),
+		ChunkRange.ConsumeValueOrDie(),
+		RequestSize));
+
+	bool bFoundInCas = false;
+	for (FCacheState& CacheState : CacheStates)
+	{
+		if (ResolveFromCas(CacheState, Request))
+		{
+			bFoundInCas = true;
+			break;
+		}
+	}
+
+	if (!bFoundInCas)
+	{
+		// Chunk is referenced but not in any cache. This should probably never happen currently.
+		// If pinning of unrefernced chunks was allowed, then this would be expected...
+		TUniquePtr<FChunkRequest> Detached = FChunkRequest::Detach(*Request);
+		CompleteRequest(Request, EIoErrorCode::NotInstalled);
+	}
+
+	return true;
+}
+
+bool FOnDemandInstallCache::ResolveFromCas(FCacheState& CacheState, FIoRequestImpl* Request)
+{
+	FChunkRequest& ChunkRequest = FChunkRequest::GetRef(*Request);
+
 	// Must check pending chunks before CAS because otherwise
 	// a chunk could be flushed after checking the CAS and before checking pending.
 	// Otherise we'd need take both locks.
 	{
-		TDynamicSharedLock SharedLock(PendingChunks.SharedMutex);
+		TDynamicSharedLock SharedLock(CacheState.PendingChunks.SharedMutex);
 
-		const FIoBuffer* PendingChunk = PendingChunks.FindChunk(ChunkInfo.Hash());
+		const FIoBuffer* PendingChunk = CacheState.PendingChunks.FindChunk(ChunkRequest.ChunkInfo.Hash());
 		if (PendingChunk != nullptr)
 		{
-			FChunkRequest& ChunkRequest = FChunkRequest::Attach(*Request, new FChunkRequest(
-				FSharedAsyncFileHandle(),
-				Request,
-				MoveTemp(ChunkInfo),
-				ChunkRange.ConsumeValueOrDie(),
-				RequestSize));
-
 			FMemoryView BytesToRead = PendingChunk->GetView().Mid(ChunkRequest.ChunkRange.GetOffset(), ChunkRequest.EncodedChunk.GetSize());
 			check(BytesToRead.GetSize() == ChunkRequest.EncodedChunk.GetSize());
 			FMemory::Memcpy(ChunkRequest.EncodedChunk.GetData(), BytesToRead.GetData(), ChunkRequest.EncodedChunk.GetSize());
@@ -2669,44 +2873,33 @@
 		}
 	}
 
-	const FCasLocation CasLoc = Cas.FindChunk(ChunkInfo.Hash());
+	const FCasLocation CasLoc = CacheState.Cas.FindChunk(ChunkRequest.ChunkInfo.Hash());
 	if (CasLoc.IsValid() == false)
 	{
-		CompleteRequest(Request, EIoErrorCode::NotInstalled);
-		return true;
+		return false;
 	}
 
 	TRACE_IOSTORE_BACKEND_REQUEST_STARTED(Request, this);
 	constexpr const bool bDirty = true;
-	if (Cas.TrackAccessIf(ECasTrackAccessType::Granular, CasLoc.BlockId, bDirty))
+	if (CacheState.Cas.TrackAccessIf(ECasTrackAccessType::Granular, CasLoc.BlockId, bDirty))
 	{
 		IoStore.FlushLastAccess({}, nullptr);
 	}
 
 #if UE_ONDEMANDINSTALLCACHE_EXCLUSIVE_WRITE
-	const bool bIsLocationInCurrentBlock = CasLoc.BlockId == CurrentBlock;
+	const bool bIsLocationInCurrentBlock = CasLoc.BlockId == CacheState.CurrentBlock;
 	if (bIsLocationInCurrentBlock)
 	{
 		// The current block may have open writes which may cause async reads to fail
 		// on some platforms. Schedule the reads to happen on the same pipe as writes
-
-		// The internal request parameters are attached/owned by the I/O request via
-		// the backend data parameter. The chunk request is deleted in GetCompletedRequests
-		FChunkRequest::Attach(*Request, new FChunkRequest(
-			FSharedAsyncFileHandle(),
-			Request,
-			MoveTemp(ChunkInfo),
-			ChunkRange.ConsumeValueOrDie(),
-			RequestSize));
-
-		ExclusivePipe.Launch(UE_SOURCE_LOCATION, [this, Request, CasLoc]
+		CacheState.ExclusivePipe.Launch(UE_SOURCE_LOCATION, [this, &CacheState, Request, CasLoc]
 		{
 			FChunkRequest& ChunkRequest = FChunkRequest::GetRef(*Request);
 			EIoErrorCode Status = EIoErrorCode::FileOpenFailed;
 
-			const FString Filename = Cas.GetBlockFilename(CasLoc.BlockId);
+			const FString Filename = CacheState.Cas.GetBlockFilename(CasLoc.BlockId);
 
-			TResult<FSharedFileHandle> FileOpenResult = Cas.OpenRead(CasLoc.BlockId);
+			TResult<FSharedFileHandle> FileOpenResult = CacheState.Cas.OpenRead(CasLoc.BlockId);
 			if (FileOpenResult.HasValue())
 			{
 				Status = EIoErrorCode::ReadError;
@@ -2715,7 +2908,8 @@
 				const int64 CasBlockOffset = CasLoc.BlockOffset + ChunkRequest.ChunkRange.GetOffset();
 				if (Request->IsCancelled())
 				{
-					UE_LOG(LogIoStoreOnDemand, Verbose, TEXT("Cancelled request - skipped seek to offset %" INT64_FMT " in CAS block '%s'"), CasBlockOffset, *Filename);
+					UE_LOG(LogIoStoreOnDemand, Verbose, TEXT("Cancelled request - skipped seek to offset %" INT64_FMT " in CAS block '%s' in CAS %s"), 
+						CasBlockOffset, *Filename, LexToString(CacheState.Type));
 				}
 				else if (FileHandle->Seek(CasBlockOffset))
 				{
@@ -2726,20 +2920,20 @@
 					}
 					else
 					{
-						UE_LOG(LogIoStoreOnDemand, Error, TEXT("Failed to read %" UINT64_FMT " bytes at offset %" INT64_FMT " in CAS block '%s'"),
-							ChunkRequest.EncodedChunk.GetSize(),
-							CasBlockOffset,
-							*Filename);
+						UE_LOG(LogIoStoreOnDemand, Error, TEXT("Failed to read %" UINT64_FMT " bytes at offset %" INT64_FMT " in CAS block '%s' in CAS %s"),
+							ChunkRequest.EncodedChunk.GetSize(), CasBlockOffset, *Filename, LexToString(CacheState.Type));
 					}
 				}
 				else
 				{
-					UE_LOG(LogIoStoreOnDemand, Error, TEXT("Failed to seek to offset %" INT64_FMT " in CAS block '%s'"), CasBlockOffset, *Filename);
+					UE_LOG(LogIoStoreOnDemand, Error, TEXT("Failed to seek to offset %" INT64_FMT " in CAS block '%s' in CAS %s"), 
+						CasBlockOffset, *Filename, LexToString(CacheState.Type));
 				}
 			}
 			else
 			{
-				UE_LOGFMT(LogIoStoreOnDemand, Error, "Failed to open CAS block '{Filename}' for reading ({Error})", *Filename, FileOpenResult.GetError().GetErrorMessage(true).ToString());
+				UE_LOGFMT(LogIoStoreOnDemand, Error, "Failed to open CAS block '{Filename}' in CAS {CasType} for reading ({Error})", 
+					*Filename, LexToString(CacheState.Type), FileOpenResult.GetError().GetErrorMessage(true).ToString());
 			}
 
 			UE::Tasks::Launch(UE_SOURCE_LOCATION, [this, Request, Status]
@@ -2752,26 +2946,19 @@
 	}
 #endif // UE_ONDEMANDINSTALLCACHE_EXCLUSIVE_WRITE
 
-	TResult<FSharedAsyncFileHandle> FileOpenResult = Cas.OpenAsyncRead(CasLoc.BlockId);
+
+	TResult<FSharedAsyncFileHandle> FileOpenResult = CacheState.Cas.OpenAsyncRead(CasLoc.BlockId);
 	if (FileOpenResult.HasError())
 	{
-		const FString Filename = Cas.GetBlockFilename(CasLoc.BlockId);
-		UE_LOG(LogIoStoreOnDemand, Error, TEXT("Failed to open CAS block '%s' for async reading, reason '%s'"), *Filename, *FileOpenResult.GetError().GetErrorMessage(true).ToString());
+		const FString Filename = CacheState.Cas.GetBlockFilename(CasLoc.BlockId);
+		UE_LOG(LogIoStoreOnDemand, Error, TEXT("Failed to open CAS block '%s' in CAS %s for async reading, reason '%s'"),
+			*Filename, LexToString(CacheState.Type), *FileOpenResult.GetError().GetErrorMessage(true).ToString());
 		TUniquePtr<FChunkRequest> Detached = FChunkRequest::Detach(*Request);
 		CompleteRequest(Request, EIoErrorCode::FileOpenFailed);
 		return true;
 	}
 	
-	FSharedAsyncFileHandle FileHandle(FileOpenResult.StealValue());
+	ChunkRequest.SharedFileHandle = FileOpenResult.StealValue();
 
-	// The internal request parameters are attached/owned by the I/O request via
-	// the backend data parameter. The chunk request is deleted in GetCompletedRequests
-	FChunkRequest& ChunkRequest = FChunkRequest::Attach(*Request, new FChunkRequest(
-		FileHandle,
-		Request,
-		MoveTemp(ChunkInfo),
-		ChunkRange.ConsumeValueOrDie(),
-		RequestSize));
-
 	FAsyncFileCallBack Callback = [this, Request](bool bWasCancelled, IAsyncReadRequest* ReadRequest)
 	{
@@ -2782,7 +2969,7 @@
 		});
 	};
 
-	ChunkRequest.FileReadRequest.Reset(FileHandle->ReadRequest(
+	ChunkRequest.FileReadRequest.Reset(ChunkRequest.SharedFileHandle->ReadRequest(
 		CasLoc.BlockOffset + ChunkRequest.ChunkRange.GetOffset(),
 		ChunkRequest.ChunkRange.GetLength(),
 		EAsyncIOPriorityAndFlags::AIOP_BelowNormal,
@@ -2802,53 +2989,76 @@
 
 bool FOnDemandInstallCache::IsChunkCached(const FIoHash& ChunkHash)
 {
+	for (FCacheState& CacheState : CacheStates)
 	{
-		TSharedLock SharedLock(PendingChunks.SharedMutex);
+		{
+			TSharedLock SharedLock(CacheState.PendingChunks.SharedMutex);
 
-		if (PendingChunks.ContainsChunk(ChunkHash))
+			if (CacheState.PendingChunks.ContainsChunk(ChunkHash))
+			{
+				return true;
+			}
+		}
+
+		const FCasLocation Loc = CacheState.Cas.FindChunk(ChunkHash);
+		if (Loc.IsValid())
 		{
 			return true;
 		}
 	}
 
-	const FCasLocation Loc = Cas.FindChunk(ChunkHash);
-	return Loc.IsValid();
+	return false;
 }
 
-TResult<bool> FOnDemandInstallCache::PutChunk(FIoBuffer&& Chunk, const FIoHash& ChunkHash)
+TResult<bool> FOnDemandInstallCache::PutChunk(EOnDemandInstallCasType CasType, FIoBuffer&& Chunk, const FIoHash& ChunkHash)
 {
-	if (PendingChunks.TotalSize > FPendingChunks::MaxPendingBytes)
+	FCacheState* CacheState = GetCacheState(CasType);
+	if (CacheState == nullptr)
 	{
-		const uint64 TotalBlockSize = Cas.GetTotalBlockSize();
-		if (TotalBlockSize + PendingChunks.TotalSize > MaxCacheSize)
+		using namespace UE::UnifiedError::IoStoreOnDemand;
+		return MakeCasError<bool>(
+			ECasErrorCode::WriteBlockFailed, EIoErrorCode::InvalidParameter, FString::Printf(TEXT("CAS %s is not configured"), LexToString(CasType)));
+	}
+
+	if (CacheState->PendingChunks.TotalSize > FPendingChunks::MaxPendingBytes)
+	{
+		const uint64 TotalBlockSize = CacheState->Cas.GetTotalBlockSize();
+		if (TotalBlockSize + CacheState->PendingChunks.TotalSize > CacheState->MaxCacheSize)
 		{
 			// The calling code needs to decide if this is an error.
 			// Cache being full isn't necessarily an error, but we can't put the chunk.
 			return MakeValue(false);
 		}
 
-		if (FResult Result = FlushPendingChunks(PendingChunks); Result.HasError())
+		if (FResult Result = FlushPendingChunks(*CacheState, CacheState->PendingChunks); Result.HasError())
 		{
 			return MakeError(Result.StealError());
 		}
-		check(PendingChunks.IsEmpty());
+		check(CacheState->PendingChunks.IsEmpty());
 	}
 
-	TUniqueLock Lock(PendingChunks.SharedMutex);
+	TUniqueLock Lock(CacheState->PendingChunks.SharedMutex);
 
 	FOnDemandInstallCacheStats::OnPutChunk(Chunk.GetSize());
-	PendingChunks.Append(MoveTemp(Chunk), ChunkHash);
+	CacheState->PendingChunks.Append(MoveTemp(Chunk), ChunkHash);
 	return MakeValue(true);
 }
 
-FResult FOnDemandInstallCache::Purge(const uint64 BytesToInstall)
+FResult FOnDemandInstallCache::Purge(EOnDemandInstallCasType CasType, const uint64 BytesToInstall)
 {
 	using namespace UE::UnifiedError::IoStoreOnDemand;
 
-	TUniqueLock Lock(PurgeDefragMutex);
+	FCacheState* CacheState = GetCacheState(CasType);
+	if (CacheState == nullptr)
+	{
+		return MakeCasError<void>(
+			ECasErrorCode::DeleteBlockFailed, EIoErrorCode::InvalidParameter, FString::Printf(TEXT("CAS %s is not configured"), LexToString(CasType)));
+	}
 
+	TUniqueLock Lock(CacheState->PurgeDefragMutex);
+
 	FCasBlockInfoMap	BlockInfo;
-	const uint64		TotalCachedBytes = Cas.GetBlockInfo(BlockInfo);
+	const uint64		TotalCachedBytes = CacheState->Cas.GetBlockInfo(BlockInfo);
 
 	TArray<FSharedOnDemandContainer>	Containers;
 	TArray<TBitArray<>>					ChunkEntryIndices;
@@ -2861,14 +3071,15 @@
 	uint64 TotalReferencedBlockBytes = 0;
 	int64 OldestBlockAccess = FDateTime::MaxValue().GetTicks();
 
-	AddReferencesToBlocks(Containers, ChunkEntryIndices, BlockInfo, ReferencedBytes);
+	AddReferencesToBlocks(*CacheState, Containers, ChunkEntryIndices, BlockInfo, ReferencedBytes);
 	for (const TPair<FCasBlockId, FCasBlockInfo>& Kv : BlockInfo)
 	{
 		const FCasBlockInfo& Info = Kv.Value;
 
-		if (Info.FileSize > Cas.GetMaxBlockSize())
+		if (Info.FileSize > CacheState->Cas.GetMaxBlockSize())
 		{
-			UE_LOGFMT(LogIoStoreOnDemand, Error, "CAS Block {BlockId} has total size {BlockSize} which is greater than max block size {MaxBlockSize}", Kv.Key.Id, Info.FileSize, Cas.GetMaxBlockSize());
+			UE_LOGFMT(LogIoStoreOnDemand, Error, "CAS {CasType} Block {BlockId} has total size {BlockSize} which is greater than max block size {MaxBlockSize}", 
+				LexToString(CacheState->Type), Kv.Key.Id, Info.FileSize, CacheState->Cas.GetMaxBlockSize());
 			ensure(false);
 		}
 
@@ -2878,7 +3089,8 @@
 		}
 		else if (Info.RefSize > Info.FileSize)
 		{
-			UE_LOGFMT(LogIoStoreOnDemand, Error, "CAS Block {BlockId} has RefSize {RefSize} which is greater than total size {BlockSize}", Kv.Key.Id, Info.RefSize, Info.FileSize);
+			UE_LOGFMT(LogIoStoreOnDemand, Error, "CAS {CasType} Block {BlockId} has RefSize {RefSize} which is greater than total size {BlockSize}",
+				LexToString(CacheState->Type), Kv.Key.Id, Info.RefSize, Info.FileSize);
 			ensure(false);
 		}
 
@@ -2893,35 +3105,36 @@
 		}
 	}
 
+	// TODO: IadMultiCAS: analytics need to track CAS type as well
 	FOnDemandInstallCacheStats::OnCacheUsage(
-		MaxCacheSize, TotalCachedBytes, TotalReferencedBlockBytes, ReferencedBytes, FragmentedBytes, OldestBlockAccess);
+		CacheState->MaxCacheSize, TotalCachedBytes, TotalReferencedBlockBytes, ReferencedBytes, FragmentedBytes, OldestBlockAccess);
 
-	const uint64 TotalUncachedBytes = BytesToInstall + PendingChunks.TotalSize;
+	const uint64 TotalUncachedBytes = BytesToInstall + CacheState->PendingChunks.TotalSize;
 	const uint64 TotalRequiredBytes = TotalCachedBytes + TotalUncachedBytes;
-	if (TotalRequiredBytes <= MaxCacheSize)
+	if (TotalRequiredBytes <= CacheState->MaxCacheSize)
 	{
-		UE_LOG(LogIoStoreOnDemand, Verbose, TEXT("Skipping cache purge, MaxCacheSize=%.2lf MiB, CacheSize=%.2lf MiB, ReferencedBlockSize=%.2lf MiB, ReferencedSize=%.2lf MiB, FragmentedBytes=%.2lf MiB, UncachedSize=%.2lf MiB"),
-			ToMiB(MaxCacheSize), ToMiB(TotalCachedBytes), ToMiB(TotalReferencedBlockBytes), ToMiB(ReferencedBytes), ToMiB(FragmentedBytes), ToMiB(TotalUncachedBytes));
+		UE_LOG(LogIoStoreOnDemand, Verbose, TEXT("Skipping cache purge for %s, MaxCacheSize=%.2lf MiB, CacheSize=%.2lf MiB, ReferencedBlockSize=%.2lf MiB, ReferencedSize=%.2lf MiB, FragmentedBytes=%.2lf MiB, UncachedSize=%.2lf MiB"),
+			LexToString(CacheState->Type), ToMiB(CacheState->MaxCacheSize), ToMiB(TotalCachedBytes), ToMiB(TotalReferencedBlockBytes), ToMiB(ReferencedBytes), ToMiB(FragmentedBytes), ToMiB(TotalUncachedBytes));
 		return MakeValue();
 	}
 
-	UE_LOG(LogIoStoreOnDemand, Log, TEXT("Purging install cache, MaxCacheSize=%.2lf MiB, CacheSize=%.2lf MiB, ReferencedBlockSize=%.2lf MiB, ReferencedSize=%.2lf MiB, FragmentedBytes=%.2lf MiB, UncachedSize=%.2lf MiB"),
-		ToMiB(MaxCacheSize), ToMiB(TotalCachedBytes), ToMiB(TotalReferencedBlockBytes), ToMiB(ReferencedBytes), ToMiB(FragmentedBytes), ToMiB(TotalUncachedBytes));
+	UE_LOG(LogIoStoreOnDemand, Log, TEXT("Purging install cache %s, MaxCacheSize=%.2lf MiB, CacheSize=%.2lf MiB, ReferencedBlockSize=%.2lf MiB, ReferencedSize=%.2lf MiB, FragmentedBytes=%.2lf MiB, UncachedSize=%.2lf MiB"),
+		LexToString(CacheState->Type), ToMiB(CacheState->MaxCacheSize), ToMiB(TotalCachedBytes), ToMiB(TotalReferencedBlockBytes), ToMiB(ReferencedBytes), ToMiB(FragmentedBytes), ToMiB(TotalUncachedBytes));
 
-	const uint64	TotalBytesToPurge	= TotalRequiredBytes - MaxCacheSize;
+	const uint64	TotalBytesToPurge	= TotalRequiredBytes - CacheState->MaxCacheSize;
 	uint64			TotalPurgedBytes	= 0;
-	FResult			Result				= PurgeInternal(BlockInfo, TotalBytesToPurge, TotalPurgedBytes);
+	FResult			Result				= PurgeInternal(*CacheState, BlockInfo, TotalBytesToPurge, TotalPurgedBytes);
 
 	if (TotalPurgedBytes > 0)
 	{
-		UE_LOG(LogIoStoreOnDemand, Log, TEXT("Purged %.2lf MiB (%.2lf%%) from install cache"),
-			ToMiB(TotalPurgedBytes), 100.0 * (double(TotalPurgedBytes) / double(TotalCachedBytes)));
+		UE_LOG(LogIoStoreOnDemand, Log, TEXT("Purged %.2lf MiB (%.2lf%%) from install cache %s"),
+			ToMiB(TotalPurgedBytes), 100.0 * (double(TotalPurgedBytes) / double(TotalCachedBytes)), LexToString(CacheState->Type));
 	}
 
 	const uint64 NewCachedBytes = TotalCachedBytes - TotalPurgedBytes;
-	UE_CLOG(NewCachedBytes > MaxCacheSize,
-		LogIoStoreOnDemand, Warning, TEXT("Max install cache size exceeded by %.2lf MiB (%.2lf%%)"),
-			ToMiB(NewCachedBytes - MaxCacheSize), 100.0 * (double(NewCachedBytes - MaxCacheSize) / double(MaxCacheSize)));
+	UE_CLOG(NewCachedBytes > CacheState->MaxCacheSize,
+		LogIoStoreOnDemand, Warning, TEXT("Max install cache size exceeded by %.2lf MiB (%.2lf%%) for cache %s"),
+			ToMiB(NewCachedBytes - CacheState->MaxCacheSize), 100.0 * (double(NewCachedBytes - CacheState->MaxCacheSize) / double(CacheState->MaxCacheSize)), LexToString(CacheState->Type));
 
 	uint64 DefragPurgedBytes = 0;
 	if (Result.HasError() == false && TotalPurgedBytes < TotalBytesToPurge)
@@ -2930,28 +3143,35 @@
 		{
 			// Attempt to defrag
 			const uint64 DefragBytesToPurge = TotalBytesToPurge - TotalPurgedBytes;
-			Result = Defrag(Containers, ChunkEntryIndices, BlockInfo, &DefragBytesToPurge, &DefragPurgedBytes);
+			Result = Defrag(*CacheState, Containers, ChunkEntryIndices, BlockInfo, &DefragBytesToPurge, &DefragPurgedBytes);
 		}
 		else
 		{
-			FInstallCacheErrorContext Ctx = MakeInstallCacheErrorContext(TotalCachedBytes, BytesToInstall);
+			FInstallCacheErrorContext Ctx = MakeInstallCacheErrorContext(*CacheState, TotalCachedBytes, BytesToInstall);
 			Ctx.IoError = FIoStatus(EIoErrorCode::InvalidCode, TEXT("Failed to purge required size from install cache"));
 			Result = MakeError(InstallCachePurgeError::MakeError(MoveTemp(Ctx), UE::UnifiedError::EDetailFilter::All));
 		}
 	}
 
-	FOnDemandInstallCacheStats::OnPurge(Result, MaxCacheSize, NewCachedBytes, TotalBytesToPurge, TotalPurgedBytes + DefragPurgedBytes);
+	FOnDemandInstallCacheStats::OnPurge(Result, CacheState->MaxCacheSize, NewCachedBytes, TotalBytesToPurge, TotalPurgedBytes + DefragPurgedBytes);
 	return Result;
 }
 
-FResult FOnDemandInstallCache::PurgeAllUnreferenced(bool bDefrag, const uint64* BytesToPurge /*= nullptr*/)
+FResult FOnDemandInstallCache::PurgeAllUnreferenced(EOnDemandInstallCasType CasType, bool bDefrag, const uint64* BytesToPurge /*= nullptr*/)
 {
 	using namespace UE::UnifiedError::IoStoreOnDemand;
 
-	TUniqueLock Lock(PurgeDefragMutex);
+	FCacheState* CacheState = GetCacheState(CasType);
+	if (CacheState == nullptr)
+	{
+		return MakeCasError<void>(
+			ECasErrorCode::DeleteBlockFailed, EIoErrorCode::InvalidParameter, FString::Printf(TEXT("CAS %s is not configured"), LexToString(CasType)));
+	}
 
+	TUniqueLock Lock(CacheState->PurgeDefragMutex);
+
 	FCasBlockInfoMap	BlockInfo;
-	const uint64		TotalCachedBytes = Cas.GetBlockInfo(BlockInfo);
+	const uint64		TotalCachedBytes = CacheState->Cas.GetBlockInfo(BlockInfo);
 
 	TArray<FSharedOnDemandContainer>	Containers;
 	TArray<TBitArray<>>					ChunkEntryIndices;
@@ -2960,30 +3180,30 @@
 	check(Containers.Num() == ChunkEntryIndices.Num());
 
 	uint64 ReferencedBytes = 0;
-	AddReferencesToBlocks(Containers, ChunkEntryIndices, BlockInfo, ReferencedBytes);
+	AddReferencesToBlocks(*CacheState, Containers, ChunkEntryIndices, BlockInfo, ReferencedBytes);
 
 	const uint64 TotalReferencedBytes = Algo::TransformAccumulate(BlockInfo,
 		[](const TPair<FCasBlockId, FCasBlockInfo>& Kv) { return (Kv.Value.RefSize > 0) ? Kv.Value.FileSize : uint64(0); },
 		uint64(0));
 
-	UE_LOG(LogIoStoreOnDemand, Log, TEXT("Purging install cache, MaxCacheSize=%.2lf MiB, CacheSize=%.2lf MiB, ReferencedBytes=%.2lf MiB"),
-		ToMiB(MaxCacheSize), ToMiB(TotalCachedBytes), ToMiB(TotalReferencedBytes));
+	UE_LOG(LogIoStoreOnDemand, Log, TEXT("Purging install cache %s, MaxCacheSize=%.2lf MiB, CacheSize=%.2lf MiB, ReferencedBytes=%.2lf MiB"),
+		LexToString(CacheState->Type), ToMiB(CacheState->MaxCacheSize), ToMiB(TotalCachedBytes), ToMiB(TotalReferencedBytes));
 
 	// TODO: should probably default to TotalCachedBytes and not MaxCacheSize
-	const uint64	TotalBytesToPurge	= BytesToPurge ? *BytesToPurge : MaxCacheSize;
+	const uint64	TotalBytesToPurge	= BytesToPurge ? *BytesToPurge : CacheState->MaxCacheSize;
 	uint64			TotalPurgedBytes	= 0;
-	FResult			Result				= PurgeInternal(BlockInfo, TotalBytesToPurge, TotalPurgedBytes);
+	FResult			Result				= PurgeInternal(*CacheState, BlockInfo, TotalBytesToPurge, TotalPurgedBytes);
 
 	if (TotalPurgedBytes > 0)
 	{
-		UE_LOG(LogIoStoreOnDemand, Log, TEXT("Purged %.2lf MiB (%.2lf%%) from install cache"),
-			ToMiB(TotalPurgedBytes), 100.0 * (double(TotalPurgedBytes) / double(TotalCachedBytes)));
+		UE_LOG(LogIoStoreOnDemand, Log, TEXT("Purged %.2lf MiB (%.2lf%%) from install cache %s"),
+			ToMiB(TotalPurgedBytes), 100.0 * (double(TotalPurgedBytes) / double(TotalCachedBytes)), LexToString(CacheState->Type));
 	}
 
 	const uint64 NewCachedBytes = TotalCachedBytes - TotalPurgedBytes;
-	UE_CLOG(NewCachedBytes > MaxCacheSize,
-		LogIoStoreOnDemand, Warning, TEXT("Max install cache size exceeded by %.2lf MiB (%.2lf%%)"),
-			ToMiB(NewCachedBytes - MaxCacheSize), 100.0 * (double(NewCachedBytes - MaxCacheSize) / double(MaxCacheSize)));
+	UE_CLOG(NewCachedBytes > CacheState->MaxCacheSize,
+		LogIoStoreOnDemand, Warning, TEXT("Max install cache size exceeded by %.2lf MiB (%.2lf%%) for cache %s"),
+			ToMiB(NewCachedBytes - CacheState->MaxCacheSize), 100.0 * (double(NewCachedBytes - CacheState->MaxCacheSize) / double(CacheState->MaxCacheSize)), LexToString(CacheState->Type));
 
 	if (BytesToPurge)
 	{
@@ -2991,29 +3211,38 @@
 		{
 			// Attempt to defrag
 			const uint64 DefragBytesToPurge = TotalBytesToPurge - TotalPurgedBytes;
-			Result = Defrag(Containers, ChunkEntryIndices, BlockInfo, &DefragBytesToPurge);
+			Result = Defrag(*CacheState, Containers, ChunkEntryIndices, BlockInfo, &DefragBytesToPurge);
 		}
 		else
 		{
-			FInstallCacheErrorContext Ctx = MakeInstallCacheErrorContext(TotalCachedBytes, TotalBytesToPurge);
+			FInstallCacheErrorContext Ctx = MakeInstallCacheErrorContext(*CacheState, TotalCachedBytes, TotalBytesToPurge);
 			Result = MakeError(InstallCachePurgeError::MakeError(MoveTemp(Ctx), UE::UnifiedError::EDetailFilter::All));
 		}
 	}
 	else if (bDefrag)
 	{
-		Result = Defrag(Containers, ChunkEntryIndices, BlockInfo);
+		Result = Defrag(*CacheState, Containers, ChunkEntryIndices, BlockInfo);
 	}
 
-	FOnDemandInstallCacheStats::OnPurge(Result, MaxCacheSize, NewCachedBytes, TotalBytesToPurge, TotalPurgedBytes);
+	FOnDemandInstallCacheStats::OnPurge(Result, CacheState->MaxCacheSize, NewCachedBytes, TotalBytesToPurge, TotalPurgedBytes);
 	return Result;
 }
 
-FResult FOnDemandInstallCache::DefragAll(const uint64* BytesToFree /*= nullptr*/)
+FResult FOnDemandInstallCache::DefragAll(EOnDemandInstallCasType CasType, const uint64* BytesToFree /*= nullptr*/)
 {
-	TUniqueLock Lock(PurgeDefragMutex);
+	using namespace UE::UnifiedError::IoStoreOnDemand;
 
+	FCacheState* CacheState = GetCacheState(CasType);
+	if (CacheState == nullptr)
+	{
+		return MakeCasError<void>(
+			ECasErrorCode::DeleteBlockFailed, EIoErrorCode::InvalidParameter, FString::Printf(TEXT("CAS %s is not configured"), LexToString(CasType)));
+	}
+
+	TUniqueLock Lock(CacheState->PurgeDefragMutex);
+
 	FCasBlockInfoMap	BlockInfo;
-	const uint64		TotalCachedBytes = Cas.GetBlockInfo(BlockInfo);
+	const uint64		TotalCachedBytes = CacheState->Cas.GetBlockInfo(BlockInfo);
 
 	TArray<FSharedOnDemandContainer>	Containers;
 	TArray<TBitArray<>>					ChunkEntryIndices;
@@ -3022,16 +3251,16 @@
 	check(Containers.Num() == ChunkEntryIndices.Num());
 
 	uint64 ReferencedBytes = 0;
-	AddReferencesToBlocks(Containers, ChunkEntryIndices, BlockInfo, ReferencedBytes);
+	AddReferencesToBlocks(*CacheState, Containers, ChunkEntryIndices, BlockInfo, ReferencedBytes);
 
 	const uint64 TotalReferencedBlockBytes = Algo::TransformAccumulate(BlockInfo,
 		[](const TPair<FCasBlockId, FCasBlockInfo>& Kv) { return (Kv.Value.RefSize > 0) ? Kv.Value.FileSize : uint64(0); },
 		uint64(0));
 
-	UE_LOG(LogIoStoreOnDemand, Log, TEXT("Defragmenting install cache, MaxCacheSize=%.2lf MiB, CacheSize=%.2lf MiB, ReferencedBlockSize=%.2lf MiB, ReferencedSize=%.2lf MiB"),
-		ToMiB(MaxCacheSize), ToMiB(TotalCachedBytes), ToMiB(TotalReferencedBlockBytes), ToMiB(ReferencedBytes));
+	UE_LOG(LogIoStoreOnDemand, Log, TEXT("Defragmenting install cache %s, MaxCacheSize=%.2lf MiB, CacheSize=%.2lf MiB, ReferencedBlockSize=%.2lf MiB, ReferencedSize=%.2lf MiB"),
+		LexToString(CacheState->Type), ToMiB(CacheState->MaxCacheSize), ToMiB(TotalCachedBytes), ToMiB(TotalReferencedBlockBytes), ToMiB(ReferencedBytes));
 
-	return Defrag(Containers, ChunkEntryIndices, BlockInfo, BytesToFree);
+	return Defrag(*CacheState, Containers, ChunkEntryIndices, BlockInfo, BytesToFree);
 }
 
 FResult FOnDemandInstallCache::Verify()
@@ -3056,24 +3285,15 @@
 		}
 	};
 
+	uint32		CorruptChunkCount = 0;
+	uint32		MissingChunkCount = 0;
+	uint32		ReadErrorCount = 0;
+	uint64		TotalVerifiedBytes = 0;
+	uint32		TotalChunkCount = 0;
+
 	TArray<FSharedOnDemandContainer>	Containers = IoStore.GetContainers(EOnDemandContainerFlags::InstallOnDemand);
-	TArray<FCasAddrLocation>			ChunkLocations;
 	TArray<FChunkLookup>				ChunkLookups;
 
-	{
-		TUniqueLock Lock(Cas);
-		ChunkLocations.Reserve(Cas.Lookup.Num());
-		for (const TPair<FCasAddr, FCasLocation>& Kv : Cas.Lookup)
-		{
-			ChunkLocations.Add(FCasAddrLocation
-			{
-				.Addr		= Kv.Key,
-				.Location	= Kv.Value
-			});
-		}
-	}
-	ChunkLocations.Sort();
-
 	ChunkLookups.Reserve(Containers.Num());
 	for (int32 Idx = 0; Idx < Containers.Num(); ++Idx)
 	{
@@ -3088,10 +3308,10 @@
 		}
 	}
 
-	auto FindChunkEntry = [&Containers, &ChunkLookups](const FCasAddr& Addr, int32& OutContainerIndex) -> int32
+	auto FindChunkEntry = [&ChunkLookups](const FCasAddr& Addr, int32& OutContainerIndex) -> int32
 	{
 		OutContainerIndex = INDEX_NONE;
-		for (int32 Idx = 0; Idx < Containers.Num(); ++Idx)
+		for (int32 Idx = 0; Idx < ChunkLookups.Num(); ++Idx)
 		{
 			FChunkLookup& Lookup = ChunkLookups[Idx];
 			if (const int32* EntryIndex = Lookup.AddrToIndex.Find(Addr))
@@ -3104,96 +3324,113 @@
 		return INDEX_NONE;
 	};
 
-	const int32	ChunkCount			= ChunkLocations.Num();
-	uint32		CorruptChunkCount	= 0;
-	uint32		MissingChunkCount	= 0;
-	uint32		ReadErrorCount		= 0;
-	uint64		TotalVerifiedBytes	= 0;
-	FIoBuffer	Chunk(1 << 20);
+	for (FCacheState& CacheState : CacheStates)
+	{
+		TArray<FCasAddrLocation>			ChunkLocations;
 
-	if (ChunkCount == 0)
-	{
-		UE_LOG(LogIoStoreOnDemand, Log, TEXT("Verify skipped, install cache is empty"));
-		return MakeValue();
-	}
+		{
+			TUniqueLock Lock(CacheState.Cas);
+			ChunkLocations.Reserve(CacheState.Cas.Lookup.Num());
+			for (const TPair<FCasAddr, FCasLocation>& Kv : CacheState.Cas.Lookup)
+			{
+				ChunkLocations.Add(FCasAddrLocation
+				{
+					.Addr		= Kv.Key,
+					.Location	= Kv.Value
+				});
+			}
+		}
+		ChunkLocations.Sort();
 
-	UE_LOG(LogIoStoreOnDemand, Log, TEXT("Verifying %d installed chunks..."), ChunkCount);
-	for (int32 ChunkIndex = 0; const FCasAddrLocation& ChunkLocation : ChunkLocations)
-	{
-		TResult<FSharedFileHandle> OpenResult = Cas.OpenRead(ChunkLocation.Location.BlockId);
-		if (OpenResult.HasError())
+		const int32	ChunkCount = ChunkLocations.Num();
+		TotalChunkCount += ChunkCount;
+
+		if (ChunkCount == 0)
 		{
-			UE_LOGFMT(LogIoStoreOnDemand, Error, "Failed to open block {BlockId} for reading ({Error})", ChunkLocation.Location.BlockId.Id, OpenResult.GetError());
-
-			ReadErrorCount++;
-			ChunkIndex++;
+			UE_LOG(LogIoStoreOnDemand, Log, TEXT("Verify skipped CAS %s, install cache is empty"), LexToString(CacheState.Type));
 			continue;
 		}
 
-		int32 ContainerIndex	= INDEX_NONE;
-		int32 EntryIndex		= FindChunkEntry(ChunkLocation.Addr, ContainerIndex);
+		FIoBuffer	Chunk(1 << 20);
 
-		if (EntryIndex == INDEX_NONE)
+		UE_LOG(LogIoStoreOnDemand, Log, TEXT("Verifying %d installed chunks in CAS %s..."), ChunkCount, LexToString(CacheState.Type));
+		for (int32 ChunkIndex = 0; const FCasAddrLocation& ChunkLocation : ChunkLocations)
 		{
-			UE_LOG(LogIoStoreOnDemand, Error, TEXT("Failed to find chunk entry for CAS address '%s'"), *LexToString(ChunkLocation.Addr));
+			TResult<FSharedFileHandle> OpenResult = CacheState.Cas.OpenRead(ChunkLocation.Location.BlockId);
+			if (OpenResult.HasError())
+			{
+				UE_LOGFMT(LogIoStoreOnDemand, Error, "Failed to open block {BlockId} for reading ({Error})", ChunkLocation.Location.BlockId.Id, OpenResult.GetError());
 
-			MissingChunkCount++;
+				ReadErrorCount++;
+				ChunkIndex++;
+				continue;
+			}
+
+			int32 ContainerIndex	= INDEX_NONE;
+			int32 EntryIndex		= FindChunkEntry(ChunkLocation.Addr, ContainerIndex);
+
+			if (EntryIndex == INDEX_NONE)
+			{
+				UE_LOG(LogIoStoreOnDemand, Error, TEXT("Failed to find chunk entry for CAS address '%s'"), *LexToString(ChunkLocation.Addr));
+
+				MissingChunkCount++;
+				ChunkIndex++;
+				continue;
+			}
+
+			const FSharedOnDemandContainer& Container	= Containers[ContainerIndex];
+			const FIoChunkId& ChunkId					= Container->ChunkIds[EntryIndex];
+			const FOnDemandChunkEntry& ChunkEntry		= Container->ChunkEntries[EntryIndex];
+			FSharedFileHandle FileHandle				= OpenResult.GetValue();
+			const int64 ChunkSize						= ChunkEntry.GetDiskSize();
+			TotalVerifiedBytes							+= ChunkSize;
+
+			if (int64(Chunk.GetSize()) < ChunkSize)
+			{
+				Chunk = FIoBuffer(ChunkSize);
+			}
+
+			if (FileHandle->Seek(int64(ChunkLocation.Location.BlockOffset)) == false)
+			{
+				UE_LOG(LogIoStoreOnDemand, Error, TEXT("Chunk %d/%d SEEK FAILED, Container='%s', ChunkId='%s', ChunkSize=%" INT64_FMT ", Hash='%s', Block=%u, BlockOffset=%u"),
+					ChunkIndex + 1, ChunkCount, *Container->Name(), *LexToString(ChunkId), ChunkSize, *LexToString(ChunkEntry.Hash),
+					ChunkLocation.Location.BlockId.Id, ChunkLocation.Location.BlockOffset);
+
+				ReadErrorCount++;
+				ChunkIndex++;
+				continue;
+			}
+
+			if (FileHandle->Read(reinterpret_cast<uint8*>(Chunk.GetData()), ChunkSize) == false)
+			{
+				UE_LOG(LogIoStoreOnDemand, Error, TEXT("Chunk %d/%d READ FAILED, Container='%s', ChunkId='%s', ChunkSize=%" INT64_FMT ", Hash='%s', Block=%u, BlockOffset=%u"),
+					ChunkIndex + 1, ChunkCount, *Container->Name(), *LexToString(ChunkId), ChunkSize, *LexToString(ChunkEntry.Hash),
+					ChunkLocation.Location.BlockId.Id, ChunkLocation.Location.BlockOffset);
+
+				ReadErrorCount++;
+				ChunkIndex++;
+				continue;
+			}
+
+			const FIoHash ChunkHash = FIoHash::HashBuffer(Chunk.GetView().Left(ChunkSize));
+
+			if (ChunkHash == ChunkEntry.Hash)
+			{
+				UE_LOG(LogIoStoreOnDemand, VeryVerbose, TEXT("Chunk %d/%d OK, Container='%s', ChunkId='%s', ChunkSize=%" INT64_FMT ", Hash='%s', Block=%u, BlockOffset=%u"),
+					ChunkIndex + 1, ChunkCount, *Container->Name(), *LexToString(ChunkId), ChunkSize, *LexToString(ChunkEntry.Hash),
+					ChunkLocation.Location.BlockId.Id, ChunkLocation.Location.BlockOffset);
+			}
+			else
+			{
+				UE_LOG(LogIoStoreOnDemand, Error, TEXT("Chunk %d/%d CORRUPT, Container='%s', ChunkId='%s', ChunkSize=%" INT64_FMT ", Hash='%s', ActualHash='%s', Block=%u, BlockOffset=%u"),
+					ChunkIndex + 1, ChunkCount, *Container->Name(), *LexToString(ChunkId), ChunkSize, *LexToString(ChunkEntry.Hash), *LexToString(ChunkHash),
+					ChunkLocation.Location.BlockId.Id, ChunkLocation.Location.BlockOffset);
+
+				CorruptChunkCount++;
+			}
+
 			ChunkIndex++;
-			continue;
 		}
-
-		const FSharedOnDemandContainer& Container	= Containers[ContainerIndex];
-		const FIoChunkId& ChunkId					= Container->ChunkIds[EntryIndex];
-		const FOnDemandChunkEntry& ChunkEntry		= Container->ChunkEntries[EntryIndex];
-		FSharedFileHandle FileHandle				= OpenResult.GetValue();
-		const int64 ChunkSize						= ChunkEntry.GetDiskSize();
-		TotalVerifiedBytes							+= ChunkSize;
-
-		if (int64(Chunk.GetSize()) < ChunkSize)
-		{
-			Chunk = FIoBuffer(ChunkSize);
-		}
-
-		if (FileHandle->Seek(int64(ChunkLocation.Location.BlockOffset)) == false)
-		{
-			UE_LOG(LogIoStoreOnDemand, Error, TEXT("Chunk %d/%d SEEK FAILED, Container='%s', ChunkId='%s', ChunkSize=%" INT64_FMT ", Hash='%s', Block=%u, BlockOffset=%u"),
-				ChunkIndex + 1, ChunkCount, *Container->Name(), *LexToString(ChunkId), ChunkSize, *LexToString(ChunkEntry.Hash),
-				ChunkLocation.Location.BlockId.Id, ChunkLocation.Location.BlockOffset);
-
-			ReadErrorCount++;
-			ChunkIndex++;
-			continue;
-		}
-
-		if (FileHandle->Read(reinterpret_cast<uint8*>(Chunk.GetData()), ChunkSize) == false)
-		{
-			UE_LOG(LogIoStoreOnDemand, Error, TEXT("Chunk %d/%d READ FAILED, Container='%s', ChunkId='%s', ChunkSize=%" INT64_FMT ", Hash='%s', Block=%u, BlockOffset=%u"),
-				ChunkIndex + 1, ChunkCount, *Container->Name(), *LexToString(ChunkId), ChunkSize, *LexToString(ChunkEntry.Hash),
-				ChunkLocation.Location.BlockId.Id, ChunkLocation.Location.BlockOffset);
-
-			ReadErrorCount++;
-			ChunkIndex++;
-			continue;
-		}
-
-		const FIoHash ChunkHash = FIoHash::HashBuffer(Chunk.GetView().Left(ChunkSize));
-
-		if (ChunkHash == ChunkEntry.Hash)
-		{
-			UE_LOG(LogIoStoreOnDemand, VeryVerbose, TEXT("Chunk %d/%d OK, Container='%s', ChunkId='%s', ChunkSize=%" INT64_FMT ", Hash='%s', Block=%u, BlockOffset=%u"),
-				ChunkIndex + 1, ChunkCount, *Container->Name(), *LexToString(ChunkId), ChunkSize, *LexToString(ChunkEntry.Hash),
-				ChunkLocation.Location.BlockId.Id, ChunkLocation.Location.BlockOffset);
-		}
-		else
-		{
-			UE_LOG(LogIoStoreOnDemand, Error, TEXT("Chunk %d/%d CORRUPT, Container='%s', ChunkId='%s', ChunkSize=%" INT64_FMT ", Hash='%s', ActualHash='%s', Block=%u, BlockOffset=%u"),
-				ChunkIndex + 1, ChunkCount, *Container->Name(), *LexToString(ChunkId), ChunkSize, *LexToString(ChunkEntry.Hash), *LexToString(ChunkHash),
-				ChunkLocation.Location.BlockId.Id, ChunkLocation.Location.BlockOffset);
-
-			CorruptChunkCount++;
-		}
-
-		ChunkIndex++;
 	}
 
 	if (CorruptChunkCount > 0 || MissingChunkCount > 0 || ReadErrorCount > 0)
@@ -3220,7 +3457,7 @@
 	}
 
 	UE_LOG(LogIoStoreOnDemand, Log, TEXT("Successfully verified %d chunk(s) of total %.2lf MiB"),
-		ChunkCount, ToMiB(TotalVerifiedBytes));
+		TotalChunkCount, ToMiB(TotalVerifiedBytes));
 
 	return MakeValue();
 }
@@ -3236,7 +3473,7 @@
 		{
 			UE_LOG(LogIoStoreOnDemand, Log, TEXT("Simulating critical install cache error"));
 
-			FCasJournal::FTransaction Transaction = FCasJournal::Begin(GetJournalFilename());
+			FCasJournal::FTransaction Transaction = FCasJournal::Begin(EOnDemandInstallCasType::General, GetJournalFilename());
 			Transaction.CriticalError(FCasJournal::EErrorCode::Simulated);
 			if (FResult Result = FCasJournal::Commit(MoveTemp(Transaction)); Result.HasError())
 			{
@@ -3270,9 +3507,15 @@
 			ECasErrorCode::InitializeFailed, EIoErrorCode::WriteError, FString::Printf(TEXT("Failed to create cache directory '%s'"), *CacheDirectory));
 	}
 
-	if (FResult Result = Cas.Initialize(CacheDirectory); Result.HasError())
+	for (FCacheState& CacheState : CacheStates)
 	{
-		return Result;
+		FString CasRootDir = CacheState.Cas.GetRootDirectory();
+		if (FResult Result = CacheState.Cas.Initialize(CasRootDir); Result.HasError())
+		{
+			return Result;
+		}
+
+		CacheState.CurrentBlock = FCasBlockId::Invalid;
 	}
 
 	const FString JournalFile = GetJournalFilename();
@@ -3281,89 +3524,91 @@
 		return Result;
 	}
 
-	CurrentBlock = FCasBlockId::Invalid;
-
 	UE_LOG(LogIoStoreOnDemand, Log, TEXT("Created CAS journal '%s'"), *JournalFile);
 	return MakeValue();
 }
 
 FResult FOnDemandInstallCache::InitialVerify()
 {
-	// Verify the blocks on disk with the current state of the CAS
+	// TODO: IadMultiCAS: there's no way to verify each CAS in isolation right now, so if we detect an error everything will be reset
+	for (FCacheState& CacheState : CacheStates)
 	{
-		TArray<FCasAddr> RemovedChunks;
-		TMap<FCasBlockId, uint64> OverBudgetBlocks;
-		FResult Verify = Cas.Verify(RemovedChunks, OverBudgetBlocks);
+		// Verify the blocks on disk with the current state of the CAS
+		{
+			TArray<FCasAddr> RemovedChunks;
+			TMap<FCasBlockId, uint64> OverBudgetBlocks;
+			FResult Verify = CacheState.Cas.Verify(RemovedChunks, OverBudgetBlocks);
 
-		uint64 TotalBlockBytesOverBudget = 0;
-		for (const TPair<FCasBlockId, uint64>& Kv : OverBudgetBlocks)
-		{
-			if (ensure(Kv.Value > Cas.GetMaxBlockSize()))
+			uint64 TotalBlockBytesOverBudget = 0;
+			for (const TPair<FCasBlockId, uint64>& Kv : OverBudgetBlocks)
 			{
-				TotalBlockBytesOverBudget += (Kv.Value - Cas.GetMaxBlockSize());
+				if (ensure(Kv.Value > CacheState.Cas.GetMaxBlockSize()))
+				{
+					TotalBlockBytesOverBudget += (Kv.Value - CacheState.Cas.GetMaxBlockSize());
+				}
+			}
+
+			FOnDemandInstallCacheStats::OnCasVerification(Verify, RemovedChunks.Num(), TotalBlockBytesOverBudget);
+
+			if (Verify.HasError())
+			{
+				// Remove all entries that doesn't have a valid cache block 
+				FCasJournal::FTransaction Transaction = FCasJournal::Begin(CacheState.Type, GetJournalFilename());
+				for (const FCasAddr& Addr : RemovedChunks)
+				{
+					Transaction.ChunkLocation(FCasLocation::Invalid, Addr);
+				}
+
+				if (FResult Result = FCasJournal::Commit(MoveTemp(Transaction)); Result.HasError())
+				{
+					return Result;
+				}
 			}
 		}
 
-		FOnDemandInstallCacheStats::OnCasVerification(Verify, RemovedChunks.Num(), TotalBlockBytesOverBudget);
+		// Check if the cache is over budget
+		{
+			FCasBlockInfoMap	BlockInfo;
+			uint64				CacheSize = CacheState.Cas.GetBlockInfo(BlockInfo);
 
-		if (Verify.HasError())
+			if (CacheSize > CacheState.MaxCacheSize)
+			{
+				const uint64	TotalBytesToPurge = CacheSize - CacheState.MaxCacheSize;
+				uint64			TotalPurgedBytes = 0;
+
+				UE_LOG(LogIoStoreOnDemand, Warning, TEXT("Cache size is greater than disk quota for CAS %s - Purging install cache, MaxCacheSize=%.2lf MiB, TotalSize=%.2lf MiB, TotalBytesToPurge=%.2lf MiB"),
+					LexToString(CacheState.Type), ToMiB(CacheState.MaxCacheSize), ToMiB(CacheSize), ToMiB(TotalBytesToPurge));
+
+				if (FResult Result = PurgeInternal(CacheState, BlockInfo, TotalBytesToPurge, TotalPurgedBytes); Result.HasError())
+				{
+					const FString ErrorMessage = FString::Printf(TEXT("Failed to purge %.2lf MiB from CAS %s reason '%s'"),
+						ToMiB(TotalBytesToPurge), LexToString(CacheState.Type), *LexToString(Result.GetError()));
+
+					UE_LOG(LogIoStoreOnDemand, Error, TEXT("%s"), *ErrorMessage);
+					return Result;
+				}
+
+				if (TotalPurgedBytes < TotalBytesToPurge)
+				{
+					// This should never happen since we don't have any referenced cache blocks at startup
+					const FString ErrorMessage = FString::Printf(TEXT("Failed to purge %.2lf MiB from CAS %s. Actually purged %.2lf MiB"),
+						ToMiB(TotalBytesToPurge), LexToString(CacheState.Type), ToMiB(TotalPurgedBytes));
+
+					UE_LOG(LogIoStoreOnDemand, Error, TEXT("%s"), *ErrorMessage);
+					return MakeError(UE::UnifiedError::IoStoreOnDemand::InstallCachePurgeError::MakeError());
+				}
+
+				UE_LOG(LogIoStoreOnDemand, Log, TEXT("Successfully purged %.2lf MiB from CAS %s"), ToMiB(TotalPurgedBytes), LexToString(CacheState.Type));
+			}
+		}
+
+		if (!CacheState.Cas.ContainsBlock(CacheState.CurrentBlock))
 		{
-			// Remove all entries that doesn't have a valid cache block 
-			FCasJournal::FTransaction Transaction = FCasJournal::Begin(GetJournalFilename());
-			for (const FCasAddr& Addr : RemovedChunks)
+			const FCasBlockId OldCurrentBlock = CacheState.CurrentBlock.exchange(FCasBlockId::Invalid);
+			if (OldCurrentBlock != FCasBlockId::Invalid)
 			{
-				Transaction.ChunkLocation(FCasLocation::Invalid, Addr);
-			}
-
-			if (FResult Result = FCasJournal::Commit(MoveTemp(Transaction)); Result.HasError())
-			{
-				return Result;
+				UE_LOG(LogIoStoreOnDemand, Warning, TEXT("Current Block %u was not found for CAS %s, resetting"), OldCurrentBlock.Id, LexToString(CacheState.Type));
 			}
 		}
 	}
 
-	// Check if the cache is over budget
-	{
-		FCasBlockInfoMap	BlockInfo;
-		uint64				CacheSize = Cas.GetBlockInfo(BlockInfo);
-
-		if (CacheSize > MaxCacheSize)
-		{
-			const uint64	TotalBytesToPurge = CacheSize - MaxCacheSize;
-			uint64			TotalPurgedBytes = 0;
-
-			UE_LOG(LogIoStoreOnDemand, Warning, TEXT("Cache size is greater than disk quota - Purging install cache, MaxCacheSize=%.2lf MiB, TotalSize=%.2lf MiB, TotalBytesToPurge=%.2lf MiB"),
-				ToMiB(MaxCacheSize), ToMiB(CacheSize), ToMiB(TotalBytesToPurge));
-
-			if (FResult Result = PurgeInternal(BlockInfo, TotalBytesToPurge, TotalPurgedBytes); Result.HasError())
-			{
-				const FString ErrorMessage = FString::Printf(TEXT("Failed to purge %.2lf MiB from install cache reason '%s'"),
-					ToMiB(TotalBytesToPurge), *LexToString(Result.GetError()));
-
-				UE_LOG(LogIoStoreOnDemand, Error, TEXT("%s"), *ErrorMessage);
-				return Result;
-			}
-
-			if (TotalPurgedBytes < TotalBytesToPurge)
-			{
-				// This should never happen since we don't have any referenced cache blocks at startup
-				const FString ErrorMessage = FString::Printf(TEXT("Failed to purge %.2lf MiB from install cache. Actually purged %.2lf MiB from install cache"),
-					ToMiB(TotalBytesToPurge), ToMiB(TotalPurgedBytes));
-
-				UE_LOG(LogIoStoreOnDemand, Error, TEXT("%s"), *ErrorMessage);
-				return MakeError(UE::UnifiedError::IoStoreOnDemand::InstallCachePurgeError::MakeError());
-			}
-
-			UE_LOG(LogIoStoreOnDemand, Log, TEXT("Successfully purged %.2lf MiB from install cache"), ToMiB(TotalPurgedBytes));
-		}
-	}
-
-	if (!Cas.ContainsBlock(CurrentBlock))
-	{
-		const FCasBlockId OldCurrentBlock = CurrentBlock.exchange(FCasBlockId::Invalid);
-		if (OldCurrentBlock != FCasBlockId::Invalid)
-		{
-			UE_LOG(LogIoStoreOnDemand, Warning, TEXT("Current Block %u was not found, resetting"), OldCurrentBlock.Id);
-		}
-	}
-
@@ -3371,6 +3616,7 @@
 }
 
 void FOnDemandInstallCache::AddReferencesToBlocks(
+	const FCacheState& CacheState,
 	const TArray<FSharedOnDemandContainer>& Containers, 
 	const TArray<TBitArray<>>& ChunkEntryIndices, 
 	FCasBlockInfoMap& BlockInfoMap, 
@@ -3414,13 +3660,13 @@
 			const uint64 ChunkDiskSize = Entry.GetDiskSize();
 			OutTotalReferencedBytes += ChunkDiskSize;
 
-			if (FCasLocation Loc = Cas.FindChunk(Entry.Hash); Loc.IsValid())
+			if (FCasLocation Loc = CacheState.Cas.FindChunk(Entry.Hash); Loc.IsValid())
 			{
 				FCasBlockInfo* BlockInfo = BlockInfoMap.Find(Loc.BlockId);
 				if (!BlockInfo)
 				{
-					UE_LOG(LogIoStoreOnDemand, Error, TEXT("Failed to find CAS block info for referenced chunk, ChunkId='%s', Container='%s'"),
-						*LexToString(Container->ChunkIds[EntryIndex]), *Container->Name());
+					UE_LOG(LogIoStoreOnDemand, Error, TEXT("Failed to find CAS block info for referenced chunk, CAS=%s, ChunkId='%s', Container='%s'"),
+						LexToString(CacheState.Type), *LexToString(Container->ChunkIds[EntryIndex]), *Container->Name());
 				}
 				else
 				{
@@ -3432,12 +3678,13 @@
 				// Check pending list
 				bool bChunkIsPending = false;
 				{
-					TSharedLock SharedLock(PendingChunks.SharedMutex);
-					bChunkIsPending = PendingChunks.ContainsChunk(Entry.Hash);
+					TSharedLock SharedLock(CacheState.PendingChunks.SharedMutex);
+					bChunkIsPending = CacheState.PendingChunks.ContainsChunk(Entry.Hash);
 				}
 
-				UE_CLOG(!bChunkIsPending, LogIoStoreOnDemand, Error, TEXT("Failed to find CAS location or pending chunk for chunk reference, ChunkId='%s', Container='%s'"),
-					*LexToString(Container->ChunkIds[EntryIndex]), *Container->Name());
+				// TODO: IadMultiCAS: this probably is no longer valid with multiple CAS.  Its likely that the referenced chunk is just in another CAS
+				UE_CLOG(!bChunkIsPending, LogIoStoreOnDemand, Error, TEXT("Failed to find CAS location or pending chunk for chunk reference, CAS=%s, ChunkId='%s', Container='%s'"),
+					LexToString(CacheState.Type), *LexToString(Container->ChunkIds[EntryIndex]), *Container->Name());
 
 				ensure(bChunkIsPending);
 			}
@@ -3449,7 +3696,11 @@
 	return;
 }
 
-FResult FOnDemandInstallCache::PurgeInternal(FCasBlockInfoMap& BlockInfo, const uint64 TotalBytesToPurge, uint64& OutTotalPurgedBytes)
+FResult FOnDemandInstallCache::PurgeInternal(
+	FCacheState& CacheState, 
+	FCasBlockInfoMap& BlockInfo, 
+	const uint64 TotalBytesToPurge, 
+	uint64& OutTotalPurgedBytes)
 {
 	BlockInfo.ValueSort([](const FCasBlockInfo& LHS, const FCasBlockInfo& RHS)
 	{
@@ -3467,14 +3718,14 @@
 			continue;
 		}
 
-		FCasJournal::FTransaction	Transaction = FCasJournal::Begin(GetJournalFilename());
+		FCasJournal::FTransaction	Transaction = FCasJournal::Begin(CacheState.Type, GetJournalFilename());
 
 		auto HandleChunkRemoved = [&Transaction](const FCasAddr& Addr)
 		{
 			Transaction.ChunkLocation(FCasLocation::Invalid, Addr);
 		};
 
-		if (FResult Result = Cas.DeleteBlock(BlockId, HandleChunkRemoved); Result.HasError())
+		if (FResult Result = CacheState.Cas.DeleteBlock(BlockId, HandleChunkRemoved); Result.HasError())
 		{
 			return Result;
 		}
@@ -3484,7 +3735,7 @@
 
 		// This should be the only thread writing to CurrentBlock
 		FCasBlockId MaybeCurrentBlock = BlockId;
-		CurrentBlock.compare_exchange_strong(MaybeCurrentBlock, FCasBlockId::Invalid);
+		CacheState.CurrentBlock.compare_exchange_strong(MaybeCurrentBlock, FCasBlockId::Invalid);
 
 		OutTotalPurgedBytes += Info.FileSize;
 
@@ -3506,6 +3757,7 @@
 }
 
 FResult FOnDemandInstallCache::Defrag(
+	FCacheState& CacheState,
 	const TArray<FSharedOnDemandContainer>& Containers,
 	const TArray<TBitArray<>>& ChunkEntryIndices,
 	FCasBlockInfoMap& BlockInfo, 
@@ -3593,7 +3845,7 @@
 					bPossibleToFreeBytes = true;
 				}
 			}
-			else if (Info.FileSize < Cas.GetMinBlockSize())
+			else if (Info.FileSize < CacheState.Cas.GetMinBlockSize())
 			{
 				// Block is too small whether or not its fragmented
 				if (ensure(Info.RefSize <= Info.FileSize))
@@ -3611,7 +3863,7 @@
 		{
 			FStringView ErrorMessage(TEXTVIEW("Failed to defrag the install cache due to too much data being referenced by the game"));
 			UE_LOGFMT(LogIoStoreOnDemand, Error, "{ErrorMessage}", ErrorMessage);
-			FInstallCacheErrorContext Ctx = MakeInstallCacheErrorContext(TotalCachedBytes, *TotalBytesToFree);
+			FInstallCacheErrorContext Ctx = MakeInstallCacheErrorContext(CacheState, TotalCachedBytes, *TotalBytesToFree);
 			Ctx.IoError = FIoStatus(EIoErrorCode::OutOfDiskSpace, ErrorMessage);
 			FResult Result = MakeError(InstallCacheDefragError::MakeError(MoveTemp(Ctx), UE::UnifiedError::EDetailFilter::All));
 			FOnDemandInstallCacheStats::OnDefrag(Result, TotalFragmentedBytes, TotalCachedBytes, BlocksToDefragFragmentedSize, BlocksToDefragTotalSize);
@@ -3638,7 +3890,7 @@
 
 				BlocksToDefrag.Emplace(BlockId, Info);
 			}
-			else if (Info.FileSize < Cas.GetMinBlockSize())
+			else if (Info.FileSize < CacheState.Cas.GetMinBlockSize())
 			{
 				// Block is too small whether or not its fragmented
 				if (ensure(Info.RefSize <= Info.FileSize))
@@ -3655,7 +3907,7 @@
 		if (BlocksToDefrag.IsEmpty())
 		{
 			// Already defragged
-			UE_LOG(LogIoStoreOnDemand, Display, TEXT("Cache not fragmented."));
+			UE_LOG(LogIoStoreOnDemand, Display, TEXT("Cache %s not fragmented."), LexToString(CacheState.Type));
 			FResult Result = MakeValue();
 			FOnDemandInstallCacheStats::OnDefrag(Result, TotalFragmentedBytes, TotalCachedBytes, BlocksToDefragFragmentedSize, BlocksToDefragTotalSize);
 			return Result;
@@ -3665,19 +3917,19 @@
 	UE_LOG(LogIoStoreOnDemand, Display, TEXT("Defrag found %" UINT64_FMT " fragmented bytes of %" UINT64_FMT " total bytes in %i blocks."), 
 		BlocksToDefragFragmentedSize, BlocksToDefragTotalSize, BlocksToDefrag.Num()); 
 
-	if (TotalCachedBytes > MaxCacheSize)
+	if (TotalCachedBytes > CacheState.MaxCacheSize)
 	{
 		// Ruh-Roh! There's not enough of the disk quota left to run a defrag!
-		const FString ErrorMsg = FString::Printf(TEXT("Cache size is greater than disk quota - Cannot Defragment!, MaxCacheSize=%.2lf MiB"),
-			ToMiB(MaxCacheSize));
+		const FString ErrorMsg = FString::Printf(TEXT("Cache size is greater than disk quota for %s - Cannot Defragment!, MaxCacheSize=%.2lf MiB"),
+			LexToString(CacheState.Type), ToMiB(CacheState.MaxCacheSize));
 		UE_LOG(LogIoStoreOnDemand, Error, TEXT("%s"), *ErrorMsg);
 
-		FInstallCacheErrorContext Ctx = MakeInstallCacheErrorContext(TotalCachedBytes, TotalBytesToFree ? *TotalBytesToFree : TotalFragmentedBytes);
+		FInstallCacheErrorContext Ctx = MakeInstallCacheErrorContext(CacheState, TotalCachedBytes, TotalBytesToFree ? *TotalBytesToFree : TotalFragmentedBytes);
 		Ctx.IoError = FIoStatus(EIoErrorCode::OutOfDiskSpace, ErrorMsg);
 		FResult Result = MakeError(InstallCacheDefragError::MakeError(MoveTemp(Ctx), UE::UnifiedError::EDetailFilter::All));
 
 		// Append a critical error entry to clear the cache at next startup
-		FCasJournal::FTransaction Transaction = FCasJournal::Begin(GetJournalFilename());
+		FCasJournal::FTransaction Transaction = FCasJournal::Begin(CacheState.Type, GetJournalFilename());
 		Transaction.CriticalError(FCasJournal::EErrorCode::DefragOutOfDiskSpace);
 		if (FResult CommitResult = FCasJournal::Commit(MoveTemp(Transaction)); CommitResult.HasError())
 		{
@@ -3693,7 +3945,7 @@
 	// - If we defrag the current block itself, then we would need additional tracking so we don't lose any chunks moved into it.
 	// - Additionally, this would also depend on the order blocks are defragged.
 	// This should be the only thread writing to CurrentBlock.
-	CurrentBlock = FCasBlockId::Invalid;
+	CacheState.CurrentBlock = FCasBlockId::Invalid;
 
 	uint64 TotalRefBytes = 0;
 
@@ -3709,7 +3961,7 @@
 				continue;
 			}
 
-			if (FCasLocation Loc = Cas.FindChunk(Entry.Hash); Loc.IsValid())
+			if (FCasLocation Loc = CacheState.Cas.FindChunk(Entry.Hash); Loc.IsValid())
 			{
 				if (FDefragBlock* DefragBlock = Algo::FindBy(BlocksToDefrag, Loc.BlockId, &FDefragBlock::BlockId))
 				{
@@ -3743,12 +3995,12 @@
 	{
 		if (DefragBlock.ReferencedChunks.IsEmpty() == false)
 		{
-			TResult<FSharedFileHandle> FileOpenResult = Cas.OpenRead(DefragBlock.BlockId);
+			TResult<FSharedFileHandle> FileOpenResult = CacheState.Cas.OpenRead(DefragBlock.BlockId);
 			if (FileOpenResult.HasError())
 			{
 				UE_LOGFMT(LogIoStoreOnDemand, Error, "Defrag failed to open CAS block for reading {Error}", FileOpenResult.GetError());
 
-				FInstallCacheErrorContext Ctx = MakeInstallCacheErrorContext(TotalCachedBytes, TotalBytesToFree ? *TotalBytesToFree : TotalFragmentedBytes);
+				FInstallCacheErrorContext Ctx = MakeInstallCacheErrorContext(CacheState, TotalCachedBytes, TotalBytesToFree ? *TotalBytesToFree : TotalFragmentedBytes);
 				FileOpenResult.GetError().PushErrorContext(MoveTemp(Ctx));
 				FResult Result = MakeError(FileOpenResult.StealError());
 				FOnDemandInstallCacheStats::OnDefrag(Result, TotalFragmentedBytes, TotalCachedBytes, BlocksToDefragFragmentedSize, BlocksToDefragTotalSize);
@@ -3787,7 +4039,7 @@
 
 				if (ReadResult.HasError())
 				{
-					FInstallCacheErrorContext Ctx = MakeInstallCacheErrorContext(TotalCachedBytes, TotalBytesToFree ? *TotalBytesToFree : TotalFragmentedBytes);
+					FInstallCacheErrorContext Ctx = MakeInstallCacheErrorContext(CacheState, TotalCachedBytes, TotalBytesToFree ? *TotalBytesToFree : TotalFragmentedBytes);
 					ReadResult.GetError().PushErrorContext(MoveTemp(Ctx));
 					FResult Result = MakeError(ReadResult.StealError());
 					FOnDemandInstallCacheStats::OnDefrag(Result, TotalFragmentedBytes, TotalCachedBytes, BlocksToDefragFragmentedSize, BlocksToDefragTotalSize);
@@ -3801,14 +4053,14 @@
 						*LexToString(ReffedChunk.ChunkId), DefragBlock.BlockId.Id, ReffedChunk.BlockOffset);
 
 					// Append a critical error entry to clear the cache at next startup
-					FCasJournal::FTransaction Transaction = FCasJournal::Begin(GetJournalFilename());
+					FCasJournal::FTransaction Transaction = FCasJournal::Begin(CacheState.Type, GetJournalFilename());
 					Transaction.CriticalError(FCasJournal::EErrorCode::DefragHashMismatch);
 					if (FResult CommitResult = FCasJournal::Commit(MoveTemp(Transaction)); CommitResult.HasError())
 					{
 						UE_LOGFMT(LogIoStoreOnDemand, Error, "Failed to commit critical errors to CAS journal ({Error})", CommitResult.GetError());
 					}
 
-					if (FResult Result = FlushPendingChunks(DefragPendingChunks, DefragBlock.LastAccess, DefragBlock.LastModification); Result.HasError())
+					if (FResult Result = FlushPendingChunks(CacheState, DefragPendingChunks, DefragBlock.LastAccess, DefragBlock.LastModification); Result.HasError())
 					{
 						FOnDemandInstallCacheStats::OnDefrag(Result, TotalFragmentedBytes, TotalCachedBytes, BlocksToDefragFragmentedSize, BlocksToDefragTotalSize);
 						return Result;
@@ -3828,7 +4080,7 @@
 
 				if (DefragPendingChunks.TotalSize > FPendingChunks::MaxPendingBytes)
 				{
-					if (FResult Result = FlushPendingChunks(DefragPendingChunks, DefragBlock.LastAccess, DefragBlock.LastModification); Result.HasError())
+					if (FResult Result = FlushPendingChunks(CacheState, DefragPendingChunks, DefragBlock.LastAccess, DefragBlock.LastModification); Result.HasError())
 					{
 						FOnDemandInstallCacheStats::OnDefrag(Result, TotalFragmentedBytes, TotalCachedBytes, BlocksToDefragFragmentedSize, BlocksToDefragTotalSize);
 						return Result;
@@ -3841,7 +4093,7 @@
 
 			FileHandle.Reset();
 
-			if (FResult Result = FlushPendingChunks(DefragPendingChunks, DefragBlock.LastAccess, DefragBlock.LastModification); Result.HasError())
+			if (FResult Result = FlushPendingChunks(CacheState, DefragPendingChunks, DefragBlock.LastAccess, DefragBlock.LastModification); Result.HasError())
 			{
 				FOnDemandInstallCacheStats::OnDefrag(Result, TotalFragmentedBytes, TotalCachedBytes, BlocksToDefragFragmentedSize, BlocksToDefragTotalSize);
 				return Result;
@@ -3849,7 +4101,7 @@
 			check(DefragPendingChunks.IsEmpty());
 		}
 
-		FCasJournal::FTransaction Transaction = FCasJournal::Begin(GetJournalFilename());
+		FCasJournal::FTransaction Transaction = FCasJournal::Begin(CacheState.Type, GetJournalFilename());
 
 		auto HandleChunkRemoved = [&Transaction](const FCasAddr& Addr)
 		{
@@ -3858,7 +4110,7 @@
 
 		// Flushing should overwrite the lookup info for the cas addr to point at the new block.
 		// Can now remove the old block
-		if (FResult Result = Cas.DeleteBlock(DefragBlock.BlockId, HandleChunkRemoved); Result.HasError())
+		if (FResult Result = CacheState.Cas.DeleteBlock(DefragBlock.BlockId, HandleChunkRemoved); Result.HasError())
 		{
 			FOnDemandInstallCacheStats::OnDefrag(Result, TotalFragmentedBytes, TotalCachedBytes, BlocksToDefragFragmentedSize, BlocksToDefragTotalSize);
 			return Result;
@@ -3876,8 +4128,8 @@
 		}
 	}
 
-	UE_LOG(LogIoStoreOnDemand, Display, TEXT("Defrag removed %" UINT64_FMT " fragmented bytes of %" UINT64_FMT " total bytes in %i blocks."),
-		BlocksToDefragTotalSize, BlocksToDefragTotalSize, BlocksToDefrag.Num());
+	UE_LOG(LogIoStoreOnDemand, Display, TEXT("Defrag removed %" UINT64_FMT " fragmented bytes of %" UINT64_FMT " total bytes in %i blocks from CAS %s."),
+		BlocksToDefragTotalSize, BlocksToDefragTotalSize, BlocksToDefrag.Num(), LexToString(CacheState.Type));
 
 	if (OutTotalFreedBytes)
 	{
@@ -3889,43 +4141,65 @@
 	return Result;
 }
 
-FResult FOnDemandInstallCache::Flush()
+FResult FOnDemandInstallCache::Flush(EOnDemandInstallCasType CasType)
 {
-	if (FResult Result = FlushPendingChunks(PendingChunks); Result.HasError())
+	using namespace UE::UnifiedError::IoStoreOnDemand;
+
+	FCacheState* CacheState = GetCacheState(CasType);
+	if (CacheState == nullptr)
+	{
+		return MakeCasError<void>(
+			ECasErrorCode::WriteBlockFailed, EIoErrorCode::InvalidParameter, FString::Printf(TEXT("CAS %s is not configured"), LexToString(CasType)));
+	}
+
+	if (FResult Result = FlushPendingChunks(*CacheState, CacheState->PendingChunks); Result.HasError())
 	{
 		return Result;
 	}
-	check(PendingChunks.IsEmpty());
+	check(CacheState->PendingChunks.IsEmpty());
 
-	Cas.Compact();
+	CacheState->Cas.Compact();
 	return MakeValue();
 }
 
-FResult FOnDemandInstallCache::FlushLastAccess()
+FResult FOnDemandInstallCache::FlushLastAccess(const TStaticArray<bool, EOnDemandInstallCasType::Count>& bFlushCasLastAccess)
 {
-	const FCas::FLastAccess DirtyLastAccess = Cas.GetAndClearDirtyLastAccess();
+	for (EOnDemandInstallCasType CasType : TEnumRange<EOnDemandInstallCasType>())
+	{
+		if (bFlushCasLastAccess[CasType] == false)
+		{
+			continue;
+		}
 
-	if (DirtyLastAccess.IsEmpty())
-	{
-		return MakeValue();
-	}
+		FCacheState* CacheState = GetCacheState(CasType);
+		if (CacheState == nullptr)
+		{
+			continue;
+		}
 
-	const FString JournalFile = GetJournalFilename();
-	FCasJournal::FTransaction Transaction = FCasJournal::Begin(JournalFile);
-	JournalLastAccess(Transaction, DirtyLastAccess);
+		const FCas::FLastAccess DirtyLastAccess = CacheState->Cas.GetAndClearDirtyLastAccess();
+		if (DirtyLastAccess.IsEmpty())
+		{
+			continue;
+		}
 
-	if (FResult Result = FCasJournal::Commit(MoveTemp(Transaction)); Result.HasError())
-	{
-		UE_LOG(LogIoStoreOnDemand, Error, TEXT("Failed to update CAS journal '%s' with block timestamp(s), reason '%s'"),
-			*JournalFile, *LexToString(Result.GetError()));
-		Result.GetError().PushErrorContext(MakeInstallCacheErrorContext());
-		return Result;
+		const FString JournalFile = GetJournalFilename();
+		FCasJournal::FTransaction Transaction = FCasJournal::Begin(CasType, JournalFile);
+		JournalLastAccess(Transaction, DirtyLastAccess);
+
+		if (FResult Result = FCasJournal::Commit(MoveTemp(Transaction)); Result.HasError())
+		{
+			UE_LOG(LogIoStoreOnDemand, Error, TEXT("Failed to update CAS journal '%s' with block timestamp(s), reason '%s'"),
+				*JournalFile, *LexToString(Result.GetError()));
+			Result.GetError().PushErrorContext(MakeInstallCacheErrorContext(*CacheState));
+			return Result;
+		}
 	}
 
 	return MakeValue();
 }
 
-void FOnDemandInstallCache::UpdateLastAccess(TConstArrayView<FIoHash> ChunkHashes, bool& bInOutLastAccessDirty)
+void FOnDemandInstallCache::UpdateLastAccess(TConstArrayView<FIoHash> ChunkHashes, TStaticArray<bool, EOnDemandInstallCasType::Count>& bInOutLastAccessDirty)
 {
 	if (ChunkHashes.IsEmpty())
 	{
@@ -3935,31 +4209,43 @@
 	const int64 Now = FDateTime::UtcNow().GetTicks();
 	constexpr const bool bDirty = true;
 
-	TUniqueLock Lock(Cas);
+	for (FCacheState& CacheState : CacheStates)
+	{
+		const EOnDemandInstallCasType CasType = CacheState.Type;
+		TUniqueLock Lock(CacheState.Cas);
 
-	for (const FIoHash& ChunkHash : ChunkHashes)
-	{
-		const FCasAddr& Addr = AsCasAddr(ChunkHash);
-		if (FCasLocation* CasLoc = Cas.Lookup.Find(Addr))
+		for (const FIoHash& ChunkHash : ChunkHashes)
 		{
-			const FCasBlockId BlockId = CasLoc->BlockId;
-			const uint32 BlockIdHash = GetTypeHash(BlockId);
-			const bool bUpdatedLastAccess = Cas.UnlockedTrackAccessIf(
-				ECasTrackAccessType::Granular, BlockIdHash, BlockId, Now, bDirty);
-			bInOutLastAccessDirty = bInOutLastAccessDirty || bUpdatedLastAccess;
+			const FCasAddr& Addr = AsCasAddr(ChunkHash);
+
+			if (FCasLocation* CasLoc = CacheState.Cas.Lookup.Find(Addr))
+			{
+				const FCasBlockId BlockId = CasLoc->BlockId;
+				const uint32 BlockIdHash = GetTypeHash(BlockId);
+				const bool bUpdatedLastAccess = CacheState.Cas.UnlockedTrackAccessIf(
+					ECasTrackAccessType::Granular, BlockIdHash, BlockId, Now, bDirty);
+
+				bInOutLastAccessDirty[CasType] = bInOutLastAccessDirty[CasType] || bUpdatedLastAccess;
+			}
 		}
 	}
 }
 
-FOnDemandInstallCacheUsage FOnDemandInstallCache::GetCacheUsage()
+FOnDemandInstallCacheUsage FOnDemandInstallCache::GetCacheUsage(EOnDemandInstallCasType CasType)
 {
 	// If this is called from a thread other than the OnDemandIoStore tick thread
 	// then its possible the block info and containers may not be in sync with each other
 	// or the current state of the tick thread.
 	// This should only be used for debugging and telemetry purposes.
 
+	FCacheState* CacheState = GetCacheState(CasType);
+	if (CacheState == nullptr)
+	{
+		return FOnDemandInstallCacheUsage();
+	}
+
 	FCasBlockInfoMap	BlockInfo;
-	const uint64		TotalCachedBytes = Cas.GetBlockInfo(BlockInfo);
+	const uint64		TotalCachedBytes = CacheState->Cas.GetBlockInfo(BlockInfo);
 
 	TArray<FSharedOnDemandContainer>	Containers;
 	TArray<TBitArray<>>					ChunkEntryIndices;
@@ -3967,7 +4253,7 @@
 	check(Containers.Num() == ChunkEntryIndices.Num());
 
 	uint64 ReferencedBytes = 0;
-	AddReferencesToBlocks(Containers, ChunkEntryIndices, BlockInfo, ReferencedBytes);
+	AddReferencesToBlocks(*CacheState, Containers, ChunkEntryIndices, BlockInfo, ReferencedBytes);
 
 	uint64 FragmentedBytes = 0;
 	uint64 ReferencedBlockBytes = 0;
@@ -3989,7 +4275,7 @@
 
 	return FOnDemandInstallCacheUsage
 	{
-		.MaxSize = MaxCacheSize,
+		.MaxSize = CacheState->MaxCacheSize,
 		.TotalSize = TotalCachedBytes,
 		.ReferencedBlockSize = ReferencedBlockBytes,
 		.ReferencedSize = ReferencedBytes,
@@ -3997,7 +4283,7 @@
 	};
 }
 
-FResult FOnDemandInstallCache::FlushPendingChunks(FPendingChunks& Chunks, int64 UtcAccessTicks, int64 UtcModificationTicks)
+FResult FOnDemandInstallCache::FlushPendingChunks(FCacheState& CacheState, FPendingChunks& Chunks, int64 UtcAccessTicks, int64 UtcModificationTicks)
 {
 	if (Chunks.IsEmpty())
 	{
@@ -4005,7 +4291,7 @@
 	}
 
 #if UE_ONDEMANDINSTALLCACHE_EXCLUSIVE_WRITE
-	UE::Tasks::TTask<FResult> Task = ExclusivePipe.Launch(UE_SOURCE_LOCATION, [this, &Chunks, UtcAccessTicks, UtcModificationTicks]
+	UE::Tasks::TTask<FResult> Task = CacheState.ExclusivePipe.Launch(UE_SOURCE_LOCATION, [this, &CacheState, &Chunks, UtcAccessTicks, UtcModificationTicks]
 	{
 		ON_SCOPE_EXIT 
 		{ 
@@ -4013,7 +4299,7 @@
 			Chunks.Reset();
 		};
 
-		return FlushPendingChunksImpl(Chunks, UtcAccessTicks, UtcModificationTicks);
+		return FlushPendingChunksImpl(CacheState, Chunks, UtcAccessTicks, UtcModificationTicks);
 	}, UE::Tasks::ETaskPriority::BackgroundHigh);
 
 	Task.Wait();
@@ -4027,15 +4313,15 @@
 		Chunks.Reset();
 	};
 
-	return FlushPendingChunksImpl(Chunks, UtcAccessTicks, UtcModificationTicks);
+	return FlushPendingChunksImpl(CacheState, Chunks, UtcAccessTicks, UtcModificationTicks);
 #endif // UE_ONDEMANDINSTALLCACHE_EXCLUSIVE_WRITE
 }
 
-FResult FOnDemandInstallCache::FlushPendingChunksImpl(const FPendingChunks& Chunks, int64 UtcAccessTicks, int64 UtcModificationTicks)
+FResult FOnDemandInstallCache::FlushPendingChunksImpl(FCacheState& CacheState, const FPendingChunks& Chunks, int64 UtcAccessTicks, int64 UtcModificationTicks)
 {
 	using namespace UE::UnifiedError::IoStoreOnDemand;
 
-	const bool bFromDefrag = (&PendingChunks != &Chunks);
+	const bool bFromDefrag = (&CacheState.PendingChunks != &Chunks);
 
 	uint64 TotalCasBytes		= 0;
 	uint64 TotalJournalBytes	= 0;
@@ -4046,32 +4332,32 @@
 		Governor.OnInstallCacheFlushed(TotalCasBytes + TotalJournalBytes, TotalOpCount);
 	};
 
-	FLargeMemoryWriter Ar(FMath::Min(Chunks.TotalSize, static_cast<uint64>(Cas.GetMaxBlockSize())));
+	FLargeMemoryWriter Ar(FMath::Min(Chunks.TotalSize, static_cast<uint64>(CacheState.Cas.GetMaxBlockSize())));
 
 	// This should be the only thread writing to CurrentBlock
-	FCasBlockId CurrentBlockId = CurrentBlock;
+	FCasBlockId CurrentBlockId = CacheState.CurrentBlock;
 
 	int32 ChunkIdx = 0;
 
 	while (ChunkIdx < Chunks.Chunks.Num())
 	{
-		FCasJournal::FTransaction Transaction = FCasJournal::Begin(GetJournalFilename());
+		FCasJournal::FTransaction Transaction = FCasJournal::Begin(CacheState.Type, GetJournalFilename());
 		
 		// Only open for append if continuing a block.
 		const bool bAppendToBlock = CurrentBlockId.IsValid();
 
 		if (CurrentBlockId.IsValid() == false)
 		{
-			CurrentBlockId = Cas.CreateBlock();
+			CurrentBlockId = CacheState.Cas.CreateBlock();
 			ensure(CurrentBlockId.IsValid());
-			CurrentBlock = CurrentBlockId;
+			CacheState.CurrentBlock = CurrentBlockId;
 			Transaction.BlockCreated(CurrentBlockId);
 		}
 
-		TResult<FUniqueFileHandle> OpenWriteResult = Cas.OpenWrite(CurrentBlockId, bAppendToBlock);
+		TResult<FUniqueFileHandle> OpenWriteResult = CacheState.Cas.OpenWrite(CurrentBlockId, bAppendToBlock);
 		if (OpenWriteResult.HasError())
 		{
-			FInstallCacheErrorContext Ctx = MakeInstallCacheErrorContext();
+			FInstallCacheErrorContext Ctx = MakeInstallCacheErrorContext(CacheState);
 			OpenWriteResult.GetError().PushErrorContext(MoveTemp(Ctx));
 			FResult Result = MakeError(OpenWriteResult.StealError());
 			FOnDemandInstallCacheStats::OnFlush(Result, TotalCasBytes, bFromDefrag);
@@ -4086,12 +4372,12 @@
 
 		while (ChunkIdx < Chunks.Chunks.Num())
 		{
-			if (CasBlockOffset > 0 && CasBlockOffset + Ar.Tell() + Chunks.Chunks[ChunkIdx].GetSize() > Cas.GetMaxBlockSize())
+			if (CasBlockOffset > 0 && CasBlockOffset + Ar.Tell() + Chunks.Chunks[ChunkIdx].GetSize() > CacheState.Cas.GetMaxBlockSize())
 			{
 				break;
 			}
 
-			if (!ensure(Chunks.Chunks[ChunkIdx].GetSize() <= Cas.GetMaxBlockSize()))
+			if (!ensure(Chunks.Chunks[ChunkIdx].GetSize() <= CacheState.Cas.GetMaxBlockSize()))
 			{
 				UE_LOGFMT(LogIoStoreOnDemand, Error, "Chunk is larger then block size ({ChunkId})", ("ChunkId",Chunks.ChunkHashes[ChunkIdx]));
 			}
@@ -4111,13 +4397,13 @@
 			UE_LOG(LogIoStoreOnDemand, Log, TEXT("Writing %.2lf MiB to CAS block %u"),
 				ToMiB(Ar.Tell()), CurrentBlockId.Id);
 
-			check(CasBlockOffset + BytesToWrite <= Cas.GetMaxBlockSize());
+			check(CasBlockOffset + BytesToWrite <= CacheState.Cas.GetMaxBlockSize());
 
 			if (CasFileHandle->Write(Ar.GetData(), BytesToWrite) == false)
 			{
 				FResult Result = MakeCasError<void>(
 					ECasErrorCode::WriteBlockFailed, EIoErrorCode::WriteError, FString::Printf(TEXT("Failed to write %" INT64_FMT " bytes to CAS block %u"), BytesToWrite, CurrentBlockId.Id));
-				FInstallCacheErrorContext Ctx = MakeInstallCacheErrorContext();
+				FInstallCacheErrorContext Ctx = MakeInstallCacheErrorContext(CacheState);
 				Result.GetError().PushErrorContext(MoveTemp(Ctx));
 				FOnDemandInstallCacheStats::OnFlush(Result, TotalCasBytes, bFromDefrag);
 				return Result;
@@ -4129,7 +4415,7 @@
 			{
 				FResult Result = MakeCasError<void>(
 					ECasErrorCode::WriteBlockFailed, EIoErrorCode::FileFlushFailed, FString::Printf(TEXT("Failed to flush %" INT64_FMT " bytes to CAS block %u"), BytesToWrite, CurrentBlockId.Id));
-				FInstallCacheErrorContext Ctx = MakeInstallCacheErrorContext();
+				FInstallCacheErrorContext Ctx = MakeInstallCacheErrorContext(CacheState);
 				Result.GetError().PushErrorContext(MoveTemp(Ctx));
 				FOnDemandInstallCacheStats::OnFlush(Result, TotalCasBytes, bFromDefrag);
 				return Result;
@@ -4139,7 +4425,7 @@
 			constexpr const bool bDirty = false;
 			if (UtcAccessTicks)
 			{
-				if (Cas.TrackAccessIf(ECasTrackAccessType::Newer, CurrentBlockId, UtcAccessTicks, bDirty))
+				if (CacheState.Cas.TrackAccessIf(ECasTrackAccessType::Newer, CurrentBlockId, UtcAccessTicks, bDirty))
 				{
 					Transaction.BlockAccess(CurrentBlockId, UtcAccessTicks);
 				}
@@ -4147,13 +4433,13 @@
 			else
 			{
 				const int64 Now = FDateTime::UtcNow().GetTicks();
-				verify(Cas.TrackAccessIf(ECasTrackAccessType::Always, CurrentBlockId, Now, bDirty));
+				verify(CacheState.Cas.TrackAccessIf(ECasTrackAccessType::Always, CurrentBlockId, Now, bDirty));
 				Transaction.BlockAccess(CurrentBlockId, Now);
 			}
 
 			if (UtcModificationTicks)
 			{
-				if (Cas.TrackModificationIf(ECasTrackModificationType::Newer, CurrentBlockId, UtcModificationTicks))
+				if (CacheState.Cas.TrackModificationIf(ECasTrackModificationType::Newer, CurrentBlockId, UtcModificationTicks))
 				{
 					Transaction.BlockModification(CurrentBlockId, UtcModificationTicks);
 				}
@@ -4161,11 +4447,11 @@
 			else
 			{
 				const int64 Now = FDateTime::UtcNow().GetTicks();
-				verify(Cas.TrackModificationIf(ECasTrackModificationType::Always, CurrentBlockId, Now));
+				verify(CacheState.Cas.TrackModificationIf(ECasTrackModificationType::Always, CurrentBlockId, Now));
 				Transaction.BlockModification(CurrentBlockId, Now);
 			}
 
-			Cas.UpdateBlock(CurrentBlockId, IntCastChecked<uint32>(BytesToWrite), ChunkHashes, Offsets,
+			CacheState.Cas.UpdateBlock(CurrentBlockId, IntCastChecked<uint32>(BytesToWrite), ChunkHashes, Offsets,
 				[&Transaction](const FCasAddr& CasAddr, const FCasLocation& Loc)
 				{
 					Transaction.ChunkLocation(Loc, CasAddr);
@@ -4179,7 +4465,7 @@
 		uint32 JrnlOps			= 0;
 		if (FResult Result = FCasJournal::Commit(MoveTemp(Transaction), JrnlBytesWritten, JrnlOps); Result.HasError())
 		{
-			FInstallCacheErrorContext Ctx = MakeInstallCacheErrorContext();
+			FInstallCacheErrorContext Ctx = MakeInstallCacheErrorContext(CacheState);
 			Result.GetError().PushErrorContext(MoveTemp(Ctx));
 			FOnDemandInstallCacheStats::OnFlush(Result, TotalCasBytes, bFromDefrag);
 			return Result;
@@ -4252,13 +4538,24 @@
 	BackendContext->WakeUpDispatcherThreadDelegate.Execute();
 }
 
-UE::UnifiedError::IoStoreOnDemand::FInstallCacheErrorContext FOnDemandInstallCache::MakeInstallCacheErrorContext(uint64 TotalCachedBytes, uint64 RequestedBytes, uint32 LineNo)
+FOnDemandInstallCache::FCacheState* FOnDemandInstallCache::GetCacheState(EOnDemandInstallCasType CacheType)
 {
+	return Algo::FindBy(CacheStates, CacheType, &FCacheState::Type);
+}
+
+const FOnDemandInstallCache::FCacheState* FOnDemandInstallCache::GetCacheState(EOnDemandInstallCasType CacheType) const
+{
+	return Algo::FindBy(CacheStates, CacheType, &FCacheState::Type);
+}
+
+UE::UnifiedError::IoStoreOnDemand::FInstallCacheErrorContext FOnDemandInstallCache::MakeInstallCacheErrorContext(FCacheState& CacheState, uint64 TotalCachedBytes, uint64 RequestedBytes, uint32 LineNo)
+{
 	UE::UnifiedError::IoStoreOnDemand::FInstallCacheErrorContext OutCtx;
 
-	const FString CachRootDirectory = Cas.GetRootDirectory();
+	const FString CachRootDirectory = CacheState.Cas.GetRootDirectory();
+	OutCtx.CasType					= LexToString(CacheState.Type);
 	OutCtx.bDiskQuerySucceeded		= FPlatformMisc::GetDiskTotalAndFreeSpace(CachRootDirectory, OutCtx.DiskTotalBytes, OutCtx.DiskFreeBytes);
-	OutCtx.MaxCacheSize				= MaxCacheSize;
+	OutCtx.MaxCacheSize				= CacheState.MaxCacheSize;
 	OutCtx.CacheSize				= TotalCachedBytes;
 	OutCtx.RequestedSize			= RequestedBytes;
 	OutCtx.LineNo					= LineNo;
@@ -4340,7 +4637,7 @@
 		FResult Result = FCasJournal::Create(JournalFile);
 		CHECK(Result.HasValue());
 
-		FCasJournal::FTransaction Transaction = FCasJournal::Begin(JournalFile);
+		FCasJournal::FTransaction Transaction = FCasJournal::Begin(EOnDemandInstallCasType::General, JournalFile);
 		Transaction.BlockCreated(FCasBlockId(1));
 		Result = FCasJournal::Commit(MoveTemp(Transaction));
 		CHECK(Result.HasValue());
@@ -4365,7 +4662,7 @@
 		FResult Result = FCasJournal::Create(JournalFile);
 		CHECK(Result.HasValue());
 
-		FCasJournal::FTransaction Transaction = FCasJournal::Begin(JournalFile);
+		FCasJournal::FTransaction Transaction = FCasJournal::Begin(EOnDemandInstallCasType::General, JournalFile);
 		for (int32 Idx = 0; const FCasAddr& Addr : ExpectedAddresses)
 		{
 			Transaction.ChunkLocation(
@@ -4426,7 +4723,7 @@
 		FResult Result = FCasJournal::Create(JournalFile);
 		CHECK(Result.HasValue());
 
-		FCasJournal::FTransaction Tx = FCasJournal::Begin(JournalFile);
+		FCasJournal::FTransaction Tx = FCasJournal::Begin(EOnDemandInstallCasType::General, JournalFile);
 		Tx.BlockCreated(ExpectedBlockId);
 		Tx.BlockDeleted(ExpectedBlockId);
 
@@ -4479,7 +4776,7 @@
 		FResult Result = FCasJournal::Create(JournalFile);
 		CHECK(Result.HasValue());
 
-		FCasJournal::FTransaction Tx = FCasJournal::Begin(JournalFile);
+		FCasJournal::FTransaction Tx = FCasJournal::Begin(EOnDemandInstallCasType::General, JournalFile);
 		Tx.BlockAccess(ExpectedBlockId, ExpectedTicks);
 
 		Result = FCasJournal::Commit(MoveTemp(Tx));
@@ -4528,7 +4825,7 @@
 		FResult Result = FCasJournal::Create(JournalFile);
 		CHECK(Result.HasValue());
 
-		FCasJournal::FTransaction Tx = FCasJournal::Begin(JournalFile);
+		FCasJournal::FTransaction Tx = FCasJournal::Begin(EOnDemandInstallCasType::General, JournalFile);
 		Tx.BlockModification(ExpectedBlockId, ExpectedTicks);
 
 		Result = FCasJournal::Commit(MoveTemp(Tx));
@@ -4573,12 +4870,14 @@
 	{
 		// Arrange
 		FCasSnapshot ExpectedSnapshot;
+		FCasSnapshot::FCasState& ExpectedCasState = ExpectedSnapshot.CasState.Emplace_GetRef();
+		ExpectedCasState.CasType = EOnDemandInstallCasType::General;
 
 		const int64 Now = FDateTime::UtcNow().GetTicks();
 
 		for (uint32 Id = 1; Id <= 10; ++Id)
 		{
-			ExpectedSnapshot.Blocks.Add(FCasSnapshot::FBlock
+			ExpectedCasState.Blocks.Add(FCasSnapshot::FBlock
 			{
 				.BlockId	= FCasBlockId(Id),
 				.LastAccess = Now + Id,
@@ -4593,10 +4892,10 @@
 					.BlockId		= FCasBlockId(Id),
 					.BlockOffset	= Idx * 256
 				};
-				ExpectedSnapshot.ChunkLocations.Emplace(CasAddr, Loc);
+				ExpectedCasState.ChunkLocations.Emplace(CasAddr, Loc);
 			}
 		}
-		ExpectedSnapshot.CurrentBlockId = FCasBlockId(1);
+		ExpectedCasState.CurrentBlockId = FCasBlockId(1);
 
 		// Act
 		FTmpDirectoryScope _(TestBaseDir);
@@ -4606,20 +4905,25 @@
 		const FCasSnapshot Snapshot = FCasSnapshot::Load(SnapshotFile).StealValue();
 
 		// Assert
-		CHECK(Snapshot.Blocks.Num() == ExpectedSnapshot.Blocks.Num());
-		for (int32 Idx = 0; Idx < Snapshot.Blocks.Num(); ++Idx)
+		CHECK(Snapshot.CasState.Num() == ExpectedSnapshot.CasState.Num());
+
+		const FCasSnapshot::FCasState& CasState = Snapshot.CasState[0];
+
+		CHECK(CasState.CasType == ExpectedCasState.CasType);
+		CHECK(CasState.Blocks.Num() == ExpectedCasState.Blocks.Num());
+		for (int32 Idx = 0; Idx < CasState.Blocks.Num(); ++Idx)
 		{
-			CHECK(Snapshot.Blocks[Idx].BlockId == ExpectedSnapshot.Blocks[Idx].BlockId);
-			CHECK(Snapshot.Blocks[Idx].LastAccess == ExpectedSnapshot.Blocks[Idx].LastAccess);
-			CHECK(Snapshot.Blocks[Idx].ModTime == ExpectedSnapshot.Blocks[Idx].ModTime);
+			CHECK(CasState.Blocks[Idx].BlockId == ExpectedCasState.Blocks[Idx].BlockId);
+			CHECK(CasState.Blocks[Idx].LastAccess == ExpectedCasState.Blocks[Idx].LastAccess);
+			CHECK(CasState.Blocks[Idx].ModTime == ExpectedCasState.Blocks[Idx].ModTime);
 		}
-		CHECK(Snapshot.ChunkLocations.Num() == ExpectedSnapshot.ChunkLocations.Num());
-		for (int32 Idx = 0; Idx < Snapshot.ChunkLocations.Num(); ++Idx)
+		CHECK(CasState.ChunkLocations.Num() == ExpectedCasState.ChunkLocations.Num());
+		for (int32 Idx = 0; Idx < CasState.ChunkLocations.Num(); ++Idx)
 		{
-			CHECK(Snapshot.ChunkLocations[Idx].Get<0>() == ExpectedSnapshot.ChunkLocations[Idx].Get<0>());
-			CHECK(Snapshot.ChunkLocations[Idx].Get<1>() == ExpectedSnapshot.ChunkLocations[Idx].Get<1>());
+			CHECK(CasState.ChunkLocations[Idx].Get<0>() == ExpectedCasState.ChunkLocations[Idx].Get<0>());
+			CHECK(CasState.ChunkLocations[Idx].Get<1>() == ExpectedCasState.ChunkLocations[Idx].Get<1>());
 		}
-		CHECK(Snapshot.CurrentBlockId == ExpectedSnapshot.CurrentBlockId);
+		CHECK(CasState.CurrentBlockId == ExpectedCasState.CurrentBlockId);
 	}
 
 	SECTION("CreateFromJournal")
@@ -4632,7 +4936,7 @@
 		FResult Result = FCasJournal::Create(JournalFile);
 		CHECK(Result.HasValue());
 
-		FCasJournal::FTransaction Tx = FCasJournal::Begin(JournalFile);
+		FCasJournal::FTransaction Tx = FCasJournal::Begin(EOnDemandInstallCasType::General, JournalFile);
 
 		const int64 Now = FDateTime::UtcNow().GetTicks();
 		const int64 ExpectedModTime = Now + 2;
@@ -4678,20 +4982,25 @@
 		CHECK(Result.HasValue());
 
 		// Act
-		const FCasSnapshot Snapshot = FCasSnapshot::FromJournal(JournalFile, FCasSnapshot()).StealValue();
+		TResult<FCasSnapshot> FromJournalResult = FCasSnapshot::FromJournal(JournalFile, FCasSnapshot());
+		CHECK(FromJournalResult.HasValue());
+		const FCasSnapshot Snapshot = FromJournalResult.StealValue();
+		CHECK(Snapshot.CasState.Num() == 1);
+		const FCasSnapshot::FCasState& CasState = Snapshot.CasState[0];
+		CHECK(CasState.CasType == EOnDemandInstallCasType::General);
 
 		// Assert
-		CHECK(Snapshot.CurrentBlockId == ExpectedCurrentBlockId);
-		CHECK(Snapshot.Blocks.Num() == 1); 
-		CHECK(Snapshot.Blocks[0].ModTime == ExpectedModTime);
-		CHECK(Snapshot.Blocks[0].LastAccess == ExpectedLastAccess);
-		CHECK(Snapshot.ChunkLocations.Num() == 10); 
-		for (int32 Idx = 1; Idx < Snapshot.ChunkLocations.Num(); ++Idx)
+		CHECK(CasState.CurrentBlockId == ExpectedCurrentBlockId);
+		CHECK(CasState.Blocks.Num() == 1);
+		CHECK(CasState.Blocks[0].ModTime == ExpectedModTime);
+		CHECK(CasState.Blocks[0].LastAccess == ExpectedLastAccess);
+		CHECK(CasState.ChunkLocations.Num() == 10);
+		for (int32 Idx = 1; Idx < CasState.ChunkLocations.Num(); ++Idx)
 		{
 			const FCasAddr Addr = CreateCasTestAddr(Idx);
 			const FCasSnapshot::FChunkLocation* Loc =
 				Algo::FindByPredicate(
-					Snapshot.ChunkLocations,
+					CasState.ChunkLocations,
 					[&Addr](const FCasSnapshot::FChunkLocation& L) { return L.Get<0>() == Addr; });
 			CHECK(Loc != nullptr);
 			if (Loc != nullptr)
@@ -4706,12 +5015,14 @@
 	{
 		// Arrange
 		FCasSnapshot ExpectedSnapshot;
+		FCasSnapshot::FCasState& ExpectedCasState = ExpectedSnapshot.CasState.Emplace_GetRef();
+		ExpectedCasState.CasType = EOnDemandInstallCasType::General;
 
 		const int64 Now = FDateTime::UtcNow().GetTicks();
 
 		for (uint32 Id = 1; Id <= 10; ++Id)
 		{
-			ExpectedSnapshot.Blocks.Add(FCasSnapshot::FBlock
+			ExpectedCasState.Blocks.Add(FCasSnapshot::FBlock
 			{
 				.BlockId = FCasBlockId(Id),
 				.LastAccess = Now + Id,
@@ -4726,10 +5037,10 @@
 					.BlockId = FCasBlockId(Id),
 					.BlockOffset = Idx * 256
 				};
-				ExpectedSnapshot.ChunkLocations.Emplace(CasAddr, Loc);
+				ExpectedCasState.ChunkLocations.Emplace(CasAddr, Loc);
 			}
 		}
-		ExpectedSnapshot.CurrentBlockId = FCasBlockId(1);
+		ExpectedCasState.CurrentBlockId = FCasBlockId(1);
 
 		FTmpDirectoryScope _(TestBaseDir);
 		const FString SnapshotFile = TestBaseDir / TEXT("test.snp");
@@ -4741,7 +5052,7 @@
 		FResult Result = FCasJournal::Create(JournalFile);
 		CHECK(Result.HasValue());
 
-		FCasJournal::FTransaction Tx = FCasJournal::Begin(JournalFile);
+		FCasJournal::FTransaction Tx = FCasJournal::Begin(EOnDemandInstallCasType::General, JournalFile);
 
 		// Add a block and some chunk locations
 		FCasBlockId Block11(11);
@@ -4767,32 +5078,35 @@
 
 		TResult<FCasSnapshot> LoadSnapResult = FCasSnapshot::Load(SnapshotFile);
 		CHECK(LoadSnapResult.HasValue());
-		FCasSnapshot Snapshot = LoadSnapResult.StealValue();
+		const FCasSnapshot Snapshot = LoadSnapResult.StealValue();
+		CHECK(Snapshot.CasState.Num() == ExpectedSnapshot.CasState.Num());
+		const FCasSnapshot::FCasState& CasState = Snapshot.CasState[0];
+		CHECK(CasState.CasType == EOnDemandInstallCasType::General);
 
 		// Assert
-		CHECK(Snapshot.Blocks.Num() == ExpectedSnapshot.Blocks.Num() + 1);
-		for (int32 Idx = 0; Idx < ExpectedSnapshot.Blocks.Num(); ++Idx)
+		CHECK(CasState.Blocks.Num() == ExpectedCasState.Blocks.Num() + 1);
+		for (int32 Idx = 0; Idx < ExpectedCasState.Blocks.Num(); ++Idx)
 		{
-			CHECK(Snapshot.Blocks[Idx].BlockId == ExpectedSnapshot.Blocks[Idx].BlockId);
-			CHECK(Snapshot.Blocks[Idx].LastAccess == ExpectedSnapshot.Blocks[Idx].LastAccess);
-			CHECK(Snapshot.Blocks[Idx].ModTime == ExpectedSnapshot.Blocks[Idx].ModTime);
+			CHECK(CasState.Blocks[Idx].BlockId == ExpectedCasState.Blocks[Idx].BlockId);
+			CHECK(CasState.Blocks[Idx].LastAccess == ExpectedCasState.Blocks[Idx].LastAccess);
+			CHECK(CasState.Blocks[Idx].ModTime == ExpectedCasState.Blocks[Idx].ModTime);
 		}
-		CHECK(Snapshot.ChunkLocations.Num() == ExpectedSnapshot.ChunkLocations.Num() + 10);
-		for (int32 Idx = 0; Idx < ExpectedSnapshot.ChunkLocations.Num(); ++Idx)
+		CHECK(CasState.ChunkLocations.Num() == ExpectedCasState.ChunkLocations.Num() + 10);
+		for (int32 Idx = 0; Idx < ExpectedCasState.ChunkLocations.Num(); ++Idx)
 		{
-			CHECK(Snapshot.ChunkLocations[Idx].Get<0>() == ExpectedSnapshot.ChunkLocations[Idx].Get<0>());
-			CHECK(Snapshot.ChunkLocations[Idx].Get<1>() == ExpectedSnapshot.ChunkLocations[Idx].Get<1>());
+			CHECK(CasState.ChunkLocations[Idx].Get<0>() == ExpectedCasState.ChunkLocations[Idx].Get<0>());
+			CHECK(CasState.ChunkLocations[Idx].Get<1>() == ExpectedCasState.ChunkLocations[Idx].Get<1>());
 		}
 
-		CHECK(Snapshot.CurrentBlockId == Block11);
+		CHECK(CasState.CurrentBlockId == Block11);
 
-		CHECK(Snapshot.Blocks[10].BlockId == Block11);
-		CHECK(Snapshot.Blocks[10].ModTime == Now + 100);
-		CHECK(Snapshot.Blocks[10].LastAccess == Now + 101);
+		CHECK(CasState.Blocks[10].BlockId == Block11);
+		CHECK(CasState.Blocks[10].ModTime == Now + 100);
+		CHECK(CasState.Blocks[10].LastAccess == Now + 101);
 
 		for (int32 Idx = 0; Idx < 10; ++Idx)
 		{
-			const int32 ExpectIdx = ExpectedSnapshot.ChunkLocations.Num() + Idx;
+			const int32 ExpectIdx = ExpectedCasState.ChunkLocations.Num() + Idx;
 
 			FCasAddr ExpectAddr = CreateCasTestAddr(111 + Idx);
 
@@ -4802,8 +5116,8 @@
 				.BlockOffset = (Idx + 1) * 256u
 			};
 
-			CHECK(Snapshot.ChunkLocations[ExpectIdx].Get<0>() == ExpectAddr);
-			CHECK(Snapshot.ChunkLocations[ExpectIdx].Get<1>() == ExpectLoc);
+			CHECK(CasState.ChunkLocations[ExpectIdx].Get<0>() == ExpectAddr);
+			CHECK(CasState.ChunkLocations[ExpectIdx].Get<1>() == ExpectLoc);
 		}
 	}
 }
diff --git a/Engine/Source/Runtime/Experimental/IoStore/OnDemand/Private/OnDemandInstallCache.h b/Engine/Source/Runtime/Experimental/IoStore/OnDemand/Private/OnDemandInstallCache.h
--- a/Engine/Source/Runtime/Experimental/IoStore/OnDemand/Private/OnDemandInstallCache.h
+++ b/Engine/Source/Runtime/Experimental/IoStore/OnDemand/Private/OnDemandInstallCache.h
@@ -3,9 +3,9 @@
 #pragma once
 
 #include "IO/IoDispatcherBackend.h"
+#include "IO/IoStoreOnDemand.h"
 #include "IO/OnDemandError.h"
 #include "Templates/SharedPointer.h"
-#include "Containers/StringFwd.h"
 
 struct FIoContainerHeader; 
 
@@ -30,25 +30,34 @@ class IOnDemandInstallCache
 public:
 	virtual										~IOnDemandInstallCache() = default;
 	virtual bool								IsChunkCached(const FIoHash& ChunkHash) = 0;
-	virtual TResult<bool>						PutChunk(FIoBuffer&& Chunk, const FIoHash& ChunkHash) = 0;
-	virtual FResult								Purge(uint64 BytesToInstall) = 0;
-	virtual FResult								PurgeAllUnreferenced(bool bDefrag, const uint64* BytesToPurge = nullptr) = 0;
-	virtual FResult								DefragAll(const uint64* BytesToFree = nullptr) = 0;
+	virtual TResult<bool>						PutChunk(EOnDemandInstallCasType CasType, FIoBuffer&& Chunk, const FIoHash& ChunkHash) = 0;
+	virtual FResult								Purge(EOnDemandInstallCasType CasType, uint64 BytesToInstall) = 0;
+	virtual FResult								PurgeAllUnreferenced(EOnDemandInstallCasType CasType, bool bDefrag, const uint64* BytesToPurge = nullptr) = 0;
+	virtual FResult								DefragAll(EOnDemandInstallCasType CasType, const uint64* BytesToFree = nullptr) = 0;
 	virtual FResult								Verify() = 0;
-	virtual FResult								Flush() = 0;
-	virtual FResult								FlushLastAccess() = 0;
-	virtual void								UpdateLastAccess(TConstArrayView<FIoHash> ChunkHashes, bool& bInOutLastAccessDirty) = 0;
-	virtual FOnDemandInstallCacheUsage			GetCacheUsage() = 0;
+	virtual FResult								Flush(EOnDemandInstallCasType CasType) = 0;
+	virtual FResult								FlushLastAccess(const TStaticArray<bool, EOnDemandInstallCasType::Count>& bFlushCasLastAccess) = 0;
+	virtual void								UpdateLastAccess(TConstArrayView<FIoHash> ChunkHashes, TStaticArray<bool, EOnDemandInstallCasType::Count>& bInOutLastAccessDirty) = 0;
+	virtual FOnDemandInstallCacheUsage			GetCacheUsage(EOnDemandInstallCasType CasType) = 0;
+};
+
+struct FOnDemandInstallCasConfig
+{
+	EOnDemandInstallCasType Type = EOnDemandInstallCasType::None;
+	FString CasSubdirectory;
+	uint64	DiskQuota = 1ull << 30;
+	uint64	MinBlockSize = 32 << 19;
+	uint64	MaxBlockSize = 32 << 20;
+	double	LastAccessGranularitySeconds = 60 * 60;
 };
 
 struct FOnDemandInstallCacheConfig
 {
+	TArray<FOnDemandInstallCasConfig> CasConfig;
+
 	FString RootDirectory;
-	uint64	DiskQuota		= 1ull << 30;
+	FString IndexedCacheStorageName;
 	uint64	JournalMaxSize	= 2ull << 20;
-	uint64	MinBlockSize	= 32 << 19;
-	uint64	MaxBlockSize	= 32 << 20;
-	double	LastAccessGranularitySeconds = 60 * 60;
 	bool	bDropCache		= false;
 };
 
diff --git a/Engine/Source/Runtime/Experimental/IoStore/OnDemand/Private/OnDemandIoStore.cpp b/Engine/Source/Runtime/Experimental/IoStore/OnDemand/Private/OnDemandIoStore.cpp
--- a/Engine/Source/Runtime/Experimental/IoStore/OnDemand/Private/OnDemandIoStore.cpp
+++ b/Engine/Source/Runtime/Experimental/IoStore/OnDemand/Private/OnDemandIoStore.cpp
@@ -972,52 +972,48 @@ FIoStatus FOnDemandIoStore::InitializeInstallCache()
 
 	// Check if we use ICS for IAD and update the RootDirectory accordingly
 	int32 SelectedCacheIndex = 0;
+	if (OnDemandInstallCacheConfig.IndexedCacheStorageName.IsEmpty() == false)
 	{
 		// OnDemandInstallCacheConfig.DiskQuota is the quota to store the data, but we also need some room for metadata, see GetJournalFilename/GetSnapshotFilename
-		uint64 MetaDataSize = 2 * OnDemandInstallCacheConfig.JournalMaxSize + 2 * 32;
-		uint64 OnDemandInstallRequestedSize = OnDemandInstallCacheConfig.DiskQuota + MetaDataSize;
+		const uint64 MetaDataSize = 2 * OnDemandInstallCacheConfig.JournalMaxSize + (2 << 20);
+		uint64 OnDemandInstallRequestedSize = 
+			MetaDataSize + 
+			Algo::TransformAccumulate(OnDemandInstallCacheConfig.CasConfig, &FOnDemandInstallCasConfig::DiskQuota, uint64(0));
 
-		if (OnDemandInstallRequestedSize > 0)
-		{
-			FModuleManager::Get().LoadModule(TEXT("IndexedCacheStorage"));
+		FModuleManager::Get().LoadModule(TEXT("IndexedCacheStorage"));
 
-			const FString IADStorageName("IAD");
-			SelectedCacheIndex = Experimental::FIndexedCacheStorageManager::Get().GetStorageIndex(IADStorageName);
+		SelectedCacheIndex = Experimental::FIndexedCacheStorageManager::Get().GetStorageIndex(OnDemandInstallCacheConfig.IndexedCacheStorageName);
+		if (SelectedCacheIndex > 0)
+		{
+			// ICS will allocate space for all early indices in one combined operation.
+			// In this case, it should be configured to read the early size from the IAD config, for example Engine:[OnDemandInstall]:FileCache.DiskQuota
+			// See Engine:[IndexedCacheStorage]:Storage
+			const uint64 EarlyStartupSize = Experimental::FIndexedCacheStorageManager::Get().GetCacheEarlyStartupSize(SelectedCacheIndex);
+			if (EarlyStartupSize > 0)
+			{
+				UE_CLOGF(EarlyStartupSize < OnDemandInstallRequestedSize, LogIoStoreOnDemand, Fatal, "EarlyStartupSize must be >= %" UINT64_FMT, OnDemandInstallRequestedSize);
+				OnDemandInstallRequestedSize = EarlyStartupSize;
+			}
 
-			if (SelectedCacheIndex > 0)
+			// When IAD is assigned to an indexed cache, we *have* to get enough disk space.
+			while (true)
 			{
-				uint64 EarlyStartupSize = Experimental::FIndexedCacheStorageManager::Get().GetCacheEarlyStartupSize(SelectedCacheIndex);
-				if (EarlyStartupSize > 0)
+				if (!Experimental::FIndexedCacheStorageManager::Get().CreateCacheStorage(OnDemandInstallRequestedSize, SelectedCacheIndex))
 				{
-					OnDemandInstallRequestedSize = EarlyStartupSize;
-
-					uint64 MinimumDiskQuota = 4 << 20;
-					check(EarlyStartupSize > MetaDataSize + MinimumDiskQuota);
-
-					// Update DiskQuota so that EarlyStartupSize can contain the whole request;
-					OnDemandInstallCacheConfig.DiskQuota = EarlyStartupSize - MetaDataSize;
+					continue;
 				}
 
-				// When IAD is assigned to an indexed cache, we *have* to get enough disk space.
-				while (true)
+				FString MountPath = Experimental::FIndexedCacheStorageManager::Get().MountCacheStorage(SelectedCacheIndex);
+				if (MountPath.IsEmpty())
 				{
-					if (!Experimental::FIndexedCacheStorageManager::Get().CreateCacheStorage(OnDemandInstallRequestedSize, SelectedCacheIndex))
-					{
-						continue;
-					}
-
-					FString MountPath = Experimental::FIndexedCacheStorageManager::Get().MountCacheStorage(SelectedCacheIndex);
-					if (MountPath.IsEmpty())
-					{
-						UE_LOG(LogIoStoreOnDemand, Error, TEXT("Failed to mount indexed cache storage %d (%s) even though it was created"), SelectedCacheIndex, *IADStorageName);
-						Experimental::FIndexedCacheStorageManager::Get().DestroyCacheStorage(SelectedCacheIndex);
-						continue;
-					}
-
-					FString RelativeDir(MountPath / OnDemandInstallCacheConfig.RootDirectory.Replace(*FPaths::ProjectPersistentDownloadDir(), TEXT("")));
-					OnDemandInstallCacheConfig.RootDirectory = RelativeDir;
-					break;
+					UE_LOG(LogIoStoreOnDemand, Error, TEXT("Failed to mount indexed cache storage %d (%s) even though it was created"), SelectedCacheIndex, *OnDemandInstallCacheConfig.IndexedCacheStorageName);
+					Experimental::FIndexedCacheStorageManager::Get().DestroyCacheStorage(SelectedCacheIndex);
+					continue;
 				}
+
+				FString RelativeDir(MountPath / OnDemandInstallCacheConfig.RootDirectory.Replace(*FPaths::ProjectPersistentDownloadDir(), TEXT("")));
+				OnDemandInstallCacheConfig.RootDirectory = RelativeDir;
+				break;
 			}
 		}
 	}
@@ -1084,7 +1080,7 @@ FIoStatus FOnDemandIoStore::InitializeInstallCache()
 				}
 
 				const FIoHash ChunkHash = FIoHash::HashBuffer(Chunk.GetView());
-				TResult<bool> PutResult = InstallCache->PutChunk(MoveTemp(Chunk), ChunkHash);
+				TResult<bool> PutResult = InstallCache->PutChunk(EOnDemandInstallCasType::General, MoveTemp(Chunk), ChunkHash);
 				if (PutResult.HasError())
 				{
 					Result = MakeError(PutResult.StealError());
@@ -1100,7 +1096,7 @@ FIoStatus FOnDemandIoStore::InitializeInstallCache()
 
 			if (Result.HasError() == false)
 			{
-				Result = InstallCache->Flush();
+				Result = InstallCache->Flush(EOnDemandInstallCasType::General);
 			}
 
 			UE_CLOG(Result.HasError(), LogIoStoreOnDemand, Warning, TEXT("Failed to fill install cache with dummy data, reason '%s'"), *LexToString(Result.GetError()));
@@ -1862,7 +1858,7 @@ FOnDemandCacheUsage FOnDemandIoStore::GetCacheUsage(const FOnDemandGetCacheUsage
 
 	if (InstallCache.IsValid())
 	{
-		CacheUsage.InstallCache = InstallCache->GetCacheUsage();
+		CacheUsage.InstallCache = InstallCache->GetCacheUsage(Args.CasType);
 
 		if (EnumHasAnyFlags(Args.Options, EOnDemandGetCacheUsageOptions::DumpHandlesToLog | EOnDemandGetCacheUsageOptions::DumpHandlesToResults))
 		{
diff --git a/Engine/Source/Runtime/Experimental/IoStore/OnDemandCore/Private/IoStoreOnDemand.cpp b/Engine/Source/Runtime/Experimental/IoStore/OnDemandCore/Private/IoStoreOnDemand.cpp
--- a/Engine/Source/Runtime/Experimental/IoStore/OnDemandCore/Private/IoStoreOnDemand.cpp
+++ b/Engine/Source/Runtime/Experimental/IoStore/OnDemandCore/Private/IoStoreOnDemand.cpp
@@ -20,6 +20,38 @@
 DEFINE_LOG_CATEGORY(LogIoStoreOnDemand);
 DEFINE_LOG_CATEGORY(LogIas);
 
+///////////////////////////////////////////////////////////////////////////////
+const TCHAR* LexToString(UE::IoStore::EOnDemandInstallCasType CacheType)
+{
+	using namespace UE::IoStore;
+
+	switch (CacheType)
+	{
+	case EOnDemandInstallCasType::General:
+		return TEXT("General");
+	case EOnDemandInstallCasType::MMap:
+		return TEXT("MMap");
+	}
+
+	return TEXT("None");
+	static_assert(EOnDemandInstallCasType::Count == 2); // Does not include 'None'
+}
+
+void LexFromString(UE::IoStore::EOnDemandInstallCasType& CacheType, FStringView String)
+{
+	using namespace UE::IoStore;
+
+	CacheType = EOnDemandInstallCasType::None;
+	for (EOnDemandInstallCasType Val : TEnumRange<EOnDemandInstallCasType>())
+	{
+		if (String == LexToString(Val))
+		{
+			CacheType = Val;
+			return;
+		}
+	}
+}
+
 namespace UE::IoStore
 {
 
@@ -64,10 +96,16 @@ static void PurgeInstallCache(const TArray<FString>& Args)
 		FOnDemandPurgeArgs PurgeArgs;
 		for (const FString& Arg : Args)
 		{
+			EOnDemandInstallCasType CasType = EOnDemandInstallCasType::None;
+
 			if (Arg == TEXTVIEW("defrag"))
 			{
 				PurgeArgs.Options |= EOnDemandPurgeOptions::Defrag;
 			}
+			else if (LexFromString(CasType, Arg); CasType != EOnDemandInstallCasType::None)
+			{
+				PurgeArgs.CasType = CasType;
+			}
 			else if (!PurgeArgs.BytesToPurge && UE::String::IsNumericOnlyDigits(Arg))
 			{
 				uint64 BytesToPurge = 0;
@@ -78,7 +116,7 @@ static void PurgeInstallCache(const TArray<FString>& Args)
 			}
 		}
 
-		UE_LOG(LogIoStoreOnDemand, Display, TEXT("Purging on demand install cache"));
+		UE_LOG(LogIoStoreOnDemand, Display, TEXT("Purging on demand install cache %s"), LexToString(PurgeArgs.CasType));
 		IoStore->Purge(MoveTemp(PurgeArgs), [](const FOnDemandPurgeResult& Result)
 		{
 			if (Result.IsOk())
@@ -105,7 +143,12 @@ static void DefragInstallCache(const TArray<FString>& Args)
 		FOnDemandDefragArgs DefragArgs;
 		for (const FString& Arg : Args)
 		{
-			if (UE::String::IsNumericOnlyDigits(Arg))
+			EOnDemandInstallCasType CasType = EOnDemandInstallCasType::None;
+			if (LexFromString(CasType, Arg); CasType != EOnDemandInstallCasType::None)
+			{
+				DefragArgs.CasType = CasType;
+			}
+			else if (UE::String::IsNumericOnlyDigits(Arg))
 			{
 				uint64 BytesToFree = 0;
 				if (LexTryParseString(BytesToFree, *Arg))
@@ -116,7 +159,7 @@ static void DefragInstallCache(const TArray<FString>& Args)
 			}
 		}
 
-		UE_LOG(LogIoStoreOnDemand, Display, TEXT("Defragging on demand install cache"));
+		UE_LOG(LogIoStoreOnDemand, Display, TEXT("Defragging on demand install cache %s"), LexToString(DefragArgs.CasType));
 		IoStore->Defrag(MoveTemp(DefragArgs), [](const FOnDemandDefragResult& Result)
 		{
 			if (Result.IsOk())
@@ -136,14 +179,22 @@ static void DefragInstallCache(const TArray<FString>& Args)
 }
 
 ////////////////////////////////////////////////////////////////////////////////
-static void PrintCacheUsage()
+static void PrintCacheUsage(const TArray<FString>& InArgs)
 {
 	if (UE::IoStore::IOnDemandIoStore* IoStore = UE::IoStore::TryGetOnDemandIoStore())
 	{
-		FOnDemandGetCacheUsageArgs Args;
-		Args.Options = EOnDemandGetCacheUsageOptions::DumpHandlesToLog;
+		FOnDemandGetCacheUsageArgs UsageArgs;
+		UsageArgs.Options = EOnDemandGetCacheUsageOptions::DumpHandlesToLog;
+		for (const FString& Arg : InArgs)
+		{
+			EOnDemandInstallCasType CasType = EOnDemandInstallCasType::None;
+			if (LexFromString(CasType, Arg); CasType != EOnDemandInstallCasType::None)
+			{
+				UsageArgs.CasType = CasType;
+			}
+		}
 
-		TIoStatusOr<FOnDemandCacheUsage> MaybeUsage = IoStore->GetCacheUsage(Args);
+		TIoStatusOr<FOnDemandCacheUsage> MaybeUsage = IoStore->GetCacheUsage(UsageArgs);
 		if (!MaybeUsage.IsOk())
 		{
 			UE_LOG(LogIoStoreOnDemand, Error, TEXT("iostore.CacheUsage failed: %s"), *MaybeUsage.Status().ToString());
@@ -153,7 +204,8 @@ static void PrintCacheUsage()
 		const FOnDemandCacheUsage CacheUsage = MaybeUsage.ConsumeValueOrDie();
 		UE_LOG(LogIoStoreOnDemand, Display, TEXT("iostore.CacheUsage"));
 		TStringBuilder<512> Sb;
-		Sb << TEXTVIEW("InstallCache: ") << CacheUsage.InstallCache << TEXTVIEW(", StreamingCache: ") << CacheUsage.StreamingCache;
+		Sb	<< TEXTVIEW("InstallCache ") << LexToString(UsageArgs.CasType) << TEXTVIEW(": ") << CacheUsage.InstallCache 
+			<< TEXTVIEW(", StreamingCache: ") << CacheUsage.StreamingCache;
 		UE_LOG(LogIoStoreOnDemand, Display, TEXT("%s"), Sb.ToString());
 
 		IoStore->DumpMountedContainersToLog();
@@ -386,7 +438,7 @@ static FAutoConsoleCommand DefragCacheCommand(
 static FAutoConsoleCommand PrintCacheUsageCommand(
 	TEXT("iostore.CacheUsage"),
 	TEXT("print cache usage"),
-	FConsoleCommandDelegate::CreateStatic(PrintCacheUsage),
+	FConsoleCommandWithArgsDelegate::CreateStatic(PrintCacheUsage),
 	ECVF_Cheat);
 
 static FAutoConsoleCommand MountUrlCommand(
diff --git a/Engine/Source/Runtime/Experimental/IoStore/OnDemandCore/Public/IO/IoStoreOnDemand.h b/Engine/Source/Runtime/Experimental/IoStore/OnDemandCore/Public/IO/IoStoreOnDemand.h
--- a/Engine/Source/Runtime/Experimental/IoStore/OnDemandCore/Public/IO/IoStoreOnDemand.h
+++ b/Engine/Source/Runtime/Experimental/IoStore/OnDemandCore/Public/IO/IoStoreOnDemand.h
@@ -12,6 +12,7 @@
 #include "IO/PackageId.h"
 #include "Misc/EnumClassFlags.h"
 #include "Features/IModularFeature.h"
+#include "Misc/EnumRange.h"
 #include "Misc/Optional.h"
 #include "Modules/ModuleInterface.h"
 #include "Templates/SharedPointer.h"
@@ -149,7 +150,7 @@ struct FOnDemandRegisterHostGroupResult
 };
 
 /** Options for controlling the behavior of mounted container(s). */
-enum class EOnDemandMountOptions
+enum class EOnDemandMountOptions : uint8
 {
 	/** No additional options. */
 	None					= 0,
@@ -216,6 +217,19 @@ struct FOnDemandMountResult
 /** Mount completion callback. */
 using FOnDemandMountCompleted = TUniqueFunction<void(FOnDemandMountResult)>;
 
+/** Specifies a particular CAS Type when relevant */
+struct EOnDemandInstallCasTypeNS
+{
+	enum Enum : uint8
+	{
+		None = 0xff,
+		General = 0,
+		MMap,
+		Count // Count does not include 'None'
+	};
+};
+using EOnDemandInstallCasType = EOnDemandInstallCasTypeNS::Enum;
+
 /** Options for controlling the behavior of the install request. */
 enum class EOnDemandInstallOptions : uint16
 {
@@ -361,7 +375,7 @@ private:
 };
 
 /** Options for controlling the behavior of the purge request. */
-enum class EOnDemandPurgeOptions
+enum class EOnDemandPurgeOptions : uint8
 {
 	/** No additional options. */
 	None = 0,
@@ -375,6 +389,8 @@ ENUM_CLASS_FLAGS(EOnDemandPurgeOptions)
 /** Arguments for purging on-demand content. */
 struct FOnDemandPurgeArgs
 {
+	/** Cas to purge */
+	EOnDemandInstallCasType CasType = EOnDemandInstallCasType::General;
 	/** Purge options. */
 	EOnDemandPurgeOptions Options = EOnDemandPurgeOptions::None;
 	/** Optional Size to purge. If not set will purge all unreferenced blocks. */
@@ -396,7 +412,7 @@ struct FOnDemandPurgeResult
 using FOnDemandPurgeCompleted = TUniqueFunction<void(FOnDemandPurgeResult)>;
 
 /** Options for controlling the behavior of the defrag request. */
-enum class EOnDemandDefragOptions
+enum class EOnDemandDefragOptions : uint8
 {
 	/** No additional options. */
 	None = 0,
@@ -408,6 +424,8 @@ ENUM_CLASS_FLAGS(EOnDemandDefragOptions)
 /** Arguments for defragmenting on-demand content. */
 struct FOnDemandDefragArgs
 {
+	/** Cas to defrag */
+	EOnDemandInstallCasType CasType = EOnDemandInstallCasType::General;
 	/** Defrag options. */
 	EOnDemandDefragOptions Options = EOnDemandDefragOptions::None;
 	/** Optional Size to free. If not set, will defrag all blocks and free all unreferenced chunks. */
@@ -429,7 +447,7 @@ struct FOnDemandDefragResult
 using FOnDemandDefragCompleted = TUniqueFunction<void(FOnDemandDefragResult)>;
 
 /** Options for controlling the behavior of the install size request. */
-enum class EOnDemandGetInstallSizeOptions
+enum class EOnDemandGetInstallSizeOptions : uint8
 {
 	/** No additional options. */
 	None					= 0,
@@ -461,7 +479,7 @@ struct FOnDemandInstallSizeResult
 };
 
 /** Options for controlling the behavior of the GetIsOnDemand request. */
-enum class EOnDemandIsOnDemandOptions
+enum class EOnDemandIsOnDemandOptions : uint8
 {
 	/** No additional options. */
 	None = 0,
@@ -486,7 +504,7 @@ struct FOnDemandIsOnDemandResult
 };
 
 /** Options for controlling the behavior of the install size request. */
-enum class EOnDemandGetCacheUsageOptions
+enum class EOnDemandGetCacheUsageOptions : uint8
 {
 	/** No additional options. */
 	None = 0,
@@ -500,6 +518,7 @@ ENUM_CLASS_FLAGS(EOnDemandGetCacheUsageOptions);
 /** Arguments for getting the install cache usage. */
 struct FOnDemandGetCacheUsageArgs
 {
+	EOnDemandInstallCasType CasType = EOnDemandInstallCasType::General;
 	EOnDemandGetCacheUsageOptions Options = EOnDemandGetCacheUsageOptions::None;
 };
 
@@ -598,7 +617,7 @@ struct FOnDemandImmediateAnalyticHandler
 };
 
 /** Streaming options (IAS). */
-enum class EOnDemandStreamingOptions
+enum class EOnDemandStreamingOptions : uint8
 {
 	/** No special options enabled. */
 	Default						= 0,
@@ -749,4 +768,8 @@ UE_API IOnDemandIoStore& GetOnDemandIoStore();
 
 } // namespace UE::IoStore
 
+UE_API const TCHAR* LexToString(UE::IoStore::EOnDemandInstallCasType CacheType);
+UE_API void LexFromString(UE::IoStore::EOnDemandInstallCasType& CacheType, FStringView String);
+ENUM_RANGE_BY_COUNT(UE::IoStore::EOnDemandInstallCasType, UE::IoStore::EOnDemandInstallCasType::Count);
+
 #undef UE_API
diff --git a/Engine/Source/Runtime/Experimental/IoStore/OnDemandCore/Public/IO/OnDemandError.h b/Engine/Source/Runtime/Experimental/IoStore/OnDemandCore/Public/IO/OnDemandError.h
--- a/Engine/Source/Runtime/Experimental/IoStore/OnDemandCore/Public/IO/OnDemandError.h
+++ b/Engine/Source/Runtime/Experimental/IoStore/OnDemandCore/Public/IO/OnDemandError.h
@@ -195,6 +195,7 @@ struct FChunkHashMismatchErrorContext
 struct FInstallCacheErrorContext
 {
 	FIoStatus		IoError = FIoStatus(EIoErrorCode::Unknown);
+	const TCHAR*	CasType = nullptr;
 	uint64			MaxCacheSize = 0;
 	uint64			CacheSize = 0;
 	uint64			RequestedSize = 0;
@@ -260,6 +261,7 @@ inline void SerializeForLog(FCbWriter& Writer, const UE::UnifiedError::IoStoreOn
 	Writer.BeginObject();
 	Writer.AddString(UTF8TEXTVIEW("$type"), TErrorStructFeatures<UE::UnifiedError::IoStoreOnDemand::FInstallCacheErrorContext>::GetErrorContextTypeNameAsString());
 	Writer.AddString(UTF8TEXTVIEW("$format"), TEXT("({ErrorMessage})"));
+	Writer.AddString(UTF8TEXTVIEW("CasType"), Ctx.CasType);
 	Writer.AddInteger(UTF8TEXTVIEW("ErrorCode"), int32(Ctx.IoError.GetErrorCode()));
 	Writer.AddString(UTF8TEXTVIEW("ErrorMessage"), Ctx.IoError.ToString());
 	Writer.AddInteger(UTF8TEXTVIEW("SystemErrorCode"), int32(Ctx.IoError.GetSystemErrorCode()));
