Geometry collection add the possibility to change NotifyTrailing from blueprint

#Synced-CL 51362814

#Synced-CL 51362814

diff --git a/Engine/Source/Runtime/Experimental/Chaos/Private/PhysicsProxy/GeometryCollectionPhysicsProxy.cpp b/Engine/Source/Runtime/Experimental/Chaos/Private/PhysicsProxy/GeometryCollectionPhysicsProxy.cpp
--- a/Engine/Source/Runtime/Experimental/Chaos/Private/PhysicsProxy/GeometryCollectionPhysicsProxy.cpp
+++ b/Engine/Source/Runtime/Experimental/Chaos/Private/PhysicsProxy/GeometryCollectionPhysicsProxy.cpp
@@ -3032,6 +3032,15 @@ void FGeometryCollectionPhysicsProxy::SetNotifyBreakings_External(bool bNotify)
 		});
 }
 
+void FGeometryCollectionPhysicsProxy::SetNotifyTrailings_External(bool bNotify)
+{
+	ExecuteOnPhysicsThread(*this,
+		[this, bNotify]()
+		{
+			Parameters.bGenerateTrailingData = bNotify;
+		});
+}
+
 void FGeometryCollectionPhysicsProxy::SetNotifyRemovals_External(bool bNotify)
 {
 	// nothing to do : there's no removal flag to set on the proxy game thread or physics thread side
diff --git a/Engine/Source/Runtime/Experimental/Chaos/Public/PhysicsProxy/GeometryCollectionPhysicsProxy.h b/Engine/Source/Runtime/Experimental/Chaos/Public/PhysicsProxy/GeometryCollectionPhysicsProxy.h
--- a/Engine/Source/Runtime/Experimental/Chaos/Public/PhysicsProxy/GeometryCollectionPhysicsProxy.h
+++ b/Engine/Source/Runtime/Experimental/Chaos/Public/PhysicsProxy/GeometryCollectionPhysicsProxy.h
@@ -508,6 +508,7 @@ public:
 	CHAOS_API void SetNotifyBreakings_External(bool bNotify);
 	CHAOS_API void SetNotifyRemovals_External(bool bNotify);
 	CHAOS_API void SetNotifyCrumblings_External(bool bNotify, bool bIncludeChildren);
+	CHAOS_API void SetNotifyTrailings_External(bool bNotify);
 
 	CHAOS_API void SetNotifyGlobalBreakings_External(bool bNotify);
 	CHAOS_API void SetNotifyGlobalRemovals_External(bool bNotify);
diff --git a/Engine/Source/Runtime/Experimental/GeometryCollectionEngine/Private/GeometryCollection/GeometryCollectionComponent.cpp b/Engine/Source/Runtime/Experimental/GeometryCollectionEngine/Private/GeometryCollection/GeometryCollectionComponent.cpp
--- a/Engine/Source/Runtime/Experimental/GeometryCollectionEngine/Private/GeometryCollection/GeometryCollectionComponent.cpp
+++ b/Engine/Source/Runtime/Experimental/GeometryCollectionEngine/Private/GeometryCollection/GeometryCollectionComponent.cpp
@@ -1214,6 +1214,20 @@ void UGeometryCollectionComponent::SetNotifyCrumblings(bool bNewNotifyCrumblings
 	}
 }
 
+void UGeometryCollectionComponent::SetNotifyTrailings(bool bNewNotifyTrailings)
+{
+	if (bNotifyTrailing != bNewNotifyTrailings)
+	{
+		if (PhysicsProxy)
+		{
+			PhysicsProxy->SetNotifyTrailings_External(bNewNotifyTrailings);
+		}
+		bNotifyTrailing = bNewNotifyTrailings;
+		RegisterForEvents();
+	}
+}
+
+
 void UGeometryCollectionComponent::SetNotifyGlobalBreaks(bool bNewNotifyGlobalBreaks)
 {
 	if (bNotifyGlobalBreaks != bNewNotifyGlobalBreaks)
diff --git a/Engine/Source/Runtime/Experimental/GeometryCollectionEngine/Public/GeometryCollection/GeometryCollectionComponent.h b/Engine/Source/Runtime/Experimental/GeometryCollectionEngine/Public/GeometryCollection/GeometryCollectionComponent.h
--- a/Engine/Source/Runtime/Experimental/GeometryCollectionEngine/Public/GeometryCollection/GeometryCollectionComponent.h
+++ b/Engine/Source/Runtime/Experimental/GeometryCollectionEngine/Public/GeometryCollection/GeometryCollectionComponent.h
@@ -1235,6 +1235,10 @@ public:
 	UFUNCTION(BlueprintCallable, Category = "Physics")
 	GEOMETRYCOLLECTIONENGINE_API void SetNotifyCrumblings(bool bNewNotifyCrumblings, bool bNewCrumblingEventIncludesChildren = false);
 
+	/** Changes whether or not this component will get future trailing notifications. */
+	UFUNCTION(BlueprintCallable, Category = "Physics")
+	GEOMETRYCOLLECTIONENGINE_API void SetNotifyTrailings(bool bNewNotifyTrailings);
+
 	/** Changes whether or not this component will get future global break notifications. */
 	UFUNCTION(BlueprintCallable, Category = "Physics")
 	GEOMETRYCOLLECTIONENGINE_API void SetNotifyGlobalBreaks(bool bNewNotifyGlobalBreaks);
