Chaos API: Shape Instance Part 1
Doing some prep work for the shape instance refactor, mostly to remove noise from later review.
- Removed AddShape from body in favor of CreateShape.
- Renamed the RigidCollider file.

#Synced-CL 51413247

diff --git a/Engine/Plugins/Experimental/ChaosDataflowSolver/Source/ChaosDataflowSolver/Private/RigidComponents/InstancedRigidComponentInterface.h b/Engine/Plugins/Experimental/ChaosDataflowSolver/Source/ChaosDataflowSolver/Private/RigidComponents/InstancedRigidComponentInterface.h
--- a/Engine/Plugins/Experimental/ChaosDataflowSolver/Source/ChaosDataflowSolver/Private/RigidComponents/InstancedRigidComponentInterface.h
+++ b/Engine/Plugins/Experimental/ChaosDataflowSolver/Source/ChaosDataflowSolver/Private/RigidComponents/InstancedRigidComponentInterface.h
@@ -69,7 +69,7 @@ class InstancedRigidComponentInterface
 		
 							for (int32 ShapeIndex = 0; ShapeIndex < NumShapes; ShapeIndex++)
 							{
-								BodyPtr->AddShape(Shapes[ShapeIndex].Release());
+								BodyPtr->CreateShape(MoveTemp(Shapes[ShapeIndex]));
 							}
 
 							BodyPtr->InitTransform(Transform);
diff --git a/Engine/Plugins/Experimental/ChaosDataflowSolver/Source/ChaosDataflowSolver/Private/RigidComponents/PrimitiveRigidComponentInterface.h b/Engine/Plugins/Experimental/ChaosDataflowSolver/Source/ChaosDataflowSolver/Private/RigidComponents/PrimitiveRigidComponentInterface.h
--- a/Engine/Plugins/Experimental/ChaosDataflowSolver/Source/ChaosDataflowSolver/Private/RigidComponents/PrimitiveRigidComponentInterface.h
+++ b/Engine/Plugins/Experimental/ChaosDataflowSolver/Source/ChaosDataflowSolver/Private/RigidComponents/PrimitiveRigidComponentInterface.h
@@ -74,7 +74,7 @@ class PrimitiveRigidComponentInterface
 
 						for (int32 ShapeIndex = 0; ShapeIndex < NumShapes; ++ShapeIndex)
 						{
-							BodyPtr->AddShape(Shapes[ShapeIndex].Release());
+							BodyPtr->CreateShape(MoveTemp(Shapes[ShapeIndex])); 
 						}
 
 						Transform.SetScale3D(FVector(1.0f, 1.0f, 1.0f));
diff --git a/Engine/Plugins/Experimental/ChaosRigidPhysicsAsync/Source/ChaosRigidPhysicsAsync/Private/ChaosRigidPhysicsAsync/RigidBodyAsync.cpp b/Engine/Plugins/Experimental/ChaosRigidPhysicsAsync/Source/ChaosRigidPhysicsAsync/Private/ChaosRigidPhysicsAsync/RigidBodyAsync.cpp
--- a/Engine/Plugins/Experimental/ChaosRigidPhysicsAsync/Source/ChaosRigidPhysicsAsync/Private/ChaosRigidPhysicsAsync/RigidBodyAsync.cpp
+++ b/Engine/Plugins/Experimental/ChaosRigidPhysicsAsync/Source/ChaosRigidPhysicsAsync/Private/ChaosRigidPhysicsAsync/RigidBodyAsync.cpp
@@ -4,7 +4,7 @@
 #include "ChaosRigidPhysicsAsync/RigidSceneAsync.h"
 #include "Chaos/ShapeInstance.h"
 #include "RigidPhysics/RigidBody.h"
-#include "RigidPhysics/RigidCollider.h"
+#include "RigidPhysics/RigidShapeInstance.h"
 
 #if UE_RIGIDPHYSICS_API_ENABLED
 
diff --git a/Engine/Plugins/Experimental/ChaosRigidPhysicsAsync/Tests/Private/BodyTests.cpp b/Engine/Plugins/Experimental/ChaosRigidPhysicsAsync/Tests/Private/BodyTests.cpp
--- a/Engine/Plugins/Experimental/ChaosRigidPhysicsAsync/Tests/Private/BodyTests.cpp
+++ b/Engine/Plugins/Experimental/ChaosRigidPhysicsAsync/Tests/Private/BodyTests.cpp
@@ -8,10 +8,10 @@
 #include "ChaosRigidPhysicsAsync/RigidSceneAsync.h"
 #include "ChaosRigidPhysicsAsync/RigidSceneSettingsAsync.h"
 #include "RigidPhysics/RigidBody.h"
-#include "RigidPhysics/RigidCollider.h"
 #include "RigidPhysics/RigidContext.h"
 #include "RigidPhysics/RigidModifier.h"
 #include "RigidPhysics/RigidScene.h"
+#include "RigidPhysics/RigidShapeInstance.h"
 
 #include "RigidTestUtils.h"
 
@@ -36,13 +36,13 @@ namespace Chaos::LowLevelTest
 		{
 			TRigidBodyPtr<FRigidContextGameRW> Body1Ptr = Context->CreateBody(FRigidDebugName("Body1"), ERigidMovementType::Static);
 			REQUIRE(Body1Ptr);
-			Body1Ptr->AddShape(MakeBoxShape(FVector3f(1000, 1000, 100)));
+			Body1Ptr->CreateShape(MakeBoxShape(FVector3f(1000, 1000, 100)));
 			Body1Ptr->InitTransform(FTransform(FVector(0, 0, -50)));
 			Body1Ptr->Activate();
 
 			TRigidBodyPtr<FRigidContextGameRW> Body2Ptr = Context->CreateBody(FRigidDebugName("Body2"), ERigidMovementType::Dynamic);
 			REQUIRE(Body2Ptr);
-			Body2Ptr->AddShape(MakeBoxShape(FVector3f(100, 100, 100)));
+			Body2Ptr->CreateShape(MakeBoxShape(FVector3f(100, 100, 100)));
 			Body2Ptr->SetIsSleeping(false);
 			Body2Ptr->SetMass(1000);
 			Body2Ptr->SetInertia(FVector3d(MakeSolidBoxInertia(1000, FVector3f(100, 100, 100))));
diff --git a/Engine/Plugins/Experimental/ChaosRigidPhysicsAsync/Tests/Private/JointConstraintTests.cpp b/Engine/Plugins/Experimental/ChaosRigidPhysicsAsync/Tests/Private/JointConstraintTests.cpp
--- a/Engine/Plugins/Experimental/ChaosRigidPhysicsAsync/Tests/Private/JointConstraintTests.cpp
+++ b/Engine/Plugins/Experimental/ChaosRigidPhysicsAsync/Tests/Private/JointConstraintTests.cpp
@@ -8,9 +8,9 @@
 
 #include "ChaosRigidPhysicsAsyncTest.h"
 #include "RigidPhysics/RigidBody.h"
-#include "RigidPhysics/RigidCollider.h"
 #include "RigidPhysics/RigidContext.h"
 #include "RigidPhysics/RigidScene.h"
+#include "RigidPhysics/RigidShapeInstance.h"
 #include "RigidPhysics/JointConstraint6DOF.h"
 #include "RigidTestFixture.h"
 #include "RigidTestUtils.h"
@@ -28,7 +28,7 @@ namespace Chaos::LowLevelTest
 		{
 			TRigidBodyPtr<FRigidContextGameRW> BodyPtr = Context->CreateBody(FRigidDebugName(DebugName), MovementType);
 			REQUIRE(BodyPtr);
-			BodyPtr->AddShape(MakeBoxShape(FVector3f(100, 100, 100)));
+			BodyPtr->CreateShape(MakeBoxShape(FVector3f(100, 100, 100)));
 			BodyPtr->InitTransform(FTransform(Location));
 			if (MovementType == ERigidMovementType::Dynamic)
 			{
diff --git a/Engine/Plugins/Experimental/ChaosRigidPhysicsAsync/Tests/Private/RigidTestFixture.cpp b/Engine/Plugins/Experimental/ChaosRigidPhysicsAsync/Tests/Private/RigidTestFixture.cpp
--- a/Engine/Plugins/Experimental/ChaosRigidPhysicsAsync/Tests/Private/RigidTestFixture.cpp
+++ b/Engine/Plugins/Experimental/ChaosRigidPhysicsAsync/Tests/Private/RigidTestFixture.cpp
@@ -12,8 +12,8 @@
 #include "ChaosRigidPhysicsAsync/RigidSceneAsync.h"
 
 #include "RigidPhysics/RigidBody.h"
-#include "RigidPhysics/RigidCollider.h"
 #include "RigidPhysics/RigidContext.h"
+#include "RigidPhysics/RigidShapeInstance.h"
 #include "TestSceneModifier.h"
 #include "RigidTestUtils.h"
 
diff --git a/Engine/Plugins/Experimental/ChaosRigidPhysicsAsync/Tests/Private/RigidTestUtils.h b/Engine/Plugins/Experimental/ChaosRigidPhysicsAsync/Tests/Private/RigidTestUtils.h
--- a/Engine/Plugins/Experimental/ChaosRigidPhysicsAsync/Tests/Private/RigidTestUtils.h
+++ b/Engine/Plugins/Experimental/ChaosRigidPhysicsAsync/Tests/Private/RigidTestUtils.h
@@ -7,8 +7,8 @@
 #include "Chaos/ImplicitFwd.h"
 #include "Math/Bounds.h"
 #include "Math/Transform.h"
-#include "RigidPhysics/RigidCollider.h"
 #include "RigidPhysics/RigidLog.h"
+#include "RigidPhysics/RigidShapeInstance.h"
 
 class UPhysicalMaterial;
 
diff --git a/Engine/Plugins/Experimental/ChaosRigidPhysicsAsync/Tests/Private/SceneLockBenchmarks.cpp b/Engine/Plugins/Experimental/ChaosRigidPhysicsAsync/Tests/Private/SceneLockBenchmarks.cpp
--- a/Engine/Plugins/Experimental/ChaosRigidPhysicsAsync/Tests/Private/SceneLockBenchmarks.cpp
+++ b/Engine/Plugins/Experimental/ChaosRigidPhysicsAsync/Tests/Private/SceneLockBenchmarks.cpp
@@ -10,8 +10,8 @@
 #include "Physics/Experimental/PhysInterface_Chaos.h"
 #include "PhysicsProxy/SingleParticlePhysicsProxy.h"
 #include "RigidPhysics/RigidBody.h"
-#include "RigidPhysics/RigidCollider.h"
 #include "RigidPhysics/RigidContext.h"
+#include "RigidPhysics/RigidShapeInstance.h"
 
 #include "RigidTestUtils.h"
 
@@ -44,7 +44,7 @@ namespace Chaos::LowLevelTest
 			for (int32 BodyIndex = 0; BodyIndex < NumBodies; ++BodyIndex)
 			{
 				TRigidBodyPtr<FRigidContextGameRW> BodyPtr = Context->CreateBody(FRigidDebugName::Make(TEXT("Body{0}"), BodyIndex), ERigidMovementType::Dynamic);
-				BodyPtr->AddShape(MakeBoxShape(FVector3f(50, 50, 50)));
+				BodyPtr->CreateShape(MakeBoxShape(FVector3f(50, 50, 50)));
 				BodyPtr->Activate();
 
 				BodyHandles.Add(BodyPtr);
@@ -148,7 +148,7 @@ namespace Chaos::LowLevelTest
 			for (int32 BodyIndex = 0; BodyIndex < NumBodies; ++BodyIndex)
 			{
 				TRigidBodyPtr<FRigidContextGameRW> BodyPtr = Context->CreateBody(FRigidDebugName::Make(TEXT("Body{0}"), BodyIndex), ERigidMovementType::Dynamic);
-				BodyPtr->AddShape(MakeBoxShape(FVector3f(50, 50, 50)));
+				BodyPtr->CreateShape(MakeBoxShape(FVector3f(50, 50, 50)));
 				BodyPtr->Activate();
 
 				BodyHandles.Add(BodyPtr);
diff --git a/Engine/Plugins/Experimental/ChaosRigidPhysicsAsync/Tests/Private/SceneTests.cpp b/Engine/Plugins/Experimental/ChaosRigidPhysicsAsync/Tests/Private/SceneTests.cpp
--- a/Engine/Plugins/Experimental/ChaosRigidPhysicsAsync/Tests/Private/SceneTests.cpp
+++ b/Engine/Plugins/Experimental/ChaosRigidPhysicsAsync/Tests/Private/SceneTests.cpp
@@ -10,10 +10,10 @@
 #include "Framework/Threading.h"
 #include "HAL/Thread.h"
 #include "RigidPhysics/RigidBody.h"
-#include "RigidPhysics/RigidCollider.h"
 #include "RigidPhysics/RigidContext.h"
 #include "RigidPhysics/RigidModifier.h"
 #include "RigidPhysics/RigidScene.h"
+#include "RigidPhysics/RigidShapeInstance.h"
 #include "TestSceneModifier.h"
 
 #include "RigidTestUtils.h"
@@ -144,8 +144,8 @@ namespace Chaos::LowLevelTest
 				CHECK(DynamicBodyHandle.Pin(Context)->GetTypeId() == FRigidBodyAsyncGT::GetStaticTypeId());
 
 				// Add Collision
-				StaticBodyHandle.Pin(Context)->AddShape(MakeBoxShape(FVector3f(100, 100, 100)));
-				DynamicBodyHandle.Pin(Context)->AddShape(MakeBoxShape(FVector3f(100, 100, 100)));
+				StaticBodyHandle.Pin(Context)->CreateShape(MakeBoxShape(FVector3f(100, 100, 100)));
+				DynamicBodyHandle.Pin(Context)->CreateShape(MakeBoxShape(FVector3f(100, 100, 100)));
 
 				// Activate Bodies
 				StaticBodyHandle.Pin(Context)->Activate();
@@ -211,7 +211,7 @@ namespace Chaos::LowLevelTest
 				CHECK(BodyHandleDupe.Pin(Context).Get() == BodyHandle.Pin(Context).Get());
 
 				// Add Collision
-				BodyHandle.Pin(Context)->AddShape(MakeBoxShape(FVector3f(100, 100, 100)));
+				BodyHandle.Pin(Context)->CreateShape(MakeBoxShape(FVector3f(100, 100, 100)));
 
 				Context->Tick(Dt);
 
@@ -341,12 +341,12 @@ namespace Chaos::LowLevelTest
 			Context->RegisterModifier(&SceneModifier);
 
 			TRigidBodyPtr<FRigidContextGameRW> BodyPtr1 = Context->CreateBody(FRigidDebugName::Make(TEXT("Body1")), ERigidMovementType::Dynamic);
-			BodyPtr1->AddShape(MakeBoxShape(FVector3f(100, 100, 100)));
+			BodyPtr1->CreateShape(MakeBoxShape(FVector3f(100, 100, 100)));
 			BodyPtr1->Activate();
 			BodyHandle1 = BodyPtr1;
 
 			TRigidBodyPtr<FRigidContextGameRW> BodyPtr2 = Context->CreateBody(FRigidDebugName::Make(TEXT("Body2")), ERigidMovementType::Dynamic);
-			BodyPtr2->AddShape(MakeBoxShape(FVector3f(100, 100, 100)));
+			BodyPtr2->CreateShape(MakeBoxShape(FVector3f(100, 100, 100)));
 			BodyPtr2->Activate();
 			BodyHandle2 = BodyPtr2;
 
@@ -401,7 +401,7 @@ namespace Chaos::LowLevelTest
 		{
 			if (TRigidBodyPtr<FRigidContextGameRW> BodyPtr = Context->CreateBody(FRigidDebugName("Body2"), ERigidMovementType::Dynamic))
 			{
-				BodyPtr->AddShape(MakeBoxShape(FVector3f(100, 100, 100)));
+				BodyPtr->CreateShape(MakeBoxShape(FVector3f(100, 100, 100)));
 				BodyPtr->SetIsSleeping(false);
 				BodyPtr->SetMass(1000);
 				BodyPtr->SetInertia(FVector3d(MakeSolidBoxInertia(1000, FVector3f(100, 100, 100))));
diff --git a/Engine/Source/Runtime/Experimental/RigidPhysics/Private/RigidPhysics/RigidCollider.cpp b/Engine/Source/Runtime/Experimental/RigidPhysics/Private/RigidPhysics/RigidCollider.cpp
deleted file mode 100644
diff --git a/Engine/Source/Runtime/Experimental/RigidPhysics/Private/RigidPhysics/RigidShapeInstance.cpp b/Engine/Source/Runtime/Experimental/RigidPhysics/Private/RigidPhysics/RigidShapeInstance.cpp
new file mode 100644
--- /dev/null
+++ b/Engine/Source/Runtime/Experimental/RigidPhysics/Private/RigidPhysics/RigidShapeInstance.cpp
@@ -0,0 +1,40 @@
+// Copyright Epic Games, Inc. All Rights Reserved.
+#include "RigidPhysics/RigidShapeInstance.h"
+#include "Physics/Experimental/ChaosInterfaceUtils.h"
+#include "RigidPhysics/RigidLog.h"
+
+#if UE_RIGIDPHYSICS_API_ENABLED
+
+namespace UE::Physics
+{
+	FRigidShapeInstance::FRigidShapeInstance()
+		: PerShapeData(nullptr)
+	{
+	}
+
+	FRigidShapeInstance::FRigidShapeInstance(Chaos::FPerShapeData* InPerShapeData)
+		: PerShapeData(InPerShapeData)
+	{
+	}
+
+	FRigidShapeInstance::~FRigidShapeInstance()
+	{
+	}
+
+	Chaos::FPerShapeData* FRigidShapeInstance::GetPerShapeData() const
+	{
+		return PerShapeData;
+	}
+
+	Chaos::FConstImplicitObjectPtr FRigidShapeInstance::GetImplicitObject() const
+	{
+		if (PerShapeData != nullptr)
+		{
+			return PerShapeData->GetGeometry();
+		}
+		return Chaos::FConstImplicitObjectPtr();
+	}
+
+}
+
+#endif // UE_RIGIDPHYSICS_API_ENABLED
\ No newline at end of file
diff --git a/Engine/Source/Runtime/Experimental/RigidPhysics/Public/RigidPhysics/Internal/IRigidBody.h b/Engine/Source/Runtime/Experimental/RigidPhysics/Public/RigidPhysics/Internal/IRigidBody.h
--- a/Engine/Source/Runtime/Experimental/RigidPhysics/Public/RigidPhysics/Internal/IRigidBody.h
+++ b/Engine/Source/Runtime/Experimental/RigidPhysics/Public/RigidPhysics/Internal/IRigidBody.h
@@ -8,8 +8,8 @@
 #include "Math/Bounds.h"
 #include "Math/Transform.h"
 #include "RigidPhysics/RigidBodyHandle.h"
-#include "RigidPhysics/RigidCollider.h"
 #include "RigidPhysics/RigidLog.h"
+#include "RigidPhysics/RigidShapeInstance.h"
 #include "RigidPhysics/RigidTyped.h"
 
 namespace UE::Physics
diff --git a/Engine/Source/Runtime/Experimental/RigidPhysics/Public/RigidPhysics/RigidBody.h b/Engine/Source/Runtime/Experimental/RigidPhysics/Public/RigidPhysics/RigidBody.h
--- a/Engine/Source/Runtime/Experimental/RigidPhysics/Public/RigidPhysics/RigidBody.h
+++ b/Engine/Source/Runtime/Experimental/RigidPhysics/Public/RigidPhysics/RigidBody.h
@@ -9,9 +9,9 @@
 #include "Math/Transform.h"
 #include "RigidPhysics/Internal/IRigidBody.h"
 #include "RigidPhysics/RigidBodyHandle.h"
-#include "RigidPhysics/RigidCollider.h"
 #include "RigidPhysics/RigidObjectId.h"
 #include "RigidPhysics/RigidObjectPtr.h"
+#include "RigidPhysics/RigidShapeInstance.h"
 
 namespace UE::Physics
 {
@@ -113,11 +113,18 @@ namespace UE::Physics
 			return BodyRaw->GetShape(InShapeIndex);
 		}
 
-		void AddShape(const FConstRigidShapeInstancePtr& InShape)
+		void CreateShape(const FConstRigidShapeInstancePtr& InShape)
 		{
 			BodyRaw->AddShape(InShape);
 		}
 
+		UE_RIGIDPHYSICS_LEGACY_API FConstRigidShapeInstancePtr CreateShape(TUniquePtr<Chaos::FPerShapeData>&& InShapeData)
+		{
+			FConstRigidShapeInstancePtr Ptr(InShapeData.Release());
+			CreateShape(Ptr);
+			return Ptr;
+		}
+
 		void InitTransform(const FTransform3d& InTransform)
 		{
 			BodyRaw->InitTransform(InTransform);
diff --git a/Engine/Source/Runtime/Experimental/RigidPhysics/Public/RigidPhysics/RigidCollider.h b/Engine/Source/Runtime/Experimental/RigidPhysics/Public/RigidPhysics/RigidCollider.h
deleted file mode 100644
diff --git a/Engine/Source/Runtime/Experimental/RigidPhysics/Public/RigidPhysics/RigidFwd.h b/Engine/Source/Runtime/Experimental/RigidPhysics/Public/RigidPhysics/RigidFwd.h
--- a/Engine/Source/Runtime/Experimental/RigidPhysics/Public/RigidPhysics/RigidFwd.h
+++ b/Engine/Source/Runtime/Experimental/RigidPhysics/Public/RigidPhysics/RigidFwd.h
@@ -16,6 +16,8 @@
 #include "Templates/RefCounting.h"
 #include "Templates/SharedPointer.h"
 
+#define UE_RIGIDPHYSICS_LEGACY_API UE_INTERNAL /*UE_DEPRECATED(5.8, "TODO")*/
+
 #if DO_CHECK
 
 // Enable checks in the Rigid accessors. This can add significant overhead.
diff --git a/Engine/Source/Runtime/Experimental/RigidPhysics/Public/RigidPhysics/RigidShapeInstance.h b/Engine/Source/Runtime/Experimental/RigidPhysics/Public/RigidPhysics/RigidShapeInstance.h
new file mode 100644
--- /dev/null
+++ b/Engine/Source/Runtime/Experimental/RigidPhysics/Public/RigidPhysics/RigidShapeInstance.h
@@ -0,0 +1,116 @@
+// Copyright Epic Games, Inc. All Rights Reserved.
+#pragma once
+
+#include "RigidPhysics/RigidFwd.h"
+
+#if UE_RIGIDPHYSICS_API_ENABLED
+
+#include "Chaos/ImplicitFwd.h"
+#include "Math/Bounds.h"
+#include "Math/Transform.h"
+#include "RigidPhysics/RigidLog.h"
+#include "RigidPhysics/RigidTyped.h"
+
+class UPhysicalMaterial;
+
+namespace Chaos
+{
+	class FPerShapeData;
+}
+
+namespace UE::Physics
+{
+	// TODO_CHAOS: Should this be an interface with scene-specific implementation?
+	// Not if we want to use this data type in assets/bodyinstance etc. but it would
+	// allow us to use FPerShapeData and FImplicitObject in the ChaosAsync but move
+	// away from them in new implementations.
+	class UE_INTERNAL FRigidShapeInstance
+	{
+	public:
+		UE_INTERNAL RIGIDPHYSICS_API FRigidShapeInstance();
+		UE_INTERNAL RIGIDPHYSICS_API virtual ~FRigidShapeInstance();
+
+		// Existing Chaos implementation exposed for internal use only
+		UE_INTERNAL RIGIDPHYSICS_API FRigidShapeInstance(Chaos::FPerShapeData* InPerShapeData);
+		UE_INTERNAL RIGIDPHYSICS_API Chaos::FPerShapeData* GetPerShapeData() const;
+		UE_INTERNAL RIGIDPHYSICS_API Chaos::FConstImplicitObjectPtr GetImplicitObject() const;
+
+	private:
+		// NOTE: Danger...we hold a pointer to a shape which is the current implementation, 
+		// but its lifetime is currently managed by the Proxy in a UniquePtr.
+		// See FRigidBodyAsyncGT::AddShape
+		Chaos::FPerShapeData* PerShapeData;
+	};
+
+
+	//
+	// A pointer to a collider. Eventually this will just be the equivalent of a raw pointer
+	// to FRigidShapeInstance, but for now we contain the FRigidShapeInstance by value. This is so that
+	// we don't need to store an array of FRigidShapeInstance objects in our FRigidBody class until
+	// we move away from it being a thin wrapper around FPerShapeData (the particle class
+	// already has an array of PerShapeData so we don't want to double up).
+	//
+	template<bool bConstParam>
+	class TRigidShapeInstancePtr
+	{
+	public:
+		static const bool bIsPtrToConst = bConstParam;
+
+		TRigidShapeInstancePtr()
+			: ShapeInstance(nullptr)
+		{
+		}
+
+		TRigidShapeInstancePtr(const TRigidShapeInstancePtr& R)
+			: ShapeInstance(R.Get()->GetPerShapeData())
+		{
+		}
+
+		// Auto-cast non-const to const ptr
+		TRigidShapeInstancePtr(const TRigidShapeInstancePtr<false>& R) requires (bIsPtrToConst)
+			: ShapeInstance(R.Get()->GetPerShapeData())
+		{
+		}
+
+		bool IsValid() const
+		{
+			return ShapeInstance.GetPerShapeData() != nullptr;
+		}
+
+		void Reset()
+		{
+			ShapeInstance = FRigidShapeInstance();
+		}
+
+		operator bool() const
+		{
+			return IsValid();
+		}
+
+		FRigidShapeInstance* operator->() const requires (!bIsPtrToConst)
+		{
+			return &ShapeInstance;
+		}
+
+		const FRigidShapeInstance* operator->() const  requires (bIsPtrToConst)
+		{ 
+			return &ShapeInstance;
+		}
+
+
+		UE_INTERNAL TRigidShapeInstancePtr(Chaos::FPerShapeData* InShapeInstance)
+			: ShapeInstance(InShapeInstance)
+		{
+		}
+
+		UE_INTERNAL const FRigidShapeInstance* Get() const
+		{
+			return &ShapeInstance;
+		}
+
+	private:
+		mutable FRigidShapeInstance ShapeInstance;
+	};
+}
+
+#endif
