Geometry collection : make sure UpdateKinematicProperties respect the existing sleeping state

#Synced-CL 51402414

diff --git a/Engine/Source/Runtime/Experimental/Chaos/Private/Chaos/PBDRigidClusteringAlgo.cpp b/Engine/Source/Runtime/Experimental/Chaos/Private/Chaos/PBDRigidClusteringAlgo.cpp
--- a/Engine/Source/Runtime/Experimental/Chaos/Private/Chaos/PBDRigidClusteringAlgo.cpp
+++ b/Engine/Source/Runtime/Experimental/Chaos/Private/Chaos/PBDRigidClusteringAlgo.cpp
@@ -134,6 +134,9 @@ namespace Chaos
 		SCOPE_CYCLE_COUNTER(STAT_UpdateKinematicProperties);
 		typedef FPBDRigidClusteredParticleHandle* FClusterHandle;
 
+		// preserve a previous sleeping state 
+		const bool bWasSleeping = InParent && InParent->IsSleeping();
+
 		EObjectStateType ObjectState = EObjectStateType::Dynamic;
 		check(InParent != nullptr);
 		if (FClusterHandle ClusteredCurrentNode = InParent->CastToClustered())
@@ -191,8 +194,13 @@ namespace Chaos
 					}
 				}
 
+				if (bWasSleeping && ObjectState == Chaos::EObjectStateType::Dynamic)
+				{
+					ObjectState = Chaos::EObjectStateType::Sleeping;
+				}
+
 				MEvolution.SetParticleObjectState(ClusteredCurrentNode, ObjectState);
-				if (ObjectState == Chaos::EObjectStateType::Dynamic)
+				if (ObjectState != Chaos::EObjectStateType::Kinematic)
 				{
 					MEvolution.SetParticleKinematicTarget(ClusteredCurrentNode, FKinematicTarget());
 				}
