IAD installer puts MMap chunks into MMap CAS if the MMap CAS is enabled.

#Synced-CL 51330535

diff --git a/Engine/Source/Runtime/Experimental/IoStore/OnDemand/Private/OnDemandConfig.cpp b/Engine/Source/Runtime/Experimental/IoStore/OnDemand/Private/OnDemandConfig.cpp
--- a/Engine/Source/Runtime/Experimental/IoStore/OnDemand/Private/OnDemandConfig.cpp
+++ b/Engine/Source/Runtime/Experimental/IoStore/OnDemand/Private/OnDemandConfig.cpp
@@ -550,11 +550,17 @@ TIoStatusOr<FOnDemandInstallCacheConfig> TryParseInstallCacheConfig(const TCHAR*
 	GConfig->GetBool(TEXT("OnDemandInstall"), TEXT("MMapCache.Enabled"), bUseMMapCache, GEngineIni);
 	if (bUseMMapCache)
 	{
-		FOnDemandInstallCasConfig& CasConfig = OutConfig.CasConfig.Emplace_GetRef();
-		CasConfig.Type = EOnDemandInstallCasType::MMap;
-		CasConfig.CasSubdirectory = TEXT("mmap");
+		const bool bPlatformSupportsMmap = FPlatformProperties::SupportsMemoryMappedFiles();
+		UE_CLOGF(!bPlatformSupportsMmap, LogIoStoreOnDemand, Error, "MMap cache cannot be enabled on a platform that does not support memory mapping!");
 
-		TryParseCasConfig(CasConfig, CasConfig.Type);
+		if (ensure(bPlatformSupportsMmap))
+		{
+			FOnDemandInstallCasConfig& CasConfig = OutConfig.CasConfig.Emplace_GetRef();
+			CasConfig.Type = EOnDemandInstallCasType::MMap;
+			CasConfig.CasSubdirectory = TEXT("mmap");
+
+			TryParseCasConfig(CasConfig, CasConfig.Type);
+		}
 	}
 
 	return OutConfig;
diff --git a/Engine/Source/Runtime/Experimental/IoStore/OnDemand/Private/OnDemandContentInstaller.cpp b/Engine/Source/Runtime/Experimental/IoStore/OnDemand/Private/OnDemandContentInstaller.cpp
--- a/Engine/Source/Runtime/Experimental/IoStore/OnDemand/Private/OnDemandContentInstaller.cpp
+++ b/Engine/Source/Runtime/Experimental/IoStore/OnDemand/Private/OnDemandContentInstaller.cpp
@@ -415,6 +415,7 @@ FOnDemandContentInstaller::FRequest::FInstall::FInstall(
 FOnDemandContentInstaller::FOnDemandContentInstaller(FOnDemandIoStore& InIoStore)
 	: IoStore(InIoStore)
 	, InstallerPipe(TEXT("IoStoreOnDemandInstallerPipe"))
+	, bMMapCasEnabled(InIoStore.InstallCache->IsCasEnabled(EOnDemandInstallCasType::MMap))
 {
 }
 
@@ -638,6 +639,16 @@ void FOnDemandContentInstaller::ReportAnalytics(TArray<FAnalyticsEventAttribute>
 	FOnDemandContentInstallerStats::ReportAnalytics(OutAnalyticsArray);
 }
 
+EOnDemandInstallCasType FOnDemandContentInstaller::GetCasType(EIoChunkType ChunkType) const
+{
+	EOnDemandInstallCasType CasType = EOnDemandInstallCasType::General;
+	if (bMMapCasEnabled && ChunkType == EIoChunkType::MemoryMappedBulkData)
+	{
+		CasType = EOnDemandInstallCasType::MMap;
+	}
+	return CasType;
+}
+
 bool FOnDemandContentInstaller::CanExecuteRequest(FRequest& Request) const
 {
 	bool bExecute = false;
@@ -727,11 +738,18 @@ void FOnDemandContentInstaller::TryExecuteNextRequest()
 
 					if (bExecuteFlush)
 					{
-						// TODO: IadMultiCAS: content installer needs to determine which chunks go to which CAS
-						FResult Result = IoStore.InstallCache->Flush(EOnDemandInstallCasType::General);
-						if (Result.HasError())
+						for (EOnDemandInstallCasType CasType : TEnumRange<EOnDemandInstallCasType>())
 						{
-							UE_LOG(LogIoStoreOnDemand, Error, TEXT("Failed to flush pending chunks, Reason %s"), *LexToString(Result.GetError()));
+							if (!IoStore.InstallCache->IsCasEnabled(CasType))
+							{
+								continue;
+							}
+
+							FResult Result = IoStore.InstallCache->Flush(CasType);
+							if (Result.HasError())
+							{
+								UE_LOGF(LogIoStoreOnDemand, Error, "Failed to flush pending chunks from %ls, Reason %ls", LexToString(CasType), *LexToString(Result.GetError()));
+							}
 						}
 					}
 				},
@@ -1046,7 +1064,8 @@ void FOnDemandContentInstaller::ExecuteInstallRequest(FRequest& Request)
 	PinCachedChunks(InstallRequest, [&Request, &TotalContentSize, &TotalInstallSize](const int32 ContainerIndex, const int32 EntryIndex, const bool bCached)
 	{
 		FRequest::FInstall& InstallRequest = Request.AsInstall();
-		const FOnDemandChunkEntry& Entry = InstallRequest.ResolvedChunks[ContainerIndex].Container->ChunkEntries[EntryIndex];
+		const FSharedOnDemandContainer& Container = InstallRequest.ResolvedChunks[ContainerIndex].Container;
+		const FOnDemandChunkEntry& Entry = Container->ChunkEntries[EntryIndex];
 		TotalContentSize += Entry.GetDiskSize();
 
 		if (bCached)
@@ -1083,28 +1102,34 @@ void FOnDemandContentInstaller::ExecuteInstallRequest(FRequest& Request)
 	if (CVars::EnableConcurentInstallRequests() == false)
 	{
 		// Make sure we have enough space in the cache
-		uint64 BytesToInstall = 0;
+		TStaticArray<uint64, EOnDemandInstallCasType::Count> BytesToInstall(EInPlace::InPlace, 0);
 		{
 			TSet<FIoHash> ChunksToInstall;
 			ChunksToInstall.Reserve(InstallRequest.HttpRequestHandles.Num());
 			for (FChunkHttpRequestHandle& HttpRequest : InstallRequest.HttpRequestHandles)
 			{
-				FSharedOnDemandContainer& Container = InstallRequest.ResolvedChunks[HttpRequest.ContainerIndex].Container;
+				const FSharedOnDemandContainer& Container = InstallRequest.ResolvedChunks[HttpRequest.ContainerIndex].Container;
 				const FOnDemandChunkEntry& ChunkEntry = Container->ChunkEntries[HttpRequest.EntryIndex];
 
 				bool bAlreadyInSet = false;
 				ChunksToInstall.Add(ChunkEntry.Hash, &bAlreadyInSet);
 				if (!bAlreadyInSet)
 				{
-					BytesToInstall += ChunkEntry.GetDiskSize();
+					const FIoChunkId& ChunkId = Container->ChunkIds[HttpRequest.EntryIndex];
+					BytesToInstall[GetCasType(ChunkId.GetChunkType())] += ChunkEntry.GetDiskSize();
 				}
 			}
 		}
 
-		// TODO: IadMultiCAS: content installer needs to determine which chunks go to which CAS
-		if (Request.Result = IoStore.InstallCache->Purge(EOnDemandInstallCasType::General, BytesToInstall); Request.Result.HasError())
+		for (EOnDemandInstallCasType CasType : TEnumRange<EOnDemandInstallCasType>())
 		{
-			return CompleteInstallRequest(Request);
+			if (BytesToInstall[CasType] > 0)
+			{
+				if (Request.Result = IoStore.InstallCache->Purge(CasType, BytesToInstall[CasType]); Request.Result.HasError())
+				{
+					return CompleteInstallRequest(Request);
+				}
+			}
 		}
 	}
 
@@ -1342,8 +1367,18 @@ void FOnDemandContentInstaller::ProcessDownloadedChunk(
 		{
 			TRACE_COUNTER_INCREMENT(OnDemandContentInstaller_HttpRequestsSuccess);
 
-			// TODO: IadMultiCAS: content installer needs to determine which chunks go to which CAS
-			TResult<bool> PutResult = IoStore.InstallCache->PutChunk(EOnDemandInstallCasType::General, MoveTemp(Chunk), ChunkHash);
+			// Peek the first request waiting on the chunk to determine CasType
+			EOnDemandInstallCasType CasType = EOnDemandInstallCasType::General;
+			{
+				const FChunkHttpRequestHandle& HttpRequest = *ChunkRequest.ChunkRequestHttpHandles.PeekHead();
+				const FRequest& Request = *HttpRequest.OwnerRequest;
+				const FRequest::FInstall& InstallRequest = Request.AsInstall();
+				const FSharedOnDemandContainer& Container = InstallRequest.ResolvedChunks[HttpRequest.ContainerIndex].Container;
+				const FIoChunkId& ChunkId = Container->ChunkIds[HttpRequest.EntryIndex];
+				CasType = GetCasType(ChunkId.GetChunkType());
+			}
+
+			TResult<bool> PutResult = IoStore.InstallCache->PutChunk(CasType, MoveTemp(Chunk), ChunkHash);
 			if (PutResult.HasError())
 			{
 				Result = MakeError(PutResult.StealError());
@@ -1371,17 +1406,21 @@ void FOnDemandContentInstaller::ProcessDownloadedChunk(
 
 						const FRequest& Request = *ChunkHttpRequestHandle->OwnerRequest;
 						const FRequest::FInstall& InstallRequest = Request.AsInstall();
-						const FSharedOnDemandContainer& Container = InstallRequest.ResolvedChunks[ChunkHttpRequestHandle->ContainerIndex].Container;
-						const FOnDemandChunkEntry& ChunkEntry = Container->ChunkEntries[ChunkHttpRequestHandle->EntryIndex];
+						const FSharedOnDemandContainer& Container = InstallRequest.ResolvedChunks[ChunkHttpRequestHandle->ContainerIndex].Container;						
+						const FIoChunkId& ChunkId = Container->ChunkIds[ChunkHttpRequestHandle->EntryIndex];
 
-						PendingDownloadSize += ChunkEntry.GetDiskSize();
+						const EOnDemandInstallCasType PendingCasType = GetCasType(ChunkId.GetChunkType());
+						if (PendingCasType == CasType)
+						{
+							const FOnDemandChunkEntry& ChunkEntry = Container->ChunkEntries[ChunkHttpRequestHandle->EntryIndex];
+							PendingDownloadSize += ChunkEntry.GetDiskSize();
+						}
 					}
 
-					// TODO: IadMultiCAS: content installer needs to determine which chunks go to which CAS
-					Result = IoStore.InstallCache->Purge(EOnDemandInstallCasType::General, PendingDownloadSize);
+					Result = IoStore.InstallCache->Purge(CasType, PendingDownloadSize);
 					if (Result.HasValue())
 					{
-						PutResult = IoStore.InstallCache->PutChunk(EOnDemandInstallCasType::General, MoveTemp(Chunk), ChunkHash);
+						PutResult = IoStore.InstallCache->PutChunk(CasType, MoveTemp(Chunk), ChunkHash);
 						if (PutResult.HasError())
 						{
 							Result = MakeError(PutResult.StealError());
@@ -1501,9 +1540,17 @@ void FOnDemandContentInstaller::OnProcessDownloadedChunkNotifyRequest(
 			TEXT("OK"));
 	}
 
+	InstallRequest.bUsesMMapCas = InstallRequest.bUsesMMapCas || GetCasType(ChunkId.GetChunkType()) == EOnDemandInstallCasType::MMap;
 	const bool bCompleted = ++InstallRequest.DownloadedChunkCount >= InstallRequest.HttpRequestHandles.Num();
 	if (bCompleted)
 	{
+		// Cannot read from pending buffer for MMap chunks, must flush to disk
+		if (Request.IsOk() && InstallRequest.bUsesMMapCas)
+		{
+			// TODO: this could be a partial flush if pending chunks was sorted by request
+			Request.Result = IoStore.InstallCache->Flush(EOnDemandInstallCasType::MMap);
+		}
+
 		CompleteInstallRequest(Request);
 	}
 	else
diff --git a/Engine/Source/Runtime/Experimental/IoStore/OnDemand/Private/OnDemandContentInstaller.h b/Engine/Source/Runtime/Experimental/IoStore/OnDemand/Private/OnDemandContentInstaller.h
--- a/Engine/Source/Runtime/Experimental/IoStore/OnDemand/Private/OnDemandContentInstaller.h
+++ b/Engine/Source/Runtime/Experimental/IoStore/OnDemand/Private/OnDemandContentInstaller.h
@@ -104,6 +104,7 @@ class FOnDemandContentInstaller
 			uint64										DownloadedChunkCount = 0;
 			uint64										LastProgressCycles = 0;
 			bool										bHttpRequestsIssued = false;
+			bool										bUsesMMapCas = false;
 			std::atomic_bool							bNotifyingProgressOnGameThread{false};
 		};
 
@@ -254,6 +255,7 @@ public:
 	void							ReportAnalytics(TArray<FAnalyticsEventAttribute>& OutAnalyticsArray) const;
 
 private:
+	EOnDemandInstallCasType			GetCasType(EIoChunkType ChunkType) const;
 	bool							CanExecuteRequest(FRequest& Request) const;
 	void							TryExecuteNextRequest();
 	void							ExecuteRequest(FRequest& Request);
@@ -292,6 +294,7 @@ private:
 	// requests for the same iochunk but only one of them should be active with the remainder
 	// waiting on final callbacks after being cancelled. See FChunkDownloadRequest.
 	TMultiMap<FIoHash, FChunkDownloadRequest> PendingChunkDownloads;
+	const bool						bMMapCasEnabled{false};
 	std::atomic_bool				bShuttingDown{false};
 };
 
diff --git a/Engine/Source/Runtime/Experimental/IoStore/OnDemand/Private/OnDemandInstallCache.cpp b/Engine/Source/Runtime/Experimental/IoStore/OnDemand/Private/OnDemandInstallCache.cpp
--- a/Engine/Source/Runtime/Experimental/IoStore/OnDemand/Private/OnDemandInstallCache.cpp
+++ b/Engine/Source/Runtime/Experimental/IoStore/OnDemand/Private/OnDemandInstallCache.cpp
@@ -2353,6 +2353,7 @@ public:
 													TConstArrayView<FIoHash> ChunkHashes, 
 													TStaticArray<bool, EOnDemandInstallCasType::Count>& bInOutLastAccessDirty) override;
 	virtual FOnDemandInstallCacheUsage			GetCacheUsage(EOnDemandInstallCasType CasType) override;
+	virtual bool								IsCasEnabled(EOnDemandInstallCasType CasType) const override;
 
 private:
 	void										RegisterConsoleCommands();
@@ -4283,6 +4284,12 @@ FOnDemandInstallCacheUsage FOnDemandInstallCache::GetCacheUsage(EOnDemandInstall
 	};
 }
 
+bool FOnDemandInstallCache::IsCasEnabled(EOnDemandInstallCasType CasType) const
+{
+	const FCacheState* CacheState = GetCacheState(CasType);
+	return CacheState != nullptr;
+}
+
 FResult FOnDemandInstallCache::FlushPendingChunks(FCacheState& CacheState, FPendingChunks& Chunks, int64 UtcAccessTicks, int64 UtcModificationTicks)
 {
 	if (Chunks.IsEmpty())
diff --git a/Engine/Source/Runtime/Experimental/IoStore/OnDemand/Private/OnDemandInstallCache.h b/Engine/Source/Runtime/Experimental/IoStore/OnDemand/Private/OnDemandInstallCache.h
--- a/Engine/Source/Runtime/Experimental/IoStore/OnDemand/Private/OnDemandInstallCache.h
+++ b/Engine/Source/Runtime/Experimental/IoStore/OnDemand/Private/OnDemandInstallCache.h
@@ -39,6 +39,7 @@ public:
 	virtual FResult								FlushLastAccess(const TStaticArray<bool, EOnDemandInstallCasType::Count>& bFlushCasLastAccess) = 0;
 	virtual void								UpdateLastAccess(TConstArrayView<FIoHash> ChunkHashes, TStaticArray<bool, EOnDemandInstallCasType::Count>& bInOutLastAccessDirty) = 0;
 	virtual FOnDemandInstallCacheUsage			GetCacheUsage(EOnDemandInstallCasType CasType) = 0;
+	virtual bool								IsCasEnabled(EOnDemandInstallCasType CasType) const = 0;
 };
 
 struct FOnDemandInstallCasConfig
