Deprecate dead code in Geometry collection actor

#Synced-CL 51402414

diff --git a/Engine/Source/Runtime/Experimental/GeometryCollectionEngine/Private/GeometryCollection/GeometryCollectionActor.cpp b/Engine/Source/Runtime/Experimental/GeometryCollectionEngine/Private/GeometryCollection/GeometryCollectionActor.cpp
--- a/Engine/Source/Runtime/Experimental/GeometryCollectionEngine/Private/GeometryCollection/GeometryCollectionActor.cpp
+++ b/Engine/Source/Runtime/Experimental/GeometryCollectionEngine/Private/GeometryCollection/GeometryCollectionActor.cpp
@@ -51,72 +51,11 @@ const Chaos::FPhysicsSolver* GetSolver(const AGeometryCollectionActor& GeomColle
 	return GeomCollectionActor.GetGeometryCollectionComponent()->ChaosSolverActor != nullptr ? GeomCollectionActor.GetGeometryCollectionComponent()->ChaosSolverActor->GetSolver() : GeomCollectionActor.GetWorld()->PhysicsScene_Chaos->GetSolver();
 }
 
-
-bool LowLevelRaycastImp(const Chaos::FVec3& Start, const Chaos::FVec3& Dir, float DeltaMag, const AGeometryCollectionActor& GeomCollectionActor, FHitResult& OutHit)
-{
-#if TODO_REIMPLEMENT_GET_RIGID_PARTICLES
-	using namespace Chaos;
-	FPhysScene_Chaos* Scene = GeomCollectionActor.GetGeometryCollectionComponent()->GetInnerChaosScene();
-	ensure(Scene);
-
-	const Chaos::FPhysicsSolver* Solver = GetSolver(GeomCollectionActor);
-	if(ensure(Solver))
-	{
-		const TPBDRigidParticles<float, 3>& Particles = Solver->GetRigidParticles();	//todo(ocohen): should these just get passed in instead of hopping through scene?
-		FAABB3 RayBox(Start, Start);
-		RayBox.Thicken(Dir * DeltaMag);
-		const auto& PotentialIntersections = Solver->GetSpatialAcceleration()->FindAllIntersections(RayBox);
-		Solver->ReleaseSpatialAcceleration();
-
-		for(const auto RigidBodyIdx : PotentialIntersections)
-		{
-			const TRigidTransform<float, 3> TM(Particles.X(RigidBodyIdx), Particles.R(RigidBodyIdx));
-			const FVec3 StartLocal = TM.InverseTransformPositionNoScale(Start);
-			const FVec3 DirLocal = TM.InverseTransformVectorNoScale(Dir);
-			const FVec3 EndLocal = StartLocal + DirLocal * DeltaMag;	//todo(ocohen): apeiron just undoes this later, we should fix the API
-
-			const FImplicitObject* Object = Particles.Geometry(RigidBodyIdx).Get();	//todo(ocohen): can this ever be null?
-			Pair<FVec3, bool> Result = Object->FindClosestIntersection(StartLocal, EndLocal, /*Thickness=*/0.f);
-			if(Result.Second)	//todo(ocohen): once we do more than just a bool we need to get the closest point
-			{
-				const float Distance = (Result.First - StartLocal).Size();
-				OutHit.Actor = const_cast<AGeometryCollectionActor*>(&GeomCollectionActor);
-				OutHit.Component = GeomCollectionActor.GetGeometryCollectionComponent();
-				OutHit.bBlockingHit = true;
-				OutHit.Distance = Distance;
-				OutHit.Time = Distance / (EndLocal - StartLocal).Size();
-				OutHit.Location = TM.TransformPositionNoScale(Result.First);
-				OutHit.ImpactPoint = OutHit.Location;
-				const FVec3 LocalNormal = Object->Normal(Result.First);
-				OutHit.ImpactNormal = TM.TransformVectorNoScale(LocalNormal);
-				OutHit.Normal = OutHit.ImpactNormal;
-				OutHit.Item = RigidBodyIdx;
-
-				return true;
-			}
-		}
-	}
-#endif
-
-	return false;
-}
-
 bool AGeometryCollectionActor::RaycastSingle(FVector Start, FVector End, FHitResult& OutHit) const
 {
-	if (GeometryCollectionComponent)
-	{
-		OutHit = FHitResult();
-		OutHit.TraceStart = Start;
-		OutHit.TraceEnd = End;
-		const FVector Delta = (End - Start);
-		const float DeltaMag = Delta.Size();
-		if (DeltaMag > KINDA_SMALL_NUMBER)
-		{
-			const FVector Dir = Delta / DeltaMag;
-			return LowLevelRaycastImp(Start, Dir, DeltaMag, *this, OutHit);
-		}
-	}
+	// this is a deprecated function (deprecated starting 5.8 ) and has been a no-op return false since 2021
 	return false;
+
 }
 
 #if WITH_EDITOR
diff --git a/Engine/Source/Runtime/Experimental/GeometryCollectionEngine/Public/GeometryCollection/GeometryCollectionActor.h b/Engine/Source/Runtime/Experimental/GeometryCollectionEngine/Public/GeometryCollection/GeometryCollectionActor.h
--- a/Engine/Source/Runtime/Experimental/GeometryCollectionEngine/Public/GeometryCollection/GeometryCollectionActor.h
+++ b/Engine/Source/Runtime/Experimental/GeometryCollectionEngine/Public/GeometryCollection/GeometryCollectionActor.h
@@ -35,6 +35,7 @@ public:
 	TObjectPtr<UGeometryCollectionDebugDrawComponent> GeometryCollectionDebugDrawComponent_DEPRECATED;
 	PRAGMA_ENABLE_DEPRECATION_WARNINGS
 	
+	UE_DEPRECATED(5.8, "RaycastSingle is now deprecated, use regular Unreal trace system to query geometry collections")
 	UFUNCTION(BlueprintCallable, Category = "Physics")
 	GEOMETRYCOLLECTIONENGINE_API bool RaycastSingle(FVector Start, FVector End, FHitResult& OutHit) const;
 
