Fix FConvex::FindClosestIntersectionImp to correctly handle thickness.

#Synced-CL 51472478

diff --git a/Engine/Source/Runtime/Experimental/Chaos/Public/Chaos/Convex.h b/Engine/Source/Runtime/Experimental/Chaos/Public/Chaos/Convex.h
--- a/Engine/Source/Runtime/Experimental/Chaos/Public/Chaos/Convex.h
+++ b/Engine/Source/Runtime/Experimental/Chaos/Public/Chaos/Convex.h
@@ -385,15 +385,18 @@ namespace Chaos
 			TArray<Pair<FReal, FVec3>> Intersections;
 			Intersections.Reserve(FMath::Min(static_cast<int32>(NumPlanes*.1), 16)); // Was NumPlanes, which seems excessive.
 			for (int32 Idx = 0; Idx < NumPlanes; ++Idx)
-			{
-				auto PlaneIntersection = Planes[Idx].FindClosestIntersection(StartPoint, EndPoint, Thickness);
+			{ 
+				// Offset plane with thickness in the normal direction and do a 0 thickness intersection query, so you only find intersections on the +Thickness*Normal plane, not the -Thickness*Normal plane.
+				const Pair<FVec3, bool> PlaneIntersection = Thickness > 0 ? 
+					FPlaneType(Planes[Idx].X() + Thickness * Planes[Idx].Normal(), Planes[Idx].Normal()).FindClosestIntersection(StartPoint, EndPoint, 0) :
+					Planes[Idx].FindClosestIntersection(StartPoint, EndPoint, 0);
 				if (PlaneIntersection.Second)
 				{
 					Intersections.Add(MakePair((FReal)(PlaneIntersection.First - StartPoint).SizeSquared(), PlaneIntersection.First));
 				}
 			}
 			Intersections.Sort([](const Pair<FReal, FVec3>& Elem1, const Pair<FReal, FVec3>& Elem2) { return Elem1.First < Elem2.First; });
-			for (const auto& Elem : Intersections)
+			for (const Pair<FReal, FVec3>& Elem : Intersections)
 			{
 				if (this->SignedDistance(Elem.Second) < (Thickness + 1e-4))
 				{
