Geometry collection mesh facade : add a way to get a view of the triangles for a specific bone

#Synced-CL 51402414

diff --git a/Engine/Source/Runtime/Experimental/Chaos/Private/GeometryCollection/Facades/CollectionMeshFacade.cpp b/Engine/Source/Runtime/Experimental/Chaos/Private/GeometryCollection/Facades/CollectionMeshFacade.cpp
--- a/Engine/Source/Runtime/Experimental/Chaos/Private/GeometryCollection/Facades/CollectionMeshFacade.cpp
+++ b/Engine/Source/Runtime/Experimental/Chaos/Private/GeometryCollection/Facades/CollectionMeshFacade.cpp
@@ -141,6 +141,24 @@ namespace GeometryCollection::Facades
 		return {};
 	}
 
+	const TArrayView<const FIntVector> FCollectionMeshFacade::GetTriangles(int32 BoneIdx) const
+	{
+		const TManagedArray<int32>& TransformToGeometryIndicies = TransformToGeometryIndexAttribute.Get();
+		const TManagedArray<int32>& FaceStarts = FaceStartAttribute.Get();
+		const TManagedArray<int32>& FaceCounts = FaceCountAttribute.Get();
+		const TManagedArray<FIntVector>& Triangles = IndicesAttribute.Get();
+
+		if (TransformToGeometryIndicies.IsValidIndex(BoneIdx))
+		{
+			const int32 GeometryIndex = TransformToGeometryIndicies[BoneIdx];
+			if (FaceStarts.IsValidIndex(GeometryIndex) && FaceCounts.IsValidIndex(GeometryIndex))
+			{
+				return TArrayView<const FIntVector>(Triangles.GetData() + FaceStarts[GeometryIndex], FaceCounts[GeometryIndex]);
+			}
+		}
+		return {};
+	}
+
 	const TArray<int32> FCollectionMeshFacade::GetFaceIndices(int32 BoneIdx) const
 	{
 		const TManagedArray<int32>& TransformToGeometryIndicies = TransformToGeometryIndexAttribute.Get();
diff --git a/Engine/Source/Runtime/Experimental/Chaos/Public/GeometryCollection/Facades/CollectionMeshFacade.h b/Engine/Source/Runtime/Experimental/Chaos/Public/GeometryCollection/Facades/CollectionMeshFacade.h
--- a/Engine/Source/Runtime/Experimental/Chaos/Public/GeometryCollection/Facades/CollectionMeshFacade.h
+++ b/Engine/Source/Runtime/Experimental/Chaos/Public/GeometryCollection/Facades/CollectionMeshFacade.h
@@ -37,6 +37,12 @@ namespace GeometryCollection::Facades
 		 */
 		CHAOS_API const TArrayView<const FVector3f> GetVertexPositions(int32 BoneIdx) const;
 
+		/**
+		* Returns the triangles for the bone 
+		 */
+		CHAOS_API const TArrayView<const FIntVector> GetTriangles(int32 BoneIdx) const;
+
+
 		/**
 		 * Returns the face indices for the bone
 		 */
