Updated behavior when converting a weak on-demand content handle to return an invalid weak handle rather than crash.
- added some additional checks to validate that referenced content has a valid content handle

#rnx

#Synced-CL 51451778

diff --git a/Engine/Source/Runtime/Experimental/IoStore/OnDemand/Private/OnDemandIoStore.cpp b/Engine/Source/Runtime/Experimental/IoStore/OnDemand/Private/OnDemandIoStore.cpp
--- a/Engine/Source/Runtime/Experimental/IoStore/OnDemand/Private/OnDemandIoStore.cpp
+++ b/Engine/Source/Runtime/Experimental/IoStore/OnDemand/Private/OnDemandIoStore.cpp
@@ -1830,6 +1830,12 @@ void FOnDemandIoStore::GetReferencedContentByHandle(TMap<FOnDemandWeakContentHan
 		for (const FOnDemandChunkEntryReferences& Refs : Container->ChunkEntryReferences)
 		{
 			FOnDemandWeakContentHandle WeakHandle = FOnDemandWeakContentHandle::FromUnsafeHandle(Refs.ContentHandleId);
+			if (WeakHandle.IsValid() == false)
+			{
+				UE_LOGF(LogIoStoreOnDemand, Warning, "Found invalid content handle in container '%ls'", *Container->Name());
+				continue;
+			}
+
 			TArray<FOnDemandContainerChunkEntryReferences>& Entries = OutReferencesByHandle.FindOrAdd(WeakHandle);
 			Entries.Add(FOnDemandContainerChunkEntryReferences
 			{
@@ -1840,13 +1846,6 @@ void FOnDemandIoStore::GetReferencedContentByHandle(TMap<FOnDemandWeakContentHan
 	}
 }
 
-void FOnDemandIoStore::AddReference(const FSharedOnDemandContainer& Container, int32 EntryIndex, FOnDemandContentHandle ContentHandle)
-{
-	UE::TUniqueLock Lock(ContainerMutex);
-	FOnDemandChunkEntryReferences& Refs = Container->FindOrAddChunkEntryReferences(*ContentHandle.Handle);
-	Refs.Indices[EntryIndex] = true;
-}
-
 void FOnDemandIoStore::CancelInstallRequest(FSharedInternalInstallRequest InstallRequest)
 {
 	check(InstallRequest.IsValid());
diff --git a/Engine/Source/Runtime/Experimental/IoStore/OnDemand/Private/OnDemandIoStore.h b/Engine/Source/Runtime/Experimental/IoStore/OnDemand/Private/OnDemandIoStore.h
--- a/Engine/Source/Runtime/Experimental/IoStore/OnDemand/Private/OnDemandIoStore.h
+++ b/Engine/Source/Runtime/Experimental/IoStore/OnDemand/Private/OnDemandIoStore.h
@@ -159,6 +159,8 @@ const FOnDemandChunkEntry* FOnDemandContainer::FindChunkEntry(const FIoChunkId&
 FOnDemandChunkEntryReferences& FOnDemandContainer::FindOrAddChunkEntryReferences(const FOnDemandInternalContentHandle& ContentHandle)
 {
 	const UPTRINT ContentHandleId = ContentHandle.HandleId(); 
+	check(ContentHandleId != 0);
+
 	for (FOnDemandChunkEntryReferences& Refs : ChunkEntryReferences)
 	{
 		if (Refs.ContentHandleId == ContentHandleId)
@@ -178,6 +180,7 @@ TBitArray<> FOnDemandContainer::GetReferencedChunkEntries() const
 	TBitArray<> Indices;
 	for (const FOnDemandChunkEntryReferences& Refs : ChunkEntryReferences)
 	{
+		check(Refs.ContentHandleId != 0);
 		check(Refs.Indices.Num() == ChunkEntries.Num());
 		Indices.CombineWithBitwiseOR(Refs.Indices, EBitwiseOperatorFlags::MaxSize);
 	}
@@ -347,7 +350,6 @@ public:
 	void								GetReferencedContent(TArray<FSharedOnDemandContainer>& OutContainers, TArray<TBitArray<>>& OutChunkEntryIndices, bool bPackageStore = false);
 	TBitArray<>							GetReferencedContent(const FSharedOnDemandContainer& Container);
 	void								GetReferencedContentByHandle(TMap<FOnDemandWeakContentHandle, TArray<FOnDemandContainerChunkEntryReferences>>& OutReferencesByHandle) const;
-	void								AddReference(const FSharedOnDemandContainer& Container, int32 EntryIndex, FOnDemandContentHandle ContentHandle);
 
 	// IOnDemandIoStore interface
 	virtual FIoStatus					Initialize() override;
diff --git a/Engine/Source/Runtime/Experimental/IoStore/OnDemandCore/Internal/IO/IoStoreOnDemandInternals.h b/Engine/Source/Runtime/Experimental/IoStore/OnDemandCore/Internal/IO/IoStoreOnDemandInternals.h
--- a/Engine/Source/Runtime/Experimental/IoStore/OnDemandCore/Internal/IO/IoStoreOnDemandInternals.h
+++ b/Engine/Source/Runtime/Experimental/IoStore/OnDemandCore/Internal/IO/IoStoreOnDemandInternals.h
@@ -18,6 +18,8 @@ namespace UE::IoStore
 ///////////////////////////////////////////////////////////////////////////////
 struct FOnDemandWeakContentHandle
 {
+	bool IsValid() const { return HandleId != 0; }
+
 	inline bool operator==(const FOnDemandWeakContentHandle& Other) const
 	{
 		return HandleId == Other.HandleId;
@@ -28,6 +30,8 @@ struct FOnDemandWeakContentHandle
 		return HandleId != Other.HandleId;
 	}
 
+	inline operator bool() const { return IsValid(); }
+
 	inline friend uint32 GetTypeHash(const FOnDemandWeakContentHandle& Handle)
 	{
 		return GetTypeHash(Handle.HandleId);
diff --git a/Engine/Source/Runtime/Experimental/IoStore/OnDemandCore/Private/IoStoreOnDemandInternals.cpp b/Engine/Source/Runtime/Experimental/IoStore/OnDemandCore/Private/IoStoreOnDemandInternals.cpp
--- a/Engine/Source/Runtime/Experimental/IoStore/OnDemandCore/Private/IoStoreOnDemandInternals.cpp
+++ b/Engine/Source/Runtime/Experimental/IoStore/OnDemandCore/Private/IoStoreOnDemandInternals.cpp
@@ -20,6 +20,11 @@ void FOnDemandMountResult::LogResult()
 ///////////////////////////////////////////////////////////////////////////////
 FOnDemandWeakContentHandle FOnDemandWeakContentHandle::FromUnsafeHandle(UPTRINT HandleId)
 {
+	if (HandleId == 0)
+	{
+		return FOnDemandWeakContentHandle{};
+	}
+
 	FOnDemandInternalContentHandle& ContentHandle = *reinterpret_cast<FOnDemandInternalContentHandle*>(HandleId);
 	return FOnDemandWeakContentHandle
 	{
