Added Dataflow MakeArrayNode.

This node uses AnyType to construct an array output from an array of inputs. It also supports appending arrays since nested arrays are not supported (works similar to ToArray). Additional methods were added to the DataflowCache to support this.

Additionally: fixed dependency types and conversion for FDataflowConvertToArrayNode when the input is an array.

#Synced-CL 51472478

diff --git a/Engine/Plugins/Experimental/GeometryCollectionPlugin/Source/GeometryCollectionNodes/Private/Dataflow/GeometryCollectionArrayNodes.cpp b/Engine/Plugins/Experimental/GeometryCollectionPlugin/Source/GeometryCollectionNodes/Private/Dataflow/GeometryCollectionArrayNodes.cpp
--- a/Engine/Plugins/Experimental/GeometryCollectionPlugin/Source/GeometryCollectionNodes/Private/Dataflow/GeometryCollectionArrayNodes.cpp
+++ b/Engine/Plugins/Experimental/GeometryCollectionPlugin/Source/GeometryCollectionNodes/Private/Dataflow/GeometryCollectionArrayNodes.cpp
@@ -21,6 +21,11 @@ namespace UE::Dataflow
 
 		static bool AutoConvertSingleToArrayFilter(FName FromType, FName ToType)
 		{
+			if (FDataflowArrayTypePolicy::SupportsTypeStatic(FromType))
+			{
+				// Cannot convert arrays to nested arrays.
+				return false;
+			}
 			const FName ElementType = FDataflowArrayTypePolicy::GetElementType(ToType);
 			return (FromType == ElementType);
 		}
@@ -40,6 +45,7 @@ namespace UE::Dataflow
 		DATAFLOW_NODE_REGISTER_CREATION_FACTORY(FDataflowGetArraySizeNode);
 		DATAFLOW_NODE_REGISTER_CREATION_FACTORY(FDataflowGetArrayElementNode);
 		DATAFLOW_NODE_REGISTER_CREATION_FACTORY(FDataflowConvertToArrayNode);
+		DATAFLOW_NODE_REGISTER_CREATION_FACTORY(FDataflowMakeArrayNode);
 		DATAFLOW_NODE_REGISTER_CREATION_FACTORY(FDataflowMakeManagedArrayCollectionArrayNode);
 		
 
@@ -722,6 +728,12 @@ FName FDataflowConvertToArrayNode::GetDependentConnectionType(const FDataflowCon
 	}
 	else if (SourceConnection.IsA(&Element) && DependentConnection.IsA(&Array))
 	{
+		if (FDataflowArrayTypePolicy::SupportsTypeStatic(SourceConnection.GetType()))
+		{
+			// Special case: Array elements are left unchanged.
+			return SourceConnection.GetType();
+		}
+
 		return FDataflowArrayTypePolicy::GetArrayType(SourceConnection.GetType());
 	}
 	// default for all the rest 
@@ -739,7 +751,178 @@ void FDataflowConvertToArrayNode::Evaluate(UE::Dataflow::FContext& Context, cons
 }
 
 ///////////////////////////////////////////////////////////////////////////////////////////////
+FDataflowMakeArrayNode::FDataflowMakeArrayNode(const UE::Dataflow::FNodeParameters& InParam, FGuid InGuid)
+	: FDataflowNode(InParam, InGuid)
+{	
+
+	RegisterOutputConnection(&Array)
+		.SetTypeDependencyGroup(ArrayElementDependencyGroup);
+
+	// Add initial pins
+	for (int32 Index = 0; Index < NumInitialElements; ++Index)
+	{
+		AddPins();
+	}
+	check(GetNumInputs() == NumRequiredInputs + NumInitialElements);  // Update NumRequiredInputs if you add more Inputs. This is used by Serialize.
+}
+
+FName FDataflowMakeArrayNode::GetDependentConnectionType(const FDataflowConnection& SourceConnection, const FDataflowConnection& DependentConnection) const
+{
+	auto IsElementConnection = [this](const FDataflowConnection& Connection)
+		{
+			for (const FDataflowAllTypes& Element : Elements)
+			{
+				if (Connection.IsA(&Element))
+				{
+					return true;
+				}
+			}
+			return false;
+		};
+
+	if (SourceConnection.IsA(&Array) && IsElementConnection(DependentConnection))
+	{
+		return FDataflowArrayTypePolicy::GetElementType(SourceConnection.GetType());
+	}
+	else if (DependentConnection.IsA(&Array) && IsElementConnection(SourceConnection))
+	{
+		if (FDataflowArrayTypePolicy::SupportsTypeStatic(SourceConnection.GetType()))
+		{
+			// Special case: Array elements are left unchanged.
+			return SourceConnection.GetType();
+		}
+		return FDataflowArrayTypePolicy::GetArrayType(SourceConnection.GetType());
+	}
+
+	// default for all the rest 
+	return FDataflowNode::GetDependentConnectionType(SourceConnection, DependentConnection);
+}
+
+void FDataflowMakeArrayNode::Evaluate(UE::Dataflow::FContext& Context, const FDataflowOutput* Out) const
+{
+	if (Out->IsA(&Array))
+	{
+		TArray<const UE::Dataflow::FContextValue> InElements;
+
+		InElements.Reserve(Elements.Num());
+
+		bool bIsAppendingArrays = false;
+		for (int32 Index = 0; Index < Elements.Num(); ++Index)
+		{
+			if (IsConnected(GetConnectionReference(Index)))
+			{
+				const UE::Dataflow::FContextValue& InElement =  InElements.Emplace_GetRef(GetValue(Context, GetConnectionReference(Index)));
+				if (InElement.IsArray())
+				{
+					bIsAppendingArrays = true;
+				}
+			}
+		}
 
+		if (bIsAppendingArrays)
+		{
+			UE::Dataflow::FContextValue OutArray = UE::Dataflow::FContextValue::AppendArrays(Context, InElements);
+			SetValue(Context, MoveTemp(OutArray), &Array);
+		}
+		else
+		{
+			UE::Dataflow::FContextValue OutArray = UE::Dataflow::FContextValue::MakeArray(Context, InElements);
+			SetValue(Context, MoveTemp(OutArray), &Array);
+		}
+	}
+}
+
+TArray<UE::Dataflow::FPin> FDataflowMakeArrayNode::AddPins()
+{
+	const int32 Index = Elements.AddDefaulted();
+	FDataflowInput& Input = RegisterInputArrayConnection(GetConnectionReference(Index));
+	if (Index > 0)
+	{
+		if (const FDataflowInput* Input0 = FindInput(GetConnectionReference(0)))
+		{
+			if (Input0->HasConcreteType())
+			{
+				SetConnectionConcreteType(&Input, Input0->GetType(), ArrayElementDependencyGroup);
+			}
+		}
+	}
+	Input.SetTypeDependencyGroup(ArrayElementDependencyGroup);
+	return { { UE::Dataflow::FPin::EDirection::INPUT, Input.GetType(), Input.GetName() } };
+}
+
+TArray<UE::Dataflow::FPin> FDataflowMakeArrayNode::GetPinsToRemove() const
+{
+	const int32 Index = Elements.Num() - 1;
+	check(Elements.IsValidIndex(Index));
+	if (const FDataflowInput* const Input = FindInput(GetConnectionReference(Index)))
+	{
+		return { { UE::Dataflow::FPin::EDirection::INPUT, Input->GetType(), Input->GetName() } };
+	}
+	return Super::GetPinsToRemove();
+}
+
+void FDataflowMakeArrayNode::OnPinRemoved(const UE::Dataflow::FPin& Pin)
+{
+	const int32 Index = Elements.Num() - 1;
+	check(Elements.IsValidIndex(Index));
+#if DO_CHECK
+	const FDataflowInput* const Input = FindInput(GetConnectionReference(Index));
+	check(Input);
+	check(Input->GetName() == Pin.Name);
+	check(Input->GetType() == Pin.Type);
+#endif
+	Elements.SetNum(Index);
+
+	return Super::OnPinRemoved(Pin);
+}
+
+UE::Dataflow::TConnectionReference<FDataflowAllTypes> FDataflowMakeArrayNode::GetConnectionReference(int32 Index) const
+{
+	return { &Elements[Index], Index, &Elements };
+}
+
+void FDataflowMakeArrayNode::PostSerialize(const FArchive& Ar)
+{
+	// because we add pins we need to make sure we restore them when loading
+	// to make sure they can get properly reconnected
+	if (Ar.IsLoading())
+	{
+		for (int32 Index = 0; Index < NumInitialElements; ++Index)
+		{
+			check(FindInput(GetConnectionReference(Index)));
+		}
+
+		for (int32 Index = NumInitialElements; Index < Elements.Num(); ++Index)
+		{
+			FindOrRegisterInputArrayConnection(GetConnectionReference(Index))
+				.SetTypeDependencyGroup(ArrayElementDependencyGroup);
+		}
+		if (Ar.IsTransacting())
+		{
+			const int32 OrigNumRegisteredInputs = GetNumInputs();
+			check(OrigNumRegisteredInputs >= NumRequiredInputs + NumInitialElements);
+			const int32 OrigNumElements = Elements.Num();
+			const int32 OrigNumRegisteredElements = (OrigNumRegisteredInputs - NumRequiredInputs);
+			if (OrigNumRegisteredElements > OrigNumElements)
+			{
+				// Inputs have been removed.
+				// Temporarily expand Collections so we can get connection references.
+				Elements.SetNum(OrigNumRegisteredElements);
+				for (int32 Index = OrigNumElements; Index < Elements.Num(); ++Index)
+				{
+					UnregisterInputConnection(GetConnectionReference(Index));
+				}
+				Elements.SetNum(OrigNumElements);
+			}
+		}
+		else
+		{
+			ensureAlways(Elements.Num() + NumRequiredInputs == GetNumInputs());
+		}
+	}
+}
+
+///////////////////////////////////////////////////////////////////////////////////////////////
 FDataflowMakeManagedArrayCollectionArrayNode::FDataflowMakeManagedArrayCollectionArrayNode(const UE::Dataflow::FNodeParameters& InParam, FGuid InGuid)
 	: FDataflowNode(InParam, InGuid)
 {
diff --git a/Engine/Plugins/Experimental/GeometryCollectionPlugin/Source/GeometryCollectionNodes/Public/Dataflow/GeometryCollectionArrayNodes.h b/Engine/Plugins/Experimental/GeometryCollectionPlugin/Source/GeometryCollectionNodes/Public/Dataflow/GeometryCollectionArrayNodes.h
--- a/Engine/Plugins/Experimental/GeometryCollectionPlugin/Source/GeometryCollectionNodes/Public/Dataflow/GeometryCollectionArrayNodes.h
+++ b/Engine/Plugins/Experimental/GeometryCollectionPlugin/Source/GeometryCollectionNodes/Public/Dataflow/GeometryCollectionArrayNodes.h
@@ -613,6 +613,46 @@ private:
 	virtual FName GetDependentConnectionType(const FDataflowConnection& SourceConnection, const FDataflowConnection& DependentConnection) const override;
 };
 
+/**
+ * Make an array from multiple elements of the same type. When elements are array type, the arrays are appended (nested arrays are not supported).
+ */
+USTRUCT()
+struct FDataflowMakeArrayNode : public FDataflowNode
+{
+	GENERATED_USTRUCT_BODY()
+	DATAFLOW_NODE_DEFINE_INTERNAL(FDataflowMakeArrayNode, "MakeArray", "Utilities|Array", "Element")
+
+public:
+	FDataflowMakeArrayNode(const UE::Dataflow::FNodeParameters& InParam, FGuid InGuid = FGuid::NewGuid());
+
+private:
+	/** Array to Convert to */
+	UPROPERTY(meta = (DataflowOutput))
+	FDataflowArrayTypes Array;
+
+	/** Elements to convert to an array */
+	UPROPERTY()
+	TArray<FDataflowAllTypes> Elements;
+
+	virtual void Evaluate(UE::Dataflow::FContext& Context, const FDataflowOutput* Out) const override;
+
+	virtual FName GetDependentConnectionType(const FDataflowConnection& SourceConnection, const FDataflowConnection& DependentConnection) const override;
+
+	virtual TArray<UE::Dataflow::FPin> AddPins() override;
+	virtual bool CanAddPin() const override { return true; }
+	virtual bool CanRemovePin() const override { return Elements.Num() > NumInitialElements; }
+	virtual TArray<UE::Dataflow::FPin> GetPinsToRemove() const override;
+	virtual void OnPinRemoved(const UE::Dataflow::FPin& Pin) override;
+	virtual void PostSerialize(const FArchive& Ar) override;
+
+	UE::Dataflow::TConnectionReference<FDataflowAllTypes> GetConnectionReference(int32 Index) const;
+
+	static constexpr int32 NumRequiredInputs = 0;
+	static constexpr int32 NumInitialElements = 2; 
+
+	inline static FName ArrayElementDependencyGroup = FName("ArrayElement");
+};
+
 /**
  * Append an element to an array of ManagedArrayCollections.
  */
diff --git a/Engine/Source/Runtime/Experimental/Dataflow/Core/Private/Dataflow/DataflowContextCache.cpp b/Engine/Source/Runtime/Experimental/Dataflow/Core/Private/Dataflow/DataflowContextCache.cpp
--- a/Engine/Source/Runtime/Experimental/Dataflow/Core/Private/Dataflow/DataflowContextCache.cpp
+++ b/Engine/Source/Runtime/Experimental/Dataflow/Core/Private/Dataflow/DataflowContextCache.cpp
@@ -52,6 +52,24 @@ namespace UE::Dataflow
 		return {};
 	}
 
+	TUniquePtr<FContextCacheElementBase> FContextCacheElementReference::CreateArrayFromElementAndAppend(const TArray<const FContextCacheElementBase*>& DereferencedElements, const IContextCacheStore& Context, const FProperty* InProperty, const FGuid& InNodeGuid, uint32 InNodeHash, const FTimestamp& InTimestamp) const
+	{
+		if (const TUniquePtr<FContextCacheElementBase>* Cache = Context.FindCacheElement(DataKey))
+		{
+			return (*Cache)->CreateArrayFromElementAndAppend(DereferencedElements, Context, InProperty, InNodeGuid, InNodeHash, InTimestamp);
+		}
+		return {};
+	}
+
+	TUniquePtr<FContextCacheElementBase> FContextCacheElementReference::AppendArrayElements(const TArray<const FContextCacheElementBase*>& DereferencedElements, const IContextCacheStore& Context, const FProperty* InProperty, const FGuid& InNodeGuid, uint32 InNodeHash, const FTimestamp& InTimestamp) const
+	{
+		if (const TUniquePtr<FContextCacheElementBase>* Cache = Context.FindCacheElement(DataKey))
+		{
+			return (*Cache)->AppendArrayElements(DereferencedElements, Context, InProperty, InNodeGuid, InNodeHash, InTimestamp);
+		}
+		return {};
+	}
+
 	FConstStructView FContextCacheElementReference::GetConstStructView(const IContextCacheStore& Context) const
 	{
 		if (const TUniquePtr<FContextCacheElementBase>* Cache = Context.FindCacheElement(DataKey))
@@ -73,6 +91,23 @@ namespace UE::Dataflow
 		return MakeUnique<FContextCacheElementNull>(GetNodeGuid(), GetProperty(), GetNodeHash(), GetTimestamp());
 	}
 
+	const FContextCacheElementBase* FContextCacheElementReference::GetReferencedElement(const IContextCacheStore& Context) const
+	{
+		if (const TUniquePtr<FContextCacheElementBase>* ReferencedCacheEntry = Context.FindCacheElement(DataKey))
+		{
+			if (ReferencedCacheEntry->IsValid())
+			{
+				if ((*ReferencedCacheEntry)->GetType() == EType::CacheElementReference)
+				{
+					// Recursive call
+					return static_cast<const FContextCacheElementReference*>(ReferencedCacheEntry->Get())->GetReferencedElement(Context);
+				}
+				return ReferencedCacheEntry->Get();
+			}
+		}
+		return nullptr;
+	}
+
 	//////////////////////////////////////////////////////////////////////////////////////////////////////
 
 	const void* FContextCacheElementNull::GetUntypedData(const IContextCacheStore& Context, const FProperty* PropertyIn) const
@@ -100,6 +135,16 @@ namespace UE::Dataflow
 		return {};
 	}
 
+	TUniquePtr<FContextCacheElementBase> FContextCacheElementNull::CreateArrayFromElementAndAppend(const TArray<const FContextCacheElementBase*>& DereferencedElements, const IContextCacheStore& Context, const FProperty* InProperty, const FGuid& InNodeGuid, uint32 InNodeHash, const FTimestamp& InTimestamp) const
+	{
+		return {};
+	}
+
+	TUniquePtr<FContextCacheElementBase> FContextCacheElementNull::AppendArrayElements(const TArray<const FContextCacheElementBase*>& DereferencedElements, const IContextCacheStore& Context, const FProperty* InProperty, const FGuid& InNodeGuid, uint32 InNodeHash, const FTimestamp& InTimestamp) const
+	{
+		return {};
+	}
+
 	TUniquePtr<FContextCacheElementBase> FContextCacheElementNull::Clone(const IContextCacheStore& Context) const
 	{
 		return MakeUnique<FContextCacheElementNull>(GetNodeGuid(), GetProperty(), GetNodeHash(), GetTimestamp());
@@ -142,6 +187,27 @@ namespace UE::Dataflow
 		return MakeUnique<FContextCacheElementUStructArray>(InNodeGuid, InProperty, FConstStructView(InstancedStruct), InNodeHash, InTimestamp);
 	}
 
+	TUniquePtr<FContextCacheElementBase> FContextCacheElementUStruct::CreateArrayFromElementAndAppend(const TArray<const FContextCacheElementBase*>& DereferencedElements, const IContextCacheStore& Context, const FProperty* InProperty, const FGuid& InNodeGuid, uint32 InNodeHash, const FTimestamp& InTimestamp) const
+	{
+		TArray<FConstStructView> ElementViews{ FConstStructView(InstancedStruct) };
+		ElementViews.Reserve(DereferencedElements.Num() + 1);
+		for (int32 Index = 0; Index < DereferencedElements.Num(); ++Index)
+		{
+			check(DereferencedElements[Index]);
+			FConstStructView ElemView = DereferencedElements[Index]->GetConstStructView(Context);
+			if (ElemView.IsValid() && ElemView.GetScriptStruct() == InstancedStruct.GetScriptStruct())
+			{
+				ElementViews.Add(ElemView);
+			}
+		}
+		return MakeUnique<FContextCacheElementUStructArray>(InNodeGuid, InProperty, ElementViews, InNodeHash, InTimestamp);
+	}
+
+	TUniquePtr<FContextCacheElementBase> FContextCacheElementUStruct::AppendArrayElements(const TArray<const FContextCacheElementBase*>& DereferencedElements, const IContextCacheStore& Context, const FProperty* InProperty, const FGuid& InNodeGuid, uint32 InNodeHash, const FTimestamp& InTimestamp) const
+	{
+		return {};
+	}
+
 	FConstStructView FContextCacheElementUStruct::GetConstStructView(const IContextCacheStore& Context) const
 	{
 		return FConstStructView(InstancedStruct);
@@ -194,6 +260,49 @@ namespace UE::Dataflow
 		return {}; // can't make arrays of arrays
 	}
 
+	TUniquePtr<FContextCacheElementBase> FContextCacheElementUStructArray::CreateArrayFromElementAndAppend(const TArray<const FContextCacheElementBase*>& DereferencedElements, const IContextCacheStore& Context, const FProperty* InProperty, const FGuid& InNodeGuid, uint32 InNodeHash, const FTimestamp& InTimestamp) const
+	{
+		return {}; // can't make arrays of arrays
+	}
+
+	TUniquePtr<FContextCacheElementBase> FContextCacheElementUStructArray::AppendArrayElements(const TArray<const FContextCacheElementBase*>& DereferencedElements, const IContextCacheStore& Context, const FProperty* InProperty, const FGuid& InNodeGuid, uint32 InNodeHash, const FTimestamp& InTimestamp) const
+	{
+		TUniquePtr<FContextCacheElementBase> Result;
+		if (InstancedStructArray.GetScriptStruct())
+		{
+			Result = Clone(Context);
+			check(Result.IsValid());
+		}
+		FContextCacheElementUStructArray* TypedResult = static_cast<FContextCacheElementUStructArray*>(Result.Get());
+
+		TArray<const FInstancedStructArray*> OtherArrays;
+		OtherArrays.Reserve(DereferencedElements.Num());
+		for (const FContextCacheElementBase* Element : DereferencedElements)
+		{
+			check(Element);
+			check(Element->GetType() == EType::CacheElementUStructArray);
+			const FContextCacheElementUStructArray* TypedElement = static_cast<const FContextCacheElementUStructArray*>(Element);
+			if (TypedElement->InstancedStructArray.GetScriptStruct())
+			{
+				if (!TypedResult)
+				{
+					Result = TypedElement->Clone(Context);
+					TypedResult = static_cast<FContextCacheElementUStructArray*>(Result.Get());
+					continue;
+				}
+
+				check(TypedElement->InstancedStructArray.GetScriptStruct() == TypedResult->InstancedStructArray.GetScriptStruct());
+				OtherArrays.Add(&TypedElement->InstancedStructArray);
+			}
+		}
+
+		if (TypedResult)
+		{
+			TypedResult->InstancedStructArray.AppendElements(OtherArrays);
+		}
+		return Result;
+	}
+
 	TUniquePtr<FContextCacheElementBase> FContextCacheElementUStructArray::Clone(const IContextCacheStore& Context) const
 	{
 		return MakeUnique<FContextCacheElementUStructArray>(*this);
@@ -226,6 +335,24 @@ namespace UE::Dataflow
 		InitFromRawData(StructView.GetMemory(), 1);
 	}
 
+	FContextCacheElementUStructArray::FInstancedStructArray::FInstancedStructArray(const TArray<FConstStructView>& StructViews)
+		: ScriptStruct(StructViews.Num() ? StructViews[0].GetScriptStruct() : nullptr)
+	{
+		if (ScriptStruct && StructViews.Num())
+		{
+			ArrayNum = StructViews.Num();
+			SetNumUnsafeInternal(ArrayNum);
+			ArrayMax = ArrayNum;
+			GetAllocatorInstance().ResizeAllocation(0, ArrayNum, GetStructureSize(), ScriptStruct->GetMinAlignment());
+			ScriptStruct->InitializeStruct(GetData(), StructViews.Num());
+			for (int32 Index = 0; Index < StructViews.Num(); ++Index)
+			{
+				check(StructViews[Index].GetScriptStruct() == ScriptStruct);
+				ScriptStruct->CopyScriptStruct(GetData() + (GetStructureSize() * Index), StructViews[Index].GetMemory());
+			}			
+		}
+	}
+
 	FContextCacheElementUStructArray::FInstancedStructArray::FInstancedStructArray(const FConstStructArrayView& StructArrayView)
 		: ScriptStruct(StructArrayView.GetScriptStruct())
 	{
@@ -288,6 +415,54 @@ namespace UE::Dataflow
 		}
 	}
 
+	void FContextCacheElementUStructArray::FInstancedStructArray::AddElements(const TArray<FConstStructView>& StructViews)
+	{
+		if (ScriptStruct && StructViews.Num())
+		{
+			const int32 OldArrayNum = ArrayNum;
+			ArrayNum += StructViews.Num();
+			SetNumUnsafeInternal(ArrayNum);
+			ArrayMax = ArrayNum;
+			GetAllocatorInstance().ResizeAllocation(OldArrayNum, ArrayNum, GetStructureSize(), ScriptStruct->GetMinAlignment());
+			ScriptStruct->InitializeStruct(GetData() + (GetStructureSize() * OldArrayNum), StructViews.Num());
+			for (int32 Index = 0; Index < StructViews.Num(); ++Index)
+			{
+				check(StructViews[Index].GetScriptStruct() == ScriptStruct);
+				ScriptStruct->CopyScriptStruct(GetData() + (GetStructureSize() * (OldArrayNum + Index)), StructViews[Index].GetMemory());
+			}
+		}
+	}
+
+	void FContextCacheElementUStructArray::FInstancedStructArray::AppendElements(const TArray<const FInstancedStructArray*>& OtherArrays)
+	{
+		int32 NumNewElements = 0;
+		for (const FInstancedStructArray* OtherArray : OtherArrays)
+		{
+			check(OtherArray);
+			NumNewElements += OtherArray->Num();
+		}
+
+		if (ScriptStruct && NumNewElements > 0)
+		{
+			const int32 OldArrayNum = ArrayNum;
+			ArrayNum += NumNewElements;
+			SetNumUnsafeInternal(ArrayNum);
+			ArrayMax = ArrayNum;
+			GetAllocatorInstance().ResizeAllocation(OldArrayNum, ArrayNum, GetStructureSize(), ScriptStruct->GetMinAlignment());
+			ScriptStruct->InitializeStruct(GetData() + (GetStructureSize() * OldArrayNum), NumNewElements);
+			int32 ElementOffset = OldArrayNum;
+			for(int32 Index = 0; Index < OtherArrays.Num(); ++ Index)
+			{
+				const FInstancedStructArray* OtherArray = OtherArrays[Index];
+				check(OtherArray);
+				check(OtherArray->GetScriptStruct() == ScriptStruct);
+				check(ElementOffset + OtherArray->Num() <= ArrayNum);
+				ScriptStruct->CopyScriptStruct(GetData() + (GetStructureSize() * ElementOffset), OtherArray->GetMemory(), OtherArray->Num());
+				ElementOffset += OtherArray->Num();
+			}
+		}
+	}
+
 	//////////////////////////////////////////////////////////////////////////////////////////////////////
 
 	FContextValue::FContextValue(IContextCacheStore& InContext, FContextCacheKey InCacheKey)
@@ -298,7 +473,7 @@ namespace UE::Dataflow
 	FContextValue::FContextValue(IContextCacheStore& InContext, TUniquePtr<FContextCacheElementBase>&& InCacheElement)
 		: Context(InContext)
 		, CacheKey(0)
-		, CacheElement(MoveTemp(InCacheElement))
+		, CacheElement(InCacheElement.Release())
 	{}
 
 	int32 FContextValue::Num() const
@@ -310,6 +485,15 @@ namespace UE::Dataflow
 		return 0;
 	}
 
+	bool FContextValue::IsArray() const
+	{
+		if (const FContextCacheElementBase* CacheEntry = GetCacheEntryPrivate())
+		{
+			return CacheEntry->IsArray(Context);
+		}
+		return false;
+	}
+
 	FContextValue FContextValue::GetAt(int32 Index) const
 	{
 		if (const FContextCacheElementBase* CacheEntry = GetCacheEntryPrivate())
@@ -342,6 +526,111 @@ namespace UE::Dataflow
 		return FContextValue(Context, 0);
 	}
 
+
+	FContextValue FContextValue::MakeArray(IContextCacheStore& InContext, const TArray<const FContextValue>& Elements)
+	{
+		if (Elements.IsEmpty())
+		{
+			return FContextValue(InContext, 0);
+		}
+
+		// Find the first valid element's cache entry and collect remaining valid cache entries
+		const FContextCacheElementBase* FirstCacheEntry = nullptr;
+		TArray<const FContextCacheElementBase*> ElementEntries;
+		ElementEntries.Reserve(Elements.Num());
+		for (int32 Index = 0; Index < Elements.Num(); ++Index)
+		{
+			if (const FContextCacheElementBase* CacheEntry = Elements[Index].GetCacheEntryPrivate())
+			{
+				if (const FContextCacheElementBase* DeferencedCacheEntry = (CacheEntry->GetType() == FContextCacheElementBase::CacheElementReference) ?
+					static_cast<const FContextCacheElementReference*>(CacheEntry)->GetReferencedElement(InContext) :
+					CacheEntry)
+				{
+					if (!DeferencedCacheEntry->IsArray(InContext))
+					{
+						if (FirstCacheEntry == nullptr)
+						{
+							FirstCacheEntry = DeferencedCacheEntry;
+							continue;
+						}
+						ElementEntries.Emplace(DeferencedCacheEntry);
+					}
+				}
+			}
+		}
+
+		if (!FirstCacheEntry)
+		{
+			// Failed to find a valid cache entry
+			return FContextValue(InContext, 0);
+		}
+
+		TUniquePtr<FContextCacheElementBase> NewCacheElement;
+		if (ElementEntries.IsEmpty())
+		{
+			// Only one cache entry found. Equivalent to ToArray 
+			NewCacheElement = FirstCacheEntry->CreateArrayFromElement(InContext, FirstCacheEntry->GetProperty(), FirstCacheEntry->GetNodeGuid(), FirstCacheEntry->GetNodeHash(), FirstCacheEntry->GetTimestamp());
+		}
+		else
+		{
+			NewCacheElement = FirstCacheEntry->CreateArrayFromElementAndAppend(ElementEntries, InContext, FirstCacheEntry->GetProperty(), FirstCacheEntry->GetNodeGuid(), FirstCacheEntry->GetNodeHash(), FirstCacheEntry->GetTimestamp());
+		}
+		return FContextValue(InContext, MoveTemp(NewCacheElement));
+	}
+
+	FContextValue FContextValue::AppendArrays(IContextCacheStore& InContext, const TArray<const FContextValue>& ArrayElements)
+	{
+		if (ArrayElements.IsEmpty())
+		{
+			return FContextValue(InContext, 0);
+		}
+
+		// Find the first valid element's cache entry and collect remaining valid cache entries
+		const FContextCacheElementBase* FirstCacheEntry = nullptr;
+		int32 FirstCacheEntryIndex = INDEX_NONE;
+		TArray<const FContextCacheElementBase*> ElementEntries;
+		ElementEntries.Reserve(ArrayElements.Num());
+		for (int32 Index = 0; Index < ArrayElements.Num(); ++Index)
+		{
+			if (const FContextCacheElementBase* CacheEntry = ArrayElements[Index].GetCacheEntryPrivate())
+			{
+				if (const FContextCacheElementBase* DeferencedCacheEntry = (CacheEntry->GetType() == FContextCacheElementBase::CacheElementReference) ?
+					static_cast<const FContextCacheElementReference*>(CacheEntry)->GetReferencedElement(InContext) :
+					CacheEntry)
+				{
+					if (DeferencedCacheEntry->IsArray(InContext))
+					{
+						if (FirstCacheEntry == nullptr)
+						{
+							FirstCacheEntry = DeferencedCacheEntry;
+							FirstCacheEntryIndex = Index;
+							continue;
+						}
+						ElementEntries.Emplace(DeferencedCacheEntry);
+					}
+				}
+			}
+		}
+
+		if (!FirstCacheEntry)
+		{
+			// Failed to find a valid cache entry
+			return FContextValue(InContext, 0);
+		}
+
+		TUniquePtr<FContextCacheElementBase> NewCacheElement;
+		if (ElementEntries.IsEmpty())
+		{
+			// Only one cache entry found. Equivalent to ToArray 
+			return ArrayElements[FirstCacheEntryIndex].ToArray();
+		}
+		else
+		{
+			NewCacheElement = FirstCacheEntry->AppendArrayElements(ElementEntries, InContext, FirstCacheEntry->GetProperty(), FirstCacheEntry->GetNodeGuid(), FirstCacheEntry->GetNodeHash(), FirstCacheEntry->GetTimestamp());
+		}
+		return FContextValue(InContext, MoveTemp(NewCacheElement));
+	}
+
 	FConstStructView FContextValue::GetConstStructView() const
 	{
 		if (const FContextCacheElementBase* CacheEntry = GetCacheEntryPrivate())
diff --git a/Engine/Source/Runtime/Experimental/Dataflow/Core/Public/Dataflow/DataflowContextCache.h b/Engine/Source/Runtime/Experimental/Dataflow/Core/Public/Dataflow/DataflowContextCache.h
--- a/Engine/Source/Runtime/Experimental/Dataflow/Core/Public/Dataflow/DataflowContextCache.h
+++ b/Engine/Source/Runtime/Experimental/Dataflow/Core/Public/Dataflow/DataflowContextCache.h
@@ -98,6 +98,10 @@ namespace UE::Dataflow
 		virtual int32 GetNumArrayElements(const IContextCacheStore& Context) const = 0;
 		virtual TUniquePtr<FContextCacheElementBase> CreateFromArrayElement(const IContextCacheStore& Context, int32 Index, const FProperty* InProperty, const FGuid& InNodeGuid, uint32 InNodeHash, const FTimestamp& InTimestamp) const = 0;
 		virtual TUniquePtr<FContextCacheElementBase> CreateArrayFromElement(const IContextCacheStore& Context, const FProperty* InProperty, const FGuid& InNodeGuid, uint32 InNodeHash, const FTimestamp& InTimestamp) const = 0;
+		// DereferencedElements will not be FContextCacheElementReference and instead be the final element any reference stored.
+		virtual TUniquePtr<FContextCacheElementBase> CreateArrayFromElementAndAppend(const TArray<const FContextCacheElementBase*>& DereferencedElements, const IContextCacheStore& Context, const FProperty* InProperty, const FGuid& InNodeGuid, uint32 InNodeHash, const FTimestamp& InTimestamp) const = 0;
+		// DereferencedElements will not be FContextCacheElementReference and instead be the final element any reference stored.
+		virtual TUniquePtr<FContextCacheElementBase> AppendArrayElements(const TArray<const FContextCacheElementBase*>& DereferencedElements, const IContextCacheStore& Context, const FProperty* InProperty, const FGuid& InNodeGuid, uint32 InNodeHash, const FTimestamp& InTimestamp) const = 0;
 
 		// Get the matching const struct view if available, otherwise return an invalid one 
 		virtual FConstStructView GetConstStructView(const IContextCacheStore& Context) const { return {}; }
@@ -202,6 +206,50 @@ namespace UE::Dataflow
 			}
 		}
 
+		virtual TUniquePtr<FContextCacheElementBase> CreateArrayFromElementAndAppend(const TArray<const FContextCacheElementBase*>& DereferencedElements, const IContextCacheStore& Context, const FProperty* InProperty, const FGuid& InNodeGuid, uint32 InNodeHash, const FTimestamp& InTimestamp) const override
+		{
+			if constexpr (TIsTArray<FDataType>::Value)
+			{
+				return {}; // already an Array
+			}
+			else
+			{
+				TArray<FDataType> Result{ Data };
+				Result.Reserve(DereferencedElements.Num() + 1);
+				for (const FContextCacheElementBase* Element : DereferencedElements)
+				{
+					check(Element);
+					check(Element->GetType() == EType::CacheElementTyped);
+					const void* UntypedElementData = Element->GetUntypedData(Context, Element->GetProperty());
+					check(UntypedElementData);
+					Result.Emplace(*static_cast<const FDataType*>(UntypedElementData));
+				}
+
+				return MakeUnique<TContextCacheElement<TArray<FDataType>>>(InNodeGuid, InProperty, MoveTemp(Result), InNodeHash, InTimestamp);
+			}
+		}
+
+		virtual TUniquePtr<FContextCacheElementBase> AppendArrayElements(const TArray<const FContextCacheElementBase*>& DereferencedElements, const IContextCacheStore& Context, const FProperty* InProperty, const FGuid& InNodeGuid, uint32 InNodeHash, const FTimestamp& InTimestamp) const override
+		{
+			if constexpr (TIsTArray<FDataType>::Value)
+			{
+				FDataType Result(Data);
+				for (const FContextCacheElementBase* Element : DereferencedElements)
+				{
+					check(Element);
+					check(Element->GetType() == EType::CacheElementTyped);
+					const void* UntypedElementData = Element->GetUntypedData(Context, Element->GetProperty());
+					check(UntypedElementData);
+					Result.Append(*static_cast<const FDataType*>(UntypedElementData));
+				}
+				return MakeUnique<TContextCacheElement<FDataType>>(InNodeGuid, InProperty, MoveTemp(Result), InNodeHash, InTimestamp);
+			}
+			else
+			{
+				return {};
+			}
+		}
+
 		const T& GetDataDirect() const { return Data; }
 
 		inline virtual TUniquePtr<FContextCacheElementBase> Clone(const IContextCacheStore& Context) const override;
@@ -235,10 +283,17 @@ namespace UE::Dataflow
 
 		virtual TUniquePtr<FContextCacheElementBase> CreateArrayFromElement(const IContextCacheStore& Context, const FProperty* InProperty, const FGuid& InNodeGuid, uint32 InNodeHash, const FTimestamp& InTimestamp) const override;
 
+		virtual TUniquePtr<FContextCacheElementBase> CreateArrayFromElementAndAppend(const TArray<const FContextCacheElementBase*>& DereferencedElements, const IContextCacheStore& Context, const FProperty* InProperty, const FGuid& InNodeGuid, uint32 InNodeHash, const FTimestamp& InTimestamp) const override;
+
+		virtual TUniquePtr<FContextCacheElementBase> AppendArrayElements(const TArray<const FContextCacheElementBase*>& DereferencedElements, const IContextCacheStore& Context, const FProperty* InProperty, const FGuid& InNodeGuid, uint32 InNodeHash, const FTimestamp& InTimestamp) const override;
+
 		virtual FConstStructView GetConstStructView(const IContextCacheStore& Context) const;
 
 		virtual TUniquePtr<FContextCacheElementBase> Clone(const IContextCacheStore& Context) const override;
 
+		// Traverse references until the non-reference element is found.
+		DATAFLOWCORE_API const FContextCacheElementBase* GetReferencedElement(const IContextCacheStore& Context) const;
+
 	private:
 		const FContextCacheKey DataKey; // this is a key to another cache element
 	};
@@ -267,6 +322,8 @@ namespace UE::Dataflow
 		virtual int32 GetNumArrayElements(const IContextCacheStore& Context) const override;
 		virtual TUniquePtr<FContextCacheElementBase> CreateFromArrayElement(const IContextCacheStore& Context, int32 Index, const FProperty* InProperty, const FGuid& InNodeGuid, uint32 InNodeHash, const FTimestamp& InTimestamp) const override;
 		virtual TUniquePtr<FContextCacheElementBase> CreateArrayFromElement(const IContextCacheStore& Context, const FProperty* InProperty, const FGuid& InNodeGuid, uint32 InNodeHash, const FTimestamp& InTimestamp) const override;
+		virtual TUniquePtr<FContextCacheElementBase> CreateArrayFromElementAndAppend(const TArray<const FContextCacheElementBase*>& DereferencedElements, const IContextCacheStore& Context, const FProperty* InProperty, const FGuid& InNodeGuid, uint32 InNodeHash, const FTimestamp& InTimestamp) const override;
+		virtual TUniquePtr<FContextCacheElementBase> AppendArrayElements(const TArray<const FContextCacheElementBase*>& DereferencedElements, const IContextCacheStore& Context, const FProperty* InProperty, const FGuid& InNodeGuid, uint32 InNodeHash, const FTimestamp& InTimestamp) const override;
 
 		virtual TUniquePtr<FContextCacheElementBase> Clone(const IContextCacheStore& Context) const override;
 	};
@@ -318,6 +375,26 @@ namespace UE::Dataflow
 			return MakeUnique<TContextCacheElementUObjectArray<FDataType>>(InNodeGuid, InProperty, MoveTemp(Array), InNodeHash, InTimestamp);
 		}
 
+		virtual TUniquePtr<FContextCacheElementBase> CreateArrayFromElementAndAppend(const TArray<const FContextCacheElementBase*>& DereferencedElements, const IContextCacheStore& Context, const FProperty* InProperty, const FGuid& InNodeGuid, uint32 InNodeHash, const FTimestamp& InTimestamp) const override
+		{
+			TArray<FDataType> Result{ ObjectPtr };
+			Result.Reserve(DereferencedElements.Num() + 1);
+			for (const FContextCacheElementBase* Element : DereferencedElements)
+			{
+				check(Element);
+				check(Element->GetType() == EType::CacheElementUObject);
+				const TContextCacheElementUObject<T>* TypedElement = static_cast<const TContextCacheElementUObject<T>*>(Element);
+				Result.Emplace(TypedElement->ObjectPtr);
+			}
+
+			return MakeUnique<TContextCacheElementUObjectArray<FDataType>>(InNodeGuid, InProperty, MoveTemp(Result), InNodeHash, InTimestamp);
+		}
+
+		virtual TUniquePtr<FContextCacheElementBase> AppendArrayElements(const TArray<const FContextCacheElementBase*>& DereferencedElements, const IContextCacheStore& Context, const FProperty* InProperty, const FGuid& InNodeGuid, uint32 InNodeHash, const FTimestamp& InTimestamp) const override
+		{
+			return {};
+		}
+
 		inline virtual TUniquePtr<FContextCacheElementBase> Clone(const IContextCacheStore& Context) const override;
 
 		//~ Begin FGCObject interface
@@ -383,6 +460,25 @@ namespace UE::Dataflow
 			return {}; // no array of array 
 		}
 
+		virtual TUniquePtr<FContextCacheElementBase> CreateArrayFromElementAndAppend(const TArray<const FContextCacheElementBase*>& DereferencedElements, const IContextCacheStore& Context, const FProperty* InProperty, const FGuid& InNodeGuid, uint32 InNodeHash, const FTimestamp& InTimestamp) const override
+		{
+			return {}; // no array of array 
+		}
+
+		virtual TUniquePtr<FContextCacheElementBase> AppendArrayElements(const TArray<const FContextCacheElementBase*>& DereferencedElements, const IContextCacheStore& Context, const FProperty* InProperty, const FGuid& InNodeGuid, uint32 InNodeHash, const FTimestamp& InTimestamp) const override
+		{
+			TArray<FDataType> Result(ObjectPtrArray);
+			for (const FContextCacheElementBase* Element : DereferencedElements)
+			{
+				check(Element);
+				check(Element->GetType() == EType::CacheElementUObjectArray);
+				const void* UntypedElementData = Element->GetUntypedData(Context, Element->GetProperty());
+				check(UntypedElementData);
+				Result.Append(*static_cast<const TArray<FDataType>*>(UntypedElementData));
+			}
+			return MakeUnique< TContextCacheElementUObjectArray<FDataType>>(InNodeGuid, InProperty, MoveTemp(Result), InNodeHash, InTimestamp);
+		}
+
 		inline virtual TUniquePtr<FContextCacheElementBase> Clone(const IContextCacheStore& Context) const override;
 
 		//~ Begin FGCObject interface
@@ -423,6 +519,8 @@ namespace UE::Dataflow
 		DATAFLOWCORE_API virtual int32 GetNumArrayElements(const IContextCacheStore& Context) const override;
 		DATAFLOWCORE_API virtual TUniquePtr<FContextCacheElementBase> CreateFromArrayElement(const IContextCacheStore& Context, int32 Index, const FProperty* InProperty, const FGuid& InNodeGuid, uint32 InNodeHash, const FTimestamp& InTimestamp) const override;
 		DATAFLOWCORE_API virtual TUniquePtr<FContextCacheElementBase> CreateArrayFromElement(const IContextCacheStore& Context, const FProperty* InProperty, const FGuid& InNodeGuid, uint32 InNodeHash, const FTimestamp& InTimestamp) const override;
+		DATAFLOWCORE_API virtual TUniquePtr<FContextCacheElementBase> CreateArrayFromElementAndAppend(const TArray<const FContextCacheElementBase*>& DereferencedElements, const IContextCacheStore& Context, const FProperty* InProperty, const FGuid& InNodeGuid, uint32 InNodeHash, const FTimestamp& InTimestamp) const override;
+		DATAFLOWCORE_API virtual TUniquePtr<FContextCacheElementBase> AppendArrayElements(const TArray<const FContextCacheElementBase*>& DereferencedElements, const IContextCacheStore& Context, const FProperty* InProperty, const FGuid& InNodeGuid, uint32 InNodeHash, const FTimestamp& InTimestamp) const override;
 
 		// Get the matching const struct view if available, otherwise return an invalid one 
 		DATAFLOWCORE_API virtual FConstStructView GetConstStructView(const IContextCacheStore& Context) const;
@@ -453,6 +551,10 @@ namespace UE::Dataflow
 			: FContextCacheElementBase(EType::CacheElementUStructArray, InNodeGuid, InProperty, InNodeHash, Timestamp)
 			, InstancedStructArray(InStructView)
 		{}
+		explicit FContextCacheElementUStructArray(FGuid InNodeGuid, const FProperty* InProperty, const TArray<FConstStructView>& InStructViews, uint32 InNodeHash, FTimestamp Timestamp)
+			: FContextCacheElementBase(EType::CacheElementUStructArray, InNodeGuid, InProperty, InNodeHash, Timestamp)
+			, InstancedStructArray(InStructViews)
+		{}
 
 		template<typename T>
 		explicit FContextCacheElementUStructArray(FGuid InNodeGuid, const FProperty* InProperty, TArray<T>&& InStructArray, uint32 InNodeHash, FTimestamp Timestamp)
@@ -480,6 +582,8 @@ namespace UE::Dataflow
 		virtual int32 GetNumArrayElements(const IContextCacheStore& Context) const override;
 		virtual TUniquePtr<FContextCacheElementBase> CreateFromArrayElement(const IContextCacheStore& Context, int32 Index, const FProperty* InProperty, const FGuid& InNodeGuid, uint32 InNodeHash, const FTimestamp& InTimestamp) const override;
 		virtual TUniquePtr<FContextCacheElementBase> CreateArrayFromElement(const IContextCacheStore& Context, const FProperty* InProperty, const FGuid& InNodeGuid, uint32 InNodeHash, const FTimestamp& InTimestamp) const override;
+		virtual TUniquePtr<FContextCacheElementBase> CreateArrayFromElementAndAppend(const TArray<const FContextCacheElementBase*>& DereferencedElements, const IContextCacheStore& Context, const FProperty* InProperty, const FGuid& InNodeGuid, uint32 InNodeHash, const FTimestamp& InTimestamp) const override;
+		virtual TUniquePtr<FContextCacheElementBase> AppendArrayElements(const TArray<const FContextCacheElementBase*>& DereferencedElements, const IContextCacheStore& Context, const FProperty* InProperty, const FGuid& InNodeGuid, uint32 InNodeHash, const FTimestamp& InTimestamp) const override;
 
 		virtual TUniquePtr<FContextCacheElementBase> Clone(const IContextCacheStore& Context) const override;
 
@@ -496,6 +600,7 @@ namespace UE::Dataflow
 		{
 			explicit FInstancedStructArray(const UScriptStruct* const InScriptStruct);
 			explicit FInstancedStructArray(const FConstStructView& StructView);
+			explicit FInstancedStructArray(const TArray<FConstStructView>& StructViews);
 			explicit FInstancedStructArray(const FConstStructArrayView& StructArrayView);
 			~FInstancedStructArray();
 
@@ -514,6 +619,12 @@ namespace UE::Dataflow
 			const void* GetMemory() const;
 
 			void AddReferencedObjects(FReferenceCollector& Collector);
+			
+			// Add elements from other struct views.
+			void AddElements(const TArray<FConstStructView>& StructViews);
+
+			// Append arrays from other FInstancedStructArrays
+			void AppendElements(const TArray<const FInstancedStructArray*>& OtherArrays);
 
 			template<typename T>
 			const TArray<T>& Get() const
@@ -651,8 +762,13 @@ namespace UE::Dataflow
 		DATAFLOWCORE_API FContextValue(IContextCacheStore& InContext, TUniquePtr<FContextCacheElementBase>&& InCacheElement);
 
 		DATAFLOWCORE_API int32 Num() const;
+		DATAFLOWCORE_API bool IsArray() const;
 		DATAFLOWCORE_API FContextValue GetAt(int32 Index) const;
 		DATAFLOWCORE_API FContextValue ToArray() const;
+		// Warning: no checking will be done to verify that Elements can be appended to the array for strongly typed TContextCacheElement arrays and elements. It is assumed they can be cast to the correct type.
+		DATAFLOWCORE_API static FContextValue MakeArray(IContextCacheStore& InContext, const TArray<const FContextValue>& Elements);
+		// Warning: no checking will be done to verify that Elements can be appended to the array for strongly typed TContextCacheElement arrays and elements. It is assumed they can be cast to the correct type.
+		DATAFLOWCORE_API static FContextValue AppendArrays(IContextCacheStore& InContext, const TArray<const FContextValue>& ArrayElements);
 
 		DATAFLOWCORE_API FConstStructView GetConstStructView() const;
 		DATAFLOWCORE_API FInstancedStruct GetInstancedStruct() const;
@@ -685,7 +801,7 @@ namespace UE::Dataflow
 
 		IContextCacheStore& Context;
 		FContextCacheKey CacheKey; // only set when CacheElement is null
-		TUniquePtr<FContextCacheElementBase> CacheElement; // when CacheKey is define this should be null
+		TSharedPtr<FContextCacheElementBase> CacheElement; // when CacheKey is define this should be null
 	};
 };
 
