Chaos API: Doing a cleanup pass.
- Fixing whitespace / formatting
- Adding missing namespace / ifdef closer comments
- Updating classes to have member default initializers, especially for pointers
- Re-ordering some functions to be public interface more at the top
- Giving all handles GetSceneId and removing GetSceneHandle
- Adding UE_EXPERIMENTAL to most files

#Synced-CL 51413247

diff --git a/Engine/Source/Runtime/Experimental/RigidPhysics/Private/RigidPhysics/JointConstraintHandle.cpp b/Engine/Source/Runtime/Experimental/RigidPhysics/Private/RigidPhysics/JointConstraintHandle.cpp
--- a/Engine/Source/Runtime/Experimental/RigidPhysics/Private/RigidPhysics/JointConstraintHandle.cpp
+++ b/Engine/Source/Runtime/Experimental/RigidPhysics/Private/RigidPhysics/JointConstraintHandle.cpp
@@ -24,6 +24,11 @@ namespace UE::Physics
 		return ConstraintId;
 	}
 
+	const FRigidSceneId& FJointConstraintHandle::GetSceneId() const
+	{
+		return SceneId;
+	}
+
 	void FJointConstraintHandle::Reset()
 	{
 		SceneId.Reset();
diff --git a/Engine/Source/Runtime/Experimental/RigidPhysics/Private/RigidPhysics/RigidBodyContainerHandle.cpp b/Engine/Source/Runtime/Experimental/RigidPhysics/Private/RigidPhysics/RigidBodyContainerHandle.cpp
--- a/Engine/Source/Runtime/Experimental/RigidPhysics/Private/RigidPhysics/RigidBodyContainerHandle.cpp
+++ b/Engine/Source/Runtime/Experimental/RigidPhysics/Private/RigidPhysics/RigidBodyContainerHandle.cpp
@@ -24,20 +24,15 @@ namespace UE::Physics
 		return BodyContainerId;
 	}
 
-	const FRigidSceneId FRigidBodyContainerHandle::GetSceneId() const
+	const FRigidSceneId& FRigidBodyContainerHandle::GetSceneId() const
 	{
 		return SceneId;
 	}
 
-	const FRigidSceneHandle FRigidBodyContainerHandle::GetSceneHandle() const
-	{
-		return FRigidSceneHandle(SceneId);
-	}
-
 	IRigidBodyContainer* FRigidBodyContainerHandle::PinInternal(const IRigidScene* InScene) const
 	{
 		return InScene->GetBodyContainer(GetId());
 	}
 } // namespace UE::Physics
 
-#endif // UE_RIGIDPHYSICS_API_ENABLED
\ No newline at end of file
+#endif // UE_RIGIDPHYSICS_API_ENABLED
diff --git a/Engine/Source/Runtime/Experimental/RigidPhysics/Private/RigidPhysics/RigidGeometryCollectionHandle.cpp b/Engine/Source/Runtime/Experimental/RigidPhysics/Private/RigidPhysics/RigidGeometryCollectionHandle.cpp
--- a/Engine/Source/Runtime/Experimental/RigidPhysics/Private/RigidPhysics/RigidGeometryCollectionHandle.cpp
+++ b/Engine/Source/Runtime/Experimental/RigidPhysics/Private/RigidPhysics/RigidGeometryCollectionHandle.cpp
@@ -23,9 +23,9 @@ namespace UE::Physics
 		return BodyContainerHandle.GetId();
 	}
 
-	const FRigidSceneHandle FRigidGeometryCollectionHandle::GetSceneHandle() const
+	const FRigidSceneId& FRigidGeometryCollectionHandle::GetSceneId() const
 	{
-		return BodyContainerHandle.GetSceneHandle();
+		return BodyContainerHandle.GetSceneId();
 	}
 
 	FRigidBodyContainerHandle FRigidGeometryCollectionHandle::AsBodyContainerHandle() const
@@ -50,4 +50,4 @@ namespace UE::Physics
 	}
 } // namespace UE::Physics
 
-#endif // UE_RIGIDPHYSICS_API_ENABLED
\ No newline at end of file
+#endif // UE_RIGIDPHYSICS_API_ENABLED
diff --git a/Engine/Source/Runtime/Experimental/RigidPhysics/Private/RigidPhysics/RigidTyped.cpp b/Engine/Source/Runtime/Experimental/RigidPhysics/Private/RigidPhysics/RigidTyped.cpp
--- a/Engine/Source/Runtime/Experimental/RigidPhysics/Private/RigidPhysics/RigidTyped.cpp
+++ b/Engine/Source/Runtime/Experimental/RigidPhysics/Private/RigidPhysics/RigidTyped.cpp
@@ -5,14 +5,6 @@
 
 namespace UE::Physics
 {
-	IRigidTyped::IRigidTyped()
-	{
-	}
-
-	IRigidTyped::~IRigidTyped()
-	{
-	}
-
 	const FRigidTypeId& IRigidTyped::GetStaticTypeId()
 	{
 		static const TCHAR* StaticTypeName = TEXT("IRigidTyped");
@@ -24,6 +16,6 @@ namespace UE::Physics
 	{
 		return IRigidTyped::GetStaticTypeId();
 	}
-}
+} // namespace UE::Physics
 
-#endif // UE_RIGIDPHYSICS_API_ENABLED
\ No newline at end of file
+#endif // UE_RIGIDPHYSICS_API_ENABLED
diff --git a/Engine/Source/Runtime/Experimental/RigidPhysics/Public/RigidPhysics/Internal/IRigidGeometryCollection.h b/Engine/Source/Runtime/Experimental/RigidPhysics/Public/RigidPhysics/Internal/IRigidGeometryCollection.h
--- a/Engine/Source/Runtime/Experimental/RigidPhysics/Public/RigidPhysics/Internal/IRigidGeometryCollection.h
+++ b/Engine/Source/Runtime/Experimental/RigidPhysics/Public/RigidPhysics/Internal/IRigidGeometryCollection.h
@@ -21,7 +21,7 @@ namespace UE::Physics
 		virtual ~IRigidGeometryCollection() = default;
 
 		// TODO_CHAOSAPI: Figure out the public GeometryCollection API!!
-		UE_INTERNAL RIGIDPHYSICS_API virtual IRigidBody* GetBodyAt(int32 BodyIndex) const = 0;
+		UE_INTERNAL virtual IRigidBody* GetBodyAt(int32 BodyIndex) const = 0;
 	};
 } // namespace UE::Physics
 
diff --git a/Engine/Source/Runtime/Experimental/RigidPhysics/Public/RigidPhysics/JointConstraintHandle.h b/Engine/Source/Runtime/Experimental/RigidPhysics/Public/RigidPhysics/JointConstraintHandle.h
--- a/Engine/Source/Runtime/Experimental/RigidPhysics/Public/RigidPhysics/JointConstraintHandle.h
+++ b/Engine/Source/Runtime/Experimental/RigidPhysics/Public/RigidPhysics/JointConstraintHandle.h
@@ -39,6 +39,7 @@ namespace UE::Physics
 		}
 
 		UE_INTERNAL RIGIDPHYSICS_API const FRigidObjectId& GetId() const;
+		UE_INTERNAL RIGIDPHYSICS_API const FRigidSceneId& GetSceneId() const;
 		UE_INTERNAL RIGIDPHYSICS_API void Reset();
 
 	private:
diff --git a/Engine/Source/Runtime/Experimental/RigidPhysics/Public/RigidPhysics/RigidBody.h b/Engine/Source/Runtime/Experimental/RigidPhysics/Public/RigidPhysics/RigidBody.h
--- a/Engine/Source/Runtime/Experimental/RigidPhysics/Public/RigidPhysics/RigidBody.h
+++ b/Engine/Source/Runtime/Experimental/RigidPhysics/Public/RigidPhysics/RigidBody.h
@@ -13,12 +13,13 @@
 #include "RigidPhysics/RigidObjectPtr.h"
 #include "RigidPhysics/RigidShapeInstance.h"
 
+UE_EXPERIMENTAL(5.8, "The new Chaos API is experimental")
 namespace UE::Physics
 {
 	// This class only exists so that TRigidBodyPtr can have pointer semantics rather than object semantics,
 	// and so that we can disable functions based on ContextType
 	// See TRigidBodyPtr for comments.
-	template<typename ContextType>
+	template <typename ContextType>
 	class UE_INTERNAL TRigidBodyPtrImpl
 	{
 	public:
@@ -26,11 +27,7 @@ namespace UE::Physics
 		using FInterface = IRigidBody;
 		using FHandle = FRigidBodyHandle;
 
-		UE_INTERNAL TRigidBodyPtrImpl()
-			: BodyRaw(nullptr)
-		{
-		}
-
+		UE_INTERNAL TRigidBodyPtrImpl() = default;
 		UE_INTERNAL TRigidBodyPtrImpl(IRigidBody* InBody)
 			: BodyRaw(InBody)
 		{
@@ -49,21 +46,14 @@ namespace UE::Physics
 			Reset();
 		}
 
-		UE_INTERNAL void Reset()
-		{
-			BodyRaw = nullptr;
-		}
+		friend bool operator==(const TRigidBodyPtrImpl&, const TRigidBodyPtrImpl&) = default;
+		friend bool operator!=(const TRigidBodyPtrImpl&, const TRigidBodyPtrImpl&) = default;
 
-		UE_INTERNAL bool IsValid() const
+		bool IsValid() const
 		{
 			return (BodyRaw != nullptr);
 		}
 
-		const FRigidTypeId& GetTypeId() const
-		{
-			return BodyRaw->GetTypeId();
-		}
-
 		FRigidBodyHandle GetHandle() const
 		{
 			if (BodyRaw != nullptr)
@@ -92,7 +82,7 @@ namespace UE::Physics
 		{
 			return BodyRaw->IsStatic();
 		}
-		
+
 		bool IsKinematic() const
 		{
 			return BodyRaw->IsKinematic();
@@ -192,7 +182,7 @@ namespace UE::Physics
 		{
 			BodyRaw->SetMaxLinearVelocity(InMax);
 		}
-		
+
 		void SetMaxAngularVelocity(float InMax)
 		{
 			BodyRaw->SetMaxAngularVelocity(InMax);
@@ -212,7 +202,7 @@ namespace UE::Physics
 		{
 			BodyRaw->SetGravityGroup(InId);
 		}
-		
+
 		void SetGravityScale(float InScale)
 		{
 			BodyRaw->SetGravityScale(InScale);
@@ -243,18 +233,25 @@ namespace UE::Physics
 			BodyRaw->SetIsSleeping(bInSleeping);
 		}
 
-		friend bool operator==(const TRigidBodyPtrImpl& L, const TRigidBodyPtrImpl& R) = default;
-		friend bool operator!=(const TRigidBodyPtrImpl& L, const TRigidBodyPtrImpl& R) = default;
+		UE_INTERNAL void Reset()
+		{
+			BodyRaw = nullptr;
+		}
 
 		UE_INTERNAL IRigidBody* Get() const
 		{
 			return BodyRaw;
 		}
 
+		UE_INTERNAL const FRigidTypeId& GetTypeId() const
+		{
+			return BodyRaw->GetTypeId();
+		}
+
 	private:
 		// IRigidBody implementation depends on Context type
-		IRigidBody* BodyRaw;
+		IRigidBody* BodyRaw = nullptr;
 	};
-}
+} // namespace UE::Physics
 
-#endif // UE_RIGIDPHYSICS_API_ENABLED
\ No newline at end of file
+#endif // UE_RIGIDPHYSICS_API_ENABLED
diff --git a/Engine/Source/Runtime/Experimental/RigidPhysics/Public/RigidPhysics/RigidBodyContainer.h b/Engine/Source/Runtime/Experimental/RigidPhysics/Public/RigidPhysics/RigidBodyContainer.h
--- a/Engine/Source/Runtime/Experimental/RigidPhysics/Public/RigidPhysics/RigidBodyContainer.h
+++ b/Engine/Source/Runtime/Experimental/RigidPhysics/Public/RigidPhysics/RigidBodyContainer.h
@@ -13,11 +13,12 @@
 #include "RigidPhysics/RigidObjectRegistry.h"
 #include "RigidPhysics/RigidTyped.h"
 
+UE_EXPERIMENTAL(5.8, "The new Chaos API is experimental")
 namespace UE::Physics
 {
 	// Utility wrapper around IRigidBodyContainer that adds Context and allows
 	// TRigidBodyContainerPtr to have pointer semantics without exposing IRigidBodyContainer
-	template<typename ContextType>
+	template <typename ContextType>
 	class UE_INTERNAL TRigidBodyContainerPtrImpl
 	{
 	public:
@@ -31,30 +32,14 @@ namespace UE::Physics
 		{
 		}
 
-		virtual ~TRigidBodyContainerPtrImpl()
-		{
-		}
+		~TRigidBodyContainerPtrImpl() = default;
 
-		UE_INTERNAL void Reset()
-		{
-			BodyContainerRaw = nullptr;
-		}
 
-		UE_INTERNAL bool IsValid() const
+		bool IsValid() const
 		{
 			return (BodyContainerRaw != nullptr);
 		}
 
-		const FRigidTypeId& GetTypeId() const
-		{
-			return BodyContainerRaw->GetTypeId();
-		}
-
-		FRigidObjectId GetId() const
-		{
-			return BodyContainerRaw->GetId();
-		}
-
 		FRigidBodyContainerHandle GetHandle() const
 		{
 			return BodyContainerRaw->GetHandle();
@@ -85,6 +70,16 @@ namespace UE::Physics
 			BodyContainerRaw->DestroyBody(InBodyPtr.Get());
 		}
 
+		UE_INTERNAL void Reset()
+		{
+			BodyContainerRaw = nullptr;
+		}
+
+		UE_INTERNAL const FRigidTypeId& GetTypeId() const
+		{
+			return BodyContainerRaw->GetTypeId();
+		}
+
 		UE_INTERNAL IRigidBodyContainer* Get() const
 		{
 			return BodyContainerRaw;
@@ -92,8 +87,8 @@ namespace UE::Physics
 
 	private:
 		// Pointer to implementation, object depends on Context type
-		IRigidBodyContainer* BodyContainerRaw;
+		IRigidBodyContainer* BodyContainerRaw = nullptr;
 	};
-}
+} // namespace UE::Physics
 
-#endif // UE_RIGIDPHYSICS_API_ENABLED
\ No newline at end of file
+#endif // UE_RIGIDPHYSICS_API_ENABLED
diff --git a/Engine/Source/Runtime/Experimental/RigidPhysics/Public/RigidPhysics/RigidBodyContainerHandle.h b/Engine/Source/Runtime/Experimental/RigidPhysics/Public/RigidPhysics/RigidBodyContainerHandle.h
--- a/Engine/Source/Runtime/Experimental/RigidPhysics/Public/RigidPhysics/RigidBodyContainerHandle.h
+++ b/Engine/Source/Runtime/Experimental/RigidPhysics/Public/RigidPhysics/RigidBodyContainerHandle.h
@@ -9,6 +9,7 @@
 #include "RigidPhysics/RigidObjectPtr.h"
 #include "RigidPhysics/RigidSceneHandle.h"
 
+UE_EXPERIMENTAL(5.8, "The new Chaos API is experimental")
 namespace UE::Physics
 {
 	template <typename ContextType>
@@ -36,8 +37,7 @@ namespace UE::Physics
 
 		UE_INTERNAL RIGIDPHYSICS_API void Reset();
 		UE_INTERNAL RIGIDPHYSICS_API const FRigidObjectId& GetId() const;
-		UE_INTERNAL RIGIDPHYSICS_API const FRigidSceneId GetSceneId() const;
-		UE_INTERNAL RIGIDPHYSICS_API const FRigidSceneHandle GetSceneHandle() const;
+		UE_INTERNAL RIGIDPHYSICS_API const FRigidSceneId& GetSceneId() const;
 
 	private:
 		RIGIDPHYSICS_API IRigidBodyContainer* PinInternal(const IRigidScene* InScene) const;
@@ -47,4 +47,4 @@ namespace UE::Physics
 	};
 } // namespace UE::Physics
 
-#endif // UE_RIGIDPHYSICS_API_ENABLED
\ No newline at end of file
+#endif // UE_RIGIDPHYSICS_API_ENABLED
diff --git a/Engine/Source/Runtime/Experimental/RigidPhysics/Public/RigidPhysics/RigidBodyHandle.h b/Engine/Source/Runtime/Experimental/RigidPhysics/Public/RigidPhysics/RigidBodyHandle.h
--- a/Engine/Source/Runtime/Experimental/RigidPhysics/Public/RigidPhysics/RigidBodyHandle.h
+++ b/Engine/Source/Runtime/Experimental/RigidPhysics/Public/RigidPhysics/RigidBodyHandle.h
@@ -9,6 +9,7 @@
 #include "RigidPhysics/RigidObjectPtr.h"
 #include "RigidPhysics/RigidSceneHandle.h"
 
+UE_EXPERIMENTAL(5.8, "The new Chaos API is experimental")
 namespace UE::Physics
 {
 	class IRigidBody;
diff --git a/Engine/Source/Runtime/Experimental/RigidPhysics/Public/RigidPhysics/RigidContext.h b/Engine/Source/Runtime/Experimental/RigidPhysics/Public/RigidPhysics/RigidContext.h
--- a/Engine/Source/Runtime/Experimental/RigidPhysics/Public/RigidPhysics/RigidContext.h
+++ b/Engine/Source/Runtime/Experimental/RigidPhysics/Public/RigidPhysics/RigidContext.h
@@ -9,11 +9,12 @@
 #include "RigidPhysics/RigidLog.h"
 #include "RigidPhysics/RigidScene.h"
 
+UE_EXPERIMENTAL(5.8, "The new Chaos API is experimental")
 namespace UE::Physics
 {
 	// A Context holds a lock (RW or RO) and converts object handles
 	// to context-specific object pointers.
-	template<ERigidContextType ContextTypeParam, ERigidLockType LockTypeParam>
+	template <ERigidContextType ContextTypeParam, ERigidLockType LockTypeParam>
 	class TRigidContext
 	{
 	public:
@@ -76,6 +77,6 @@ namespace UE::Physics
 		TRigidScenePtr<FContext> ScenePtr;
 		FRigidSceneId SceneId;
 	};
-}
+} // namespace UE::Physics
 
 #endif // UE_RIGIDPHYSICS_API_ENABLED
diff --git a/Engine/Source/Runtime/Experimental/RigidPhysics/Public/RigidPhysics/RigidFwd.h b/Engine/Source/Runtime/Experimental/RigidPhysics/Public/RigidPhysics/RigidFwd.h
--- a/Engine/Source/Runtime/Experimental/RigidPhysics/Public/RigidPhysics/RigidFwd.h
+++ b/Engine/Source/Runtime/Experimental/RigidPhysics/Public/RigidPhysics/RigidFwd.h
@@ -104,7 +104,7 @@ namespace UE::Physics
 
 	using FRigidDebugName = ::Chaos::FSharedDebugName;
 	using FRigidSceneWeakPtr = TWeakPtr<IRigidScene, ESPMode::ThreadSafe>;
-}
+} // namespace UE::Physics
 
 namespace UE::Physics
 {
@@ -149,7 +149,7 @@ namespace UE::Physics
 	using FRigidContextGameRO = TRigidContext<ERigidContextType::Game, ERigidLockType::ReadOnly>;
 	using FRigidContextGameRW = TRigidContext<ERigidContextType::Game, ERigidLockType::ReadWrite>;
 	using FRigidContextSimRW = TRigidContext<ERigidContextType::Sim, ERigidLockType::ReadWrite>;
-}
+} // namespace UE::Physics
 
 //
 // Internal Forward Decls
@@ -162,8 +162,8 @@ namespace UE::Physics
 	using UE_INTERNAL FRigidBodyContainerRegistry = TRigidObjectRegistry<IRigidBodyContainer>;
 	using UE_INTERNAL FRigidGeometryCollectionRegistry = TRigidObjectRegistry<IRigidGeometryCollection>;
 	using UE_INTERNAL FRigidBodyRegistry = TRigidObjectRegistry<IRigidBody>;
-}
+} // namespace UE::Physics
 
 extern uint32 GetTypeHash(UE::Physics::FRigidObjectId& RigidBodyId);
 
-#endif // UE_RIGIDPHYSICS_API_ENABLED
\ No newline at end of file
+#endif // UE_RIGIDPHYSICS_API_ENABLED
diff --git a/Engine/Source/Runtime/Experimental/RigidPhysics/Public/RigidPhysics/RigidGeometryCollection.h b/Engine/Source/Runtime/Experimental/RigidPhysics/Public/RigidPhysics/RigidGeometryCollection.h
--- a/Engine/Source/Runtime/Experimental/RigidPhysics/Public/RigidPhysics/RigidGeometryCollection.h
+++ b/Engine/Source/Runtime/Experimental/RigidPhysics/Public/RigidPhysics/RigidGeometryCollection.h
@@ -10,11 +10,12 @@
 #include "RigidPhysics/RigidGeometryCollectionHandle.h"
 #include "RigidPhysics/RigidObjectPtr.h"
 
+UE_EXPERIMENTAL(5.8, "The new Chaos API is experimental")
 namespace UE::Physics
 {
 	// Utility wrapper around IRigidGeometryCollection that adds Context and allows
 	// TRigidGeometryCollectionPtr to have pointer semantics without exposing IRigidGeometryCollection
-	template<typename ContextType>
+	template <typename ContextType>
 	class UE_INTERNAL TRigidGeometryCollectionPtrImpl
 	{
 	public:
@@ -23,6 +24,7 @@ namespace UE::Physics
 		using FHandle = FRigidGeometryCollectionHandle;
 		using FRigidScenePtr = TRigidScenePtr<FContext>;
 
+		TRigidGeometryCollectionPtrImpl() = default;
 		UE_INTERNAL TRigidGeometryCollectionPtrImpl(IRigidGeometryCollection* InGeometryCollection)
 			: GeometryCollectionRaw(InGeometryCollection)
 		{
@@ -37,30 +39,16 @@ namespace UE::Physics
 			}
 		}
 
-		virtual ~TRigidGeometryCollectionPtrImpl()
-		{
-		}
+		~TRigidGeometryCollectionPtrImpl() = default;
 
-		UE_INTERNAL void Reset()
-		{
-			GeometryCollectionRaw = nullptr;
-		}
+		friend bool operator==(const TRigidGeometryCollectionPtrImpl&, const TRigidGeometryCollectionPtrImpl&) = default;
+		friend bool operator!=(const TRigidGeometryCollectionPtrImpl&, const TRigidGeometryCollectionPtrImpl&) = default;
 
-		UE_INTERNAL bool IsValid() const
+		bool IsValid() const
 		{
 			return (GeometryCollectionRaw != nullptr);
 		}
 
-		const FRigidTypeId& GetTypeId() const
-		{
-			return GeometryCollectionRaw->GetTypeId();
-		}
-
-		FRigidObjectId GetId() const
-		{
-			return GeometryCollectionRaw->GetId();
-		}
-
 		FRigidGeometryCollectionHandle GetHandle() const
 		{
 			if (GeometryCollectionRaw != nullptr)
@@ -91,15 +79,30 @@ namespace UE::Physics
 			return GeometryCollectionRaw->GetBody(InBodyId);
 		}
 
+		UE_INTERNAL void Reset()
+		{
+			GeometryCollectionRaw = nullptr;
+		}
+
+		UE_INTERNAL const FRigidTypeId& GetTypeId() const
+		{
+			return GeometryCollectionRaw->GetTypeId();
+		}
+
 		UE_INTERNAL IRigidGeometryCollection* Get() const
 		{
 			return GeometryCollectionRaw;
 		}
 
+		UE_INTERNAL TRigidBodyPtr<FContext> CreateBody(const FRigidDebugName& InName, ERigidMovementType InMovementType)
+		{
+			return GeometryCollectionRaw->CreateBody(InName, InMovementType);
+		}
+
 	private:
 		// Pointer to implementation, object depends on Context type
-		IRigidGeometryCollection* GeometryCollectionRaw;
+		IRigidGeometryCollection* GeometryCollectionRaw = nullptr;
 	};
-}
+} // namespace UE::Physics
 
-#endif // UE_RIGIDPHYSICS_API_ENABLED
\ No newline at end of file
+#endif // UE_RIGIDPHYSICS_API_ENABLED
diff --git a/Engine/Source/Runtime/Experimental/RigidPhysics/Public/RigidPhysics/RigidGeometryCollectionHandle.h b/Engine/Source/Runtime/Experimental/RigidPhysics/Public/RigidPhysics/RigidGeometryCollectionHandle.h
--- a/Engine/Source/Runtime/Experimental/RigidPhysics/Public/RigidPhysics/RigidGeometryCollectionHandle.h
+++ b/Engine/Source/Runtime/Experimental/RigidPhysics/Public/RigidPhysics/RigidGeometryCollectionHandle.h
@@ -8,6 +8,7 @@
 #include "RigidPhysics/RigidBodyContainerHandle.h"
 #include "RigidPhysics/RigidObjectPtr.h"
 
+UE_EXPERIMENTAL(5.8, "The new Chaos API is experimental")
 namespace UE::Physics
 {
 	class UE_INTERNAL IRigidGeometryCollection;
@@ -41,7 +42,7 @@ namespace UE::Physics
 		friend bool operator!=(const FRigidBodyContainerHandle&, const FRigidBodyContainerHandle&) = default;
 
 		// Convert the handle to a pointer in the specified context
-		template<typename ContextType>
+		template <typename ContextType>
 		TRigidGeometryCollectionPtr<ContextType> Pin(const ContextType& InContext) const
 		{
 			if (InContext.IsValid())
@@ -54,7 +55,7 @@ namespace UE::Physics
 
 		UE_INTERNAL RIGIDPHYSICS_API void Reset();
 		UE_INTERNAL RIGIDPHYSICS_API const FRigidObjectId& GetId() const;
-		UE_INTERNAL RIGIDPHYSICS_API const FRigidSceneHandle GetSceneHandle() const;
+		UE_INTERNAL RIGIDPHYSICS_API const FRigidSceneId& GetSceneId() const;
 
 		UE_INTERNAL RIGIDPHYSICS_API FRigidBodyContainerHandle AsBodyContainerHandle() const;
 		UE_INTERNAL RIGIDPHYSICS_API operator FRigidBodyContainerHandle() const;
@@ -66,4 +67,4 @@ namespace UE::Physics
 	};
 } // namespace UE::Physics
 
-#endif // UE_RIGIDPHYSICS_API_ENABLED
\ No newline at end of file
+#endif // UE_RIGIDPHYSICS_API_ENABLED
diff --git a/Engine/Source/Runtime/Experimental/RigidPhysics/Public/RigidPhysics/RigidLockable.h b/Engine/Source/Runtime/Experimental/RigidPhysics/Public/RigidPhysics/RigidLockable.h
--- a/Engine/Source/Runtime/Experimental/RigidPhysics/Public/RigidPhysics/RigidLockable.h
+++ b/Engine/Source/Runtime/Experimental/RigidPhysics/Public/RigidPhysics/RigidLockable.h
@@ -6,6 +6,7 @@
 
 #if UE_RIGIDPHYSICS_API_ENABLED
 
+UE_EXPERIMENTAL(5.8, "The new Chaos API is experimental")
 namespace UE::Physics
 {
 	class UE_INTERNAL IRigidLockable : public IRigidTyped
diff --git a/Engine/Source/Runtime/Experimental/RigidPhysics/Public/RigidPhysics/RigidLog.h b/Engine/Source/Runtime/Experimental/RigidPhysics/Public/RigidPhysics/RigidLog.h
--- a/Engine/Source/Runtime/Experimental/RigidPhysics/Public/RigidPhysics/RigidLog.h
+++ b/Engine/Source/Runtime/Experimental/RigidPhysics/Public/RigidPhysics/RigidLog.h
@@ -11,15 +11,15 @@ RIGIDPHYSICS_API DECLARE_LOG_CATEGORY_EXTERN(LogRigidPhysics, Log, All);
 
 namespace UE::Physics::Private
 {
-	template<typename FuncNameType, typename TypeNameType>
+	template <typename FuncNameType, typename TypeNameType>
 	UE_INTERNAL inline void LogUnsupportedAPI(const FuncNameType& FuncName, const TypeNameType& TypeName)
 	{
 		// This warning means a function is being called in a context where it is not supported.
 		// E.g., Calling SetMass from the physics thread.
 		UE_LOGFMT(LogRigidPhysics, Warning, "{0} not supported on {1}", FuncName, TypeName);
 	}
-}
+} // namespace UE::Physics::Private
 
 #define RIGIDPHYSICS_API_UNSUPPORTED() UE::Physics::Private::LogUnsupportedAPI(__FUNCTION__, GetTypeName());
 
-#endif // UE_RIGIDPHYSICS_API_ENABLED
\ No newline at end of file
+#endif // UE_RIGIDPHYSICS_API_ENABLED
diff --git a/Engine/Source/Runtime/Experimental/RigidPhysics/Public/RigidPhysics/RigidModifier.h b/Engine/Source/Runtime/Experimental/RigidPhysics/Public/RigidPhysics/RigidModifier.h
--- a/Engine/Source/Runtime/Experimental/RigidPhysics/Public/RigidPhysics/RigidModifier.h
+++ b/Engine/Source/Runtime/Experimental/RigidPhysics/Public/RigidPhysics/RigidModifier.h
@@ -7,18 +7,14 @@
 
 #if UE_RIGIDPHYSICS_API_ENABLED
 
+UE_EXPERIMENTAL(5.8, "The new Chaos API is experimental")
 namespace UE::Physics
 {
 	class IRigidSceneModifier
 	{
 	public:
-		IRigidSceneModifier()
-		{
-		}
-
-		virtual ~IRigidSceneModifier()
-		{
-		}
+		IRigidSceneModifier() = default;
+		virtual ~IRigidSceneModifier() = default;
 
 		// Called from the game context before physics is ticked for this frame. Can be used to
 		// update inputs to the objects used in the sim context.
@@ -44,7 +40,6 @@ namespace UE::Physics
 		{
 		}
 	};
+} // namespace UE::Physics
 
-}
-
-#endif // UE_RIGIDPHYSICS_API_ENABLED
\ No newline at end of file
+#endif // UE_RIGIDPHYSICS_API_ENABLED
diff --git a/Engine/Source/Runtime/Experimental/RigidPhysics/Public/RigidPhysics/RigidObjectId.h b/Engine/Source/Runtime/Experimental/RigidPhysics/Public/RigidPhysics/RigidObjectId.h
--- a/Engine/Source/Runtime/Experimental/RigidPhysics/Public/RigidPhysics/RigidObjectId.h
+++ b/Engine/Source/Runtime/Experimental/RigidPhysics/Public/RigidPhysics/RigidObjectId.h
@@ -5,6 +5,7 @@
 
 #if UE_RIGIDPHYSICS_API_ENABLED
 
+UE_EXPERIMENTAL(5.8, "The new Chaos API is experimental")
 namespace UE::Physics
 {
 	// An ID for an object held in an array that is guaranteed to be unique and 
@@ -31,6 +32,18 @@ namespace UE::Physics
 			return (ArrayIndex >= 0);
 		}
 
+		friend bool operator==(const FRigidObjectId& L, const FRigidObjectId& R) = default;
+
+		friend bool operator<(const FRigidObjectId& L, const FRigidObjectId& R)
+		{
+			return L.ToUInt64() < R.ToUInt64();
+		}
+
+		friend uint32 GetTypeHash(const FRigidObjectId& InId)
+		{
+			return GetTypeHash(InId.ToUInt64());
+		}
+
 		UE_INTERNAL int32 GetArrayIndex() const
 		{
 			return ArrayIndex;
@@ -47,18 +60,6 @@ namespace UE::Physics
 			Epoch = INDEX_NONE;
 		}
 
-		friend bool operator==(const FRigidObjectId& L, const FRigidObjectId& R) = default;
-
-		friend bool operator<(const FRigidObjectId& L, const FRigidObjectId& R)
-		{
-			return L.ToUInt64() < R.ToUInt64();
-		}
-
-		friend uint32 GetTypeHash(const FRigidObjectId& InId)
-		{
-			return GetTypeHash(InId.ToUInt64());
-		}
-
 	protected:
 		uint64 ToUInt64() const
 		{
@@ -71,9 +72,9 @@ namespace UE::Physics
 			ArrayIndex = (int32)(InInt & 0xFFFFFFFF);
 		}
 
-		int32 ArrayIndex;
-		int32 Epoch;
+		int32 ArrayIndex = INDEX_NONE;
+		int32 Epoch = INDEX_NONE;
 	};
-}
+} // namespace UE::Physics
 
 #endif // UE_RIGIDPHYSICS_API_ENABLED
diff --git a/Engine/Source/Runtime/Experimental/RigidPhysics/Public/RigidPhysics/RigidObjectPtr.h b/Engine/Source/Runtime/Experimental/RigidPhysics/Public/RigidPhysics/RigidObjectPtr.h
--- a/Engine/Source/Runtime/Experimental/RigidPhysics/Public/RigidPhysics/RigidObjectPtr.h
+++ b/Engine/Source/Runtime/Experimental/RigidPhysics/Public/RigidPhysics/RigidObjectPtr.h
@@ -7,6 +7,7 @@
 
 #include "Templates/Requires.h"
 
+UE_EXPERIMENTAL(5.8, "The new Chaos API is experimental")
 namespace UE::Physics
 {
 
@@ -25,7 +26,7 @@ namespace UE::Physics
 	// after the underlying object has been destroyed. If the underlying object may have been destroyed, 
 	// you should re-Pin the Handle to obtain a new pointer.
 	//
-	template<typename PointerImplType>
+	template <typename PointerImplType>
 	class TRigidObjectPtr
 	{
 	public:
@@ -36,11 +37,7 @@ namespace UE::Physics
 
 		UE_NONCOPYABLE(TRigidObjectPtr);
 
-		TRigidObjectPtr()
-			: Impl(nullptr)
-		{
-		}
-
+		TRigidObjectPtr() = default;
 		UE_INTERNAL TRigidObjectPtr(FInterface* InObject)
 			: Impl(InObject)
 		{
@@ -52,15 +49,13 @@ namespace UE::Physics
 		template<
 			typename OtherPointerImplType
 			UE_REQUIRES(std::is_base_of_v<typename OtherPointerImplType::FInterface, FInterface>&& std::is_same_v<typename OtherPointerImplType::FContext, FContext>)
-			>
+		>
 		UE_INTERNAL TRigidObjectPtr(const TRigidObjectPtr<OtherPointerImplType>& InPtr)
 			: Impl(InPtr.Get())
 		{
 		}
 
-		~TRigidObjectPtr()
-		{
-		}
+		~TRigidObjectPtr() = default;
 
 		// Reset the pointer (it will now be Invalid)
 		void Reset()
@@ -130,8 +125,8 @@ namespace UE::Physics
 		// i.e., this wrapper cannot change, but the pointee is mutabale
 		// if the wrapper was created from a mutable context. This avoids
 		// an undefined behaviour compiler warning.
-		mutable FPointerImpl Impl;
+		mutable FPointerImpl Impl = nullptr;
 	};
-}
+} // namespace UE::Physics
 
 #endif // UE_RIGIDPHYSICS_API_ENABLED
diff --git a/Engine/Source/Runtime/Experimental/RigidPhysics/Public/RigidPhysics/RigidObjectRegistry.h b/Engine/Source/Runtime/Experimental/RigidPhysics/Public/RigidPhysics/RigidObjectRegistry.h
--- a/Engine/Source/Runtime/Experimental/RigidPhysics/Public/RigidPhysics/RigidObjectRegistry.h
+++ b/Engine/Source/Runtime/Experimental/RigidPhysics/Public/RigidPhysics/RigidObjectRegistry.h
@@ -9,10 +9,11 @@
 #include "RigidPhysics/RigidLog.h"
 #include "RigidPhysics/RigidObjectId.h"
 
+UE_EXPERIMENTAL(5.8, "The new Chaos API is experimental")
 namespace UE::Physics
 {
 	// Info held in the TRigidObjectRegistry indexed by ObjectId
-	template<typename ObjectPtrType>
+	template <typename ObjectPtrType>
 	class UE_INTERNAL TRigidObjectRegistryItem
 	{
 	public:
@@ -52,7 +53,7 @@ namespace UE::Physics
 	// wrap its Epoch. We can mitigate this by removing those indices from the free list. This will leave unusable
 	// gaps in the array but would allow for infinite growth at the cost of some wasted memory.
 	//
-	template<typename ObjectType>
+	template <typename ObjectType>
 	class UE_INTERNAL TRigidObjectRegistry
 	{
 	public:
@@ -68,9 +69,7 @@ namespace UE::Physics
 		{
 		}
 
-		UE_INTERNAL ~TRigidObjectRegistry()
-		{
-		}
+		UE_INTERNAL ~TRigidObjectRegistry() = default;
 
 		// The number of objects that have been added (and not yet removed)
 		UE_INTERNAL int32 GetNumObjects() const
@@ -245,6 +244,6 @@ namespace UE::Physics
 		// Use to trap modifications while in the Visitor loop
 		bool bIsVisiting;
 	};
-}
+} // namespace UE::Physics
 
-#endif // UE_RIGIDPHYSICS_API_ENABLED
\ No newline at end of file
+#endif // UE_RIGIDPHYSICS_API_ENABLED
diff --git a/Engine/Source/Runtime/Experimental/RigidPhysics/Public/RigidPhysics/RigidScene.h b/Engine/Source/Runtime/Experimental/RigidPhysics/Public/RigidPhysics/RigidScene.h
--- a/Engine/Source/Runtime/Experimental/RigidPhysics/Public/RigidPhysics/RigidScene.h
+++ b/Engine/Source/Runtime/Experimental/RigidPhysics/Public/RigidPhysics/RigidScene.h
@@ -14,11 +14,12 @@
 #include "RigidPhysics/RigidTyped.h"
 #include "Templates/SharedPointer.h"
 
+UE_EXPERIMENTAL(5.8, "The new Chaos API is experimental")
 namespace UE::Physics
 {
 	// A wrapper around IRigidScene.
 	// Exists only to give pointer semantics to TRigidScenePtr.
-	template<typename ContextType>
+	template <typename ContextType>
 	class UE_INTERNAL TRigidScenePtrImpl
 	{
 	public:
@@ -26,24 +27,12 @@ namespace UE::Physics
 		using FInterface = IRigidScene;
 		using FHandle = FRigidSceneHandle;
 
-		UE_INTERNAL TRigidScenePtrImpl()
-			: SceneRaw(nullptr)
-		{
-		}
-
+		UE_INTERNAL TRigidScenePtrImpl() = default;
 		UE_INTERNAL TRigidScenePtrImpl(IRigidScene* InScene)
 			: SceneRaw(InScene)
 		{
 		}
-
-		~TRigidScenePtrImpl()
-		{
-		}
-
-		void Reset()
-		{
-			SceneRaw = nullptr;
-		}
+		~TRigidScenePtrImpl() = default;
 
 		bool IsValid() const
 		{
@@ -99,7 +88,7 @@ namespace UE::Physics
 			SceneRaw->DestroyGeometryCollection(InGC.Get());
 			InGC.Reset();
 		}
-		
+
 		TJointConstraint6DOFPtr<FContext> CreateJointConstraint6DOF()
 		{
 			return SceneRaw->CreateJointConstraint6DOF();
@@ -129,14 +118,19 @@ namespace UE::Physics
 			SceneRaw->Tick(Dt);
 		}
 
+		UE_INTERNAL void Reset()
+		{
+			SceneRaw = nullptr;
+		}
+
 		UE_INTERNAL IRigidScene* Get() const
 		{
 			return SceneRaw;
 		}
 
 	private:
-		IRigidScene* SceneRaw;
+		IRigidScene* SceneRaw = nullptr;
 	};
-}
+} // namespace UE::Physics
 
-#endif // UE_RIGIDPHYSICS_API_ENABLED
\ No newline at end of file
+#endif // UE_RIGIDPHYSICS_API_ENABLED
diff --git a/Engine/Source/Runtime/Experimental/RigidPhysics/Public/RigidPhysics/RigidSceneHandle.h b/Engine/Source/Runtime/Experimental/RigidPhysics/Public/RigidPhysics/RigidSceneHandle.h
--- a/Engine/Source/Runtime/Experimental/RigidPhysics/Public/RigidPhysics/RigidSceneHandle.h
+++ b/Engine/Source/Runtime/Experimental/RigidPhysics/Public/RigidPhysics/RigidSceneHandle.h
@@ -7,6 +7,7 @@
 
 #include "RigidPhysics/RigidSceneId.h"
 
+UE_EXPERIMENTAL(5.8, "The new Chaos API is experimental")
 namespace UE::Physics
 {
 	class FRigidScene;
diff --git a/Engine/Source/Runtime/Experimental/RigidPhysics/Public/RigidPhysics/RigidSceneSettings.h b/Engine/Source/Runtime/Experimental/RigidPhysics/Public/RigidPhysics/RigidSceneSettings.h
--- a/Engine/Source/Runtime/Experimental/RigidPhysics/Public/RigidPhysics/RigidSceneSettings.h
+++ b/Engine/Source/Runtime/Experimental/RigidPhysics/Public/RigidPhysics/RigidSceneSettings.h
@@ -7,6 +7,7 @@
 
 #include "RigidPhysics/RigidTyped.h"
 
+UE_EXPERIMENTAL(5.8, "The new Chaos API is experimental")
 namespace UE::Physics
 {
 	class IRigidSceneSettings : public IRigidTyped
diff --git a/Engine/Source/Runtime/Experimental/RigidPhysics/Public/RigidPhysics/RigidTyped.h b/Engine/Source/Runtime/Experimental/RigidPhysics/Public/RigidPhysics/RigidTyped.h
--- a/Engine/Source/Runtime/Experimental/RigidPhysics/Public/RigidPhysics/RigidTyped.h
+++ b/Engine/Source/Runtime/Experimental/RigidPhysics/Public/RigidPhysics/RigidTyped.h
@@ -7,6 +7,7 @@
 
 #include "Misc/Crc.h"
 
+UE_EXPERIMENTAL(5.8, "The new Chaos API is experimental")
 namespace UE::Physics
 {
 	// A type ID for a rigid object to support downcasting and type checking.
@@ -65,16 +66,16 @@ namespace UE::Physics
 		}
 
 	private:
-		uint64 TypeIdValue;
-		const FRigidTypeId* BaseTypeId;
+		uint64 TypeIdValue = 0;
+		const FRigidTypeId* BaseTypeId = nullptr;
 		const FString TypeName;
 	};
 
 	class IRigidTyped
 	{
 	public:
-		RIGIDPHYSICS_API IRigidTyped();
-		RIGIDPHYSICS_API virtual ~IRigidTyped();
+		IRigidTyped() = default;
+		virtual ~IRigidTyped() = default;
 
 		// Get the TypeId of this class
 		RIGIDPHYSICS_API static const FRigidTypeId& GetStaticTypeId();
@@ -88,14 +89,14 @@ namespace UE::Physics
 		}
 
 		// Is T the same class or a base class of this object?
-		template<typename T>
+		template <typename T>
 		bool IsA() const
 		{
 			return GetTypeId().IsCastableTo(T::GetStaticTypeId());
 		}
 
 		// Try to cast this object to type T. Returns null if T is not a base class of this object.
-		template<typename T>
+		template <typename T>
 		T* AsA()
 		{
 			if (IsA<T>())
@@ -104,9 +105,9 @@ namespace UE::Physics
 			}
 			return nullptr;
 		}
-		
+
 		// Try to cast this object to type T. Returns null if T is not a base class of this object.
-		template<typename T>
+		template <typename T>
 		const T* AsA() const
 		{
 			if (IsA<T>())
@@ -117,7 +118,7 @@ namespace UE::Physics
 		}
 
 		// Try to cast this object to type T. Assert T is a base class of this object.
-		template<typename T>
+		template <typename T>
 		T* AsAChecked()
 		{
 			check(IsA<T>());
@@ -125,7 +126,7 @@ namespace UE::Physics
 		}
 
 		// Try to cast this object to type T. Assert T is a base class of this object.
-		template<typename T>
+		template <typename T>
 		const T* AsAChecked() const
 		{
 			check(IsA<T>());
@@ -133,7 +134,7 @@ namespace UE::Physics
 		}
 	};
 
-}
+} // namespace UE::Physics
 
 #define UE_RIGIDPHYSICS_RIGIDTYPED_DECL(API, T) \
 	API static const UE::Physics::FRigidTypeId& GetStaticTypeId(); \
@@ -152,4 +153,4 @@ namespace UE::Physics
 		return StaticTypeId; \
 	}
 
-#endif
+#endif // UE_RIGIDPHYSICS_API_ENABLED
