fix NaNs in geometry collection part sizes for parts with negative volume (due to inverted faces and/or degenerate geometry)

#Synced-CL 51501523

diff --git a/Engine/Source/Runtime/Experimental/Chaos/Private/GeometryCollection/GeometryCollectionConvexUtility.cpp b/Engine/Source/Runtime/Experimental/Chaos/Private/GeometryCollection/GeometryCollectionConvexUtility.cpp
--- a/Engine/Source/Runtime/Experimental/Chaos/Private/GeometryCollection/GeometryCollectionConvexUtility.cpp
+++ b/Engine/Source/Runtime/Experimental/Chaos/Private/GeometryCollection/GeometryCollectionConvexUtility.cpp
@@ -2878,7 +2878,9 @@ void FGeometryCollectionConvexUtility::SetVolumeAttributes(FManagedArrayCollecti
 		int32 Bone = MeshFacade.TransformIndexAttribute[GeoIdx];
 		if (SimulationType[Bone] == FGeometryCollection::ESimulationTypes::FST_Rigid)
 		{
-			GeoVolumes[GeoIdx] = (float)ComputeGeometryVolume(Collection, GeoIdx, Transforms[Bone], 1.0);
+			float ComputedVolume = (float)ComputeGeometryVolume(Collection, GeoIdx, Transforms[Bone], 1.0);
+			// if a mesh has inverted orientation the volume will be negative, but we take the abs() here so it's more usable as a size metric
+			GeoVolumes[GeoIdx] = FMath::Abs(ComputedVolume);
 		}
 	});
 
