FIx missing error code in OnDemandHttpLogs

#Synced-CL 51330535

diff --git a/Engine/Source/Runtime/Experimental/IoStore/OnDemand/Private/OnDemandHttpClient.cpp b/Engine/Source/Runtime/Experimental/IoStore/OnDemand/Private/OnDemandHttpClient.cpp
--- a/Engine/Source/Runtime/Experimental/IoStore/OnDemand/Private/OnDemandHttpClient.cpp
+++ b/Engine/Source/Runtime/Experimental/IoStore/OnDemand/Private/OnDemandHttpClient.cpp
@@ -861,7 +861,7 @@ void FMultiEndpointHttpClient::Log(const FRequest& Request, const HTTP::FTicketS
 	TStringBuilder<256> Reason;
 	if (TicketStatus.GetId() == FTicketStatus::EId::Error)
 	{
-		Reason << TEXT(" ") << TicketStatus.GetError().Reason;
+		Reason << TEXT(" ") << TicketStatus.GetError().Reason << TEXT(" (") << TicketStatus.GetError().Code << TEXT(")");
 	}
 	else if (TicketStatus.GetId() == FTicketStatus::EId::Cancelled)
 	{
