iax: option to use system CA store (off by default, Windows only).

#rnx

#Synced-CL 51536473

diff --git a/Engine/Source/Developer/IoStore/OnDemandUtilities/Private/Purl.cpp b/Engine/Source/Developer/IoStore/OnDemandUtilities/Private/Purl.cpp
--- a/Engine/Source/Developer/IoStore/OnDemandUtilities/Private/Purl.cpp
+++ b/Engine/Source/Developer/IoStore/OnDemandUtilities/Private/Purl.cpp
@@ -23,10 +23,21 @@ namespace UE::IoStore::Tool
 {
 
 ////////////////////////////////////////////////////////////////////////////////
-static void LoadCaCerts()
+static void LoadCaCerts(bool bSystemStore)
 {
 	using namespace UE::IoStore::HTTP;
 
+	if (bSystemStore)
+	{
+		FMemoryView NoData;
+		FCertRoots CaRoots(NoData);
+		if (CaRoots.IsValid())
+		{
+			FCertRoots::SetDefault(MoveTemp(CaRoots));
+			return;
+		}
+	}
+
 	IFileManager& Ifm = IFileManager::Get();
 	FString PemPath = FPaths::EngineContentDir() / TEXT("Certificates/ThirdParty/cacert.pem");
 	FArchive* Reader = Ifm.CreateFileReader(*PemPath);
@@ -56,7 +67,10 @@ static int32 PurlCommandEntry(const FContext& Context)
 
 	using namespace UE::IoStore::HTTP;
 
-	LoadCaCerts();
+	{
+		bool bSystemStore = Context.Get<bool>(TEXT("-SystemCA"));
+		LoadCaCerts(bSystemStore);
+	}
 
 	FStringView Url = Context.Get<FStringView>(TEXT("Url"));
 	auto AnsiUrl = StringCast<ANSICHAR>(Url.GetData(), Url.Len());
@@ -163,6 +177,7 @@ static FCommand PurlCommand(
 		TArgument<FStringView>(TEXT("Url"), TEXT("Url to download")),
 		TArgument<FStringView>(TEXT("-Method"), TEXT("Request method")),
 		TArgument<bool>(TEXT("-Redirect"), TEXT("Follow 30x redirects")),
+		TArgument<bool>(TEXT("-SystemCA"), TEXT("Load certificates from system CA store")),
 		TArgument<bool>(TEXT("-HttpTwo"), TEXT("Use HTTP/2")),
 		TArgument<int32>(TEXT("-BufSize"), TEXT("Recv buffer size in KiB")),
 		TArgument<uint32>(TEXT("-Throttle"), TEXT("Bandwidth throttle in KiB/s")),
diff --git a/Engine/Source/Runtime/Experimental/IoStore/HttpClient/Private/Peer.inl b/Engine/Source/Runtime/Experimental/IoStore/HttpClient/Private/Peer.inl
--- a/Engine/Source/Runtime/Experimental/IoStore/HttpClient/Private/Peer.inl
+++ b/Engine/Source/Runtime/Experimental/IoStore/HttpClient/Private/Peer.inl
@@ -17,6 +17,13 @@
 #	include <openssl/engine.h>
 #	include <openssl/err.h>
 #	include <openssl/ssl.h>
+#	include <openssl/store.h>
+
+#	if PLATFORM_WINDOWS
+#		include "Windows/AllowWindowsPlatformTypes.h"
+#			include <wincrypt.h>
+#		include "Windows/HideWindowsPlatformTypes.h"
+#	endif
 #endif
 
 namespace UE::IoStore::HTTP
@@ -67,24 +74,49 @@ static void Ssl_ContextDestroy(UPTRINT Handle)
 }
 
 ////////////////////////////////////////////////////////////////////////////////
-static UPTRINT Ssl_ContextCreate(FMemoryView PemData)
+static uint32 Ssl_ContextProvision(SSL_CTX* Context)
 {
-	if (static bool InitOnce = false; !InitOnce)
+	uint32 NumAdded = 0;
+
+#if PLATFORM_WINDOWS
+	HANDLE StoreHandle = CertOpenSystemStoreW(0, L"ROOT");
+	for (PCCERT_CONTEXT CertContext = nullptr;;)
 	{
-		// While OpenSSL will lazily initialise itself, the defaults used will fail
-		// initialisation on some platforms. So we have a go here. We do not register
-		// anything for clean-up as we do not know if anyone else has done so.
-		uint64 InitOpts = OPENSSL_INIT_NO_ATEXIT;
-		OPENSSL_init_ssl(InitOpts, nullptr);
-		InitOnce = true;
+		CertContext = CertEnumCertificatesInStore(StoreHandle, CertContext);
+		if (CertContext == nullptr)
+		{
+			break;
+		}
+
+		if ((CertContext->dwCertEncodingType & X509_ASN_ENCODING) == 0)
+		{
+			continue;
+		}
+
+		const uint8* CertData = (uint8*)(CertContext->pbCertEncoded);
+		uint32 CertSize = CertContext->cbCertEncoded;
+		X509* FiveOhNine = d2i_X509(nullptr, &CertData, CertSize);
+		if (FiveOhNine == nullptr)
+		{
+			continue;
+		}
+
+		X509_STORE* Store = SSL_CTX_get_cert_store(Context);
+		int32 Result = X509_STORE_add_cert(Store, FiveOhNine);
+		NumAdded += (Result == 1);
+
+		X509_free(FiveOhNine);
 	}
 
-	auto* Method = TLS_client_method();
-	SSL_CTX* Context = SSL_CTX_new(Method);
-	checkf(Context != nullptr, TEXT("ERR_get_error() == %d"), ERR_get_error());
+	CertCloseStore(StoreHandle, 0);
+#endif // PLATFORM_WINDOWS
 
-	SSL_CTX_set_options(Context, SSL_OP_NO_SSLv2|SSL_OP_NO_SSLv3);
+	return NumAdded;
+}
 
+////////////////////////////////////////////////////////////////////////////////
+static uint32 Ssl_ContextProvision(SSL_CTX* Context, FMemoryView PemData)
+{
 	const void* Data = PemData.GetData();
 	uint32 Size = uint32(PemData.GetSize());
 	BIO* Bio = BIO_new_mem_buf(Data, Size);
@@ -106,6 +138,37 @@ static UPTRINT Ssl_ContextCreate(FMemoryView PemData)
 	}
 
 	BIO_free(Bio);
+	return NumAdded;
+}
+
+////////////////////////////////////////////////////////////////////////////////
+static UPTRINT Ssl_ContextCreate(FMemoryView PemData)
+{
+	if (static bool InitOnce = false; !InitOnce)
+	{
+		// While OpenSSL will lazily initialise itself, the defaults used will fail
+		// initialisation on some platforms. So we have a go here. We do not register
+		// anything for clean-up as we do not know if anyone else has done so.
+		uint64 InitOpts = OPENSSL_INIT_NO_ATEXIT;
+		OPENSSL_init_ssl(InitOpts, nullptr);
+		InitOnce = true;
+	}
+
+	auto* Method = TLS_client_method();
+	SSL_CTX* Context = SSL_CTX_new(Method);
+	checkf(Context != nullptr, TEXT("ERR_get_error() == %d"), ERR_get_error());
+
+	SSL_CTX_set_options(Context, SSL_OP_NO_SSLv2|SSL_OP_NO_SSLv3);
+
+	uint32 NumAdded;
+	if (!PemData.IsEmpty())
+	{
+		NumAdded = Ssl_ContextProvision(Context, PemData);
+	}
+	else
+	{
+		NumAdded = Ssl_ContextProvision(Context);
+	}
 
 	if (NumAdded == 0)
 	{
diff --git a/Engine/Source/Runtime/Experimental/IoStore/OnDemand/Private/OnDemandHttpClient.cpp b/Engine/Source/Runtime/Experimental/IoStore/OnDemand/Private/OnDemandHttpClient.cpp
--- a/Engine/Source/Runtime/Experimental/IoStore/OnDemand/Private/OnDemandHttpClient.cpp
+++ b/Engine/Source/Runtime/Experimental/IoStore/OnDemand/Private/OnDemandHttpClient.cpp
@@ -40,6 +40,16 @@ static FAutoConsoleVariableRef CVar_IaxAllowPartialContentRequests(
 	TEXT("Enable/disable the use of partial content http requests")
 );
 
+/**
+ * Use CA root certificates from the system store if available.
+ */
+bool GUseSystemCaStore = false;
+static FAutoConsoleVariableRef CVar_IaxUseSystemCaStore(
+	TEXT("iax.HttpUseSystemCaStore"),
+	GUseSystemCaStore,
+	TEXT("Load root certificates from the system store (if available)")
+);
+
 ///////////////////////////////////////////////////////////////////////////////
 
 FIoBuffer JoinIoBuffers(const FIoBuffer& LHS, const FIoBuffer& RHS)
@@ -104,6 +114,18 @@ FIoStatus LoadDefaultHttpCertificates(bool& bWasLoaded)
 				GConfig->GetBool(CertSection, CertKey, bVerifyPeer, GEngineIni);
 			}
 
+			if (GUseSystemCaStore)
+			{
+				FCertRoots CaRoots(FMemoryView{});
+				if (CaRoots.IsValid())
+				{
+					uint32 NumCerts = CaRoots.Num();
+					FCertRoots::SetDefault(MoveTemp(CaRoots));
+					UE_LOG(LogIoStoreOnDemand, Display, TEXT("Loaded %u certificates from system"), NumCerts);
+					return;
+				}
+			}
+
 			// Open the certs file
 			IFileManager& Ifm = IFileManager::Get();
 			const FString PemPath = FPaths::EngineContentDir() / TEXT("Certificates/ThirdParty/cacert.pem");
