Add support for pso precache fallback or proxy creation skip when material psos are still compiling for geometry collection component (nanite & non nanite code path)

#Synced-CL 51478552

diff --git a/Engine/Source/Runtime/Experimental/GeometryCollectionEngine/Private/GeometryCollection/GeometryCollectionComponent.cpp b/Engine/Source/Runtime/Experimental/GeometryCollectionEngine/Private/GeometryCollection/GeometryCollectionComponent.cpp
--- a/Engine/Source/Runtime/Experimental/GeometryCollectionEngine/Private/GeometryCollection/GeometryCollectionComponent.cpp
+++ b/Engine/Source/Runtime/Experimental/GeometryCollectionEngine/Private/GeometryCollection/GeometryCollectionComponent.cpp
@@ -1144,7 +1144,13 @@ FPrimitiveSceneProxy* UGeometryCollectionComponent::CreateSceneProxy()
 	FPrimitiveSceneProxy* LocalSceneProxy = nullptr;
 
 	if (RestCollection && !IsUsingCustomRenderer())
-	{
+	{		
+		if (CheckPSOPrecachingAndBoostPriority() && GetPSOPrecacheProxyCreationStrategy() == EPSOPrecacheProxyCreationStrategy::DelayUntilPSOPrecached)
+		{
+			UE_LOG(LogGeometryCollectionComponent, Verbose, TEXT("Skipping CreateSceneProxy for UGeometryCollectionComponent %s (UGeometryCollectionComponent PSOs are still compiling)"), *GetFullName());
+			return nullptr;
+		}
+
 		if (UseNanite(GetScene()->GetShaderPlatform()) &&
 			RestCollection->EnableNanite &&
 			RestCollection->HasNaniteData() &&
@@ -5125,6 +5131,8 @@ void UGeometryCollectionComponent::SetRestCollection(const UGeometryCollection*
 		{
 			ApplyAssetDefaults();
 		}
+		
+		PrecachePSOs();
 	}
 }
 
diff --git a/Engine/Source/Runtime/Experimental/GeometryCollectionEngine/Private/GeometryCollection/GeometryCollectionSceneProxy.cpp b/Engine/Source/Runtime/Experimental/GeometryCollectionEngine/Private/GeometryCollection/GeometryCollectionSceneProxy.cpp
--- a/Engine/Source/Runtime/Experimental/GeometryCollectionEngine/Private/GeometryCollection/GeometryCollectionSceneProxy.cpp
+++ b/Engine/Source/Runtime/Experimental/GeometryCollectionEngine/Private/GeometryCollection/GeometryCollectionSceneProxy.cpp
@@ -24,6 +24,7 @@
 #include "RenderGraphBuilder.h"
 #include "MeshPaintVisualize.h"
 #include "SceneView.h"
+#include "PSOPrecacheSettings.h"
 
 #if RHI_RAYTRACING
 #include "RayTracingInstance.h"
@@ -185,6 +186,10 @@ FGeometryCollectionSceneProxyBase::FGeometryCollectionSceneProxyBase(UGeometryCo
 			{
 				Materials[MaterialIndex] = UMaterial::GetDefaultMaterial(MD_Surface);
 			}
+			else if (Component->GetPSOPrecacheComponentData().UsePSOPrecacheFallbackMaterial())
+			{
+				Materials[MaterialIndex] = UPSOPrecacheSettingsManager::GetFallbackMaterial() ? UPSOPrecacheSettingsManager::GetFallbackMaterial() : UMaterial::GetDefaultMaterial(MD_Surface);
+			}
 		}
 	}
 
@@ -1362,6 +1367,10 @@ FNaniteGeometryCollectionSceneProxy::FNaniteGeometryCollectionSceneProxy(UGeomet
 		{
 			// force default material 
 			MaterialInterface = UMaterial::GetDefaultMaterial(MD_Surface);
+		}		
+		else if (Component->GetPSOPrecacheComponentData().UsePSOPrecacheFallbackMaterial())
+		{
+			MaterialInterface = UPSOPrecacheSettingsManager::GetFallbackMaterial() ? UPSOPrecacheSettingsManager::GetFallbackMaterial() : UMaterial::GetDefaultMaterial(MD_Surface);
 		}
 
 		// Should never be null here
