Chaos: Fix soft angular constraints applying velocity like hard constraints

#Synced-CL 51484018

diff --git a/Engine/Source/Runtime/Experimental/Chaos/Private/Chaos/Joint/PBDJointCachedSolverGaussSeidel.cpp b/Engine/Source/Runtime/Experimental/Chaos/Private/Chaos/Joint/PBDJointCachedSolverGaussSeidel.cpp
--- a/Engine/Source/Runtime/Experimental/Chaos/Private/Chaos/Joint/PBDJointCachedSolverGaussSeidel.cpp
+++ b/Engine/Source/Runtime/Experimental/Chaos/Private/Chaos/Joint/PBDJointCachedSolverGaussSeidel.cpp
@@ -1680,15 +1680,15 @@ void FPBDJointCachedSolver::ApplyRotationSoftConstraintsSimd(
 
 void FPBDJointCachedSolver::ApplyAngularVelocityConstraints()
 {
-	if (RotationConstraints.bUseSimd)
-	{
-		ApplyAngularVelocityConstraintSimd();
-	}
-	else
+	// SIMD is only enabled when all 3 axes are limited and soft (see InitRotationConstraints).
+	// Soft constraints should not apply velocity corrections. They are handled through
+	// the damping term in the position solve phase instead. So call to ApplyAngularVelocityConstraintSimd is removed. 
+
+	if (!RotationConstraints.bUseSimd)
 	{
 		for (int32 ConstraintIndex = 0; ConstraintIndex < 3; ++ConstraintIndex)
 		{
-			if (RotationConstraints.GetValidDatas(ConstraintIndex))
+			if (RotationConstraints.GetValidDatas(ConstraintIndex) && !RotationConstraints.GetSoftLimit(ConstraintIndex))
 			{
 				ApplyAngularVelocityConstraint(ConstraintIndex);
 			}
