Chaos: Fix joint linear/angular drive position target reset when plasticity enabled on initialization.

#Synced-CL 51484018

diff --git a/Engine/Source/Runtime/Experimental/Chaos/Private/Chaos/PBDJointConstraints.cpp b/Engine/Source/Runtime/Experimental/Chaos/Private/Chaos/PBDJointConstraints.cpp
--- a/Engine/Source/Runtime/Experimental/Chaos/Private/Chaos/PBDJointConstraints.cpp
+++ b/Engine/Source/Runtime/Experimental/Chaos/Private/Chaos/PBDJointConstraints.cpp
@@ -298,6 +298,12 @@ namespace Chaos
 
 	void FPBDJointConstraintHandle::SetSettings(const FPBDJointSettings& InSettings)
 	{
+		// Track whether plasticity was already initialized before this call.
+		// We only want to restore the buffered targets for subsequent updates,
+		// not during the initial constraint setup.
+		const bool bWasLinearPlasticityInitialized = bLinearPlasticityInitialized;
+		const bool bWasAngularPlasticityInitialized = bAngularPlasticityInitialized;
+
 		// buffer the previous targets so plasticity can reuse them
 		FVec3 LinearTarget = GetSettings().LinearDrivePositionTarget;
 		FRotation3 AngularTarget = GetSettings().AngularDrivePositionTarget;
@@ -314,11 +320,11 @@ namespace Chaos
 
 
 		// transfer the previous targets when controlled by plasticity
-		if (bLinearPlasticityInitialized)
+		if (bWasLinearPlasticityInitialized)
 		{
 			ConcreteContainer()->SetLinearDrivePositionTarget(ConstraintIndex,LinearTarget);
 		}
-		if (bAngularPlasticityInitialized)
+		if (bWasAngularPlasticityInitialized)
 		{
 			ConcreteContainer()->SetAngularDrivePositionTarget(ConstraintIndex,AngularTarget);
 		}
