Chaos API: Adding constraints to some of the public API methods to only work on the game thread.
This is not exhaustive, and more meant to setup a pattern. Some of these should also eventually be valid, but do not currently work while others should probably never be valid.

#Synced-CL 51413247

diff --git a/Engine/Source/Runtime/Experimental/RigidPhysics/Public/RigidPhysics/JointConstraint6DOF.h b/Engine/Source/Runtime/Experimental/RigidPhysics/Public/RigidPhysics/JointConstraint6DOF.h
--- a/Engine/Source/Runtime/Experimental/RigidPhysics/Public/RigidPhysics/JointConstraint6DOF.h
+++ b/Engine/Source/Runtime/Experimental/RigidPhysics/Public/RigidPhysics/JointConstraint6DOF.h
@@ -67,12 +67,12 @@ namespace UE::Physics
 			ConstraintRaw->Activate();
 		}
 
-		void SetBodies(const TRigidBodyPtr<ContextType>& Body0, const TRigidBodyPtr<ContextType>& Body1)
+		void SetBodies(const TRigidBodyPtr<ContextType>& Body0, const TRigidBodyPtr<ContextType>& Body1) requires CIsGameContext<FContext>
 		{
 			ConstraintRaw->SetBodies(Body0.Get(), Body1.Get());
 		}
 
-		void SetJointTransforms(const FTransform3d& InTransform0, const FTransform3d& InTransform1)
+		void SetJointTransforms(const FTransform3d& InTransform0, const FTransform3d& InTransform1) requires CIsGameContext<FContext>
 		{
 			ConstraintRaw->SetJointTransforms(InTransform0, InTransform1);
 		}
diff --git a/Engine/Source/Runtime/Experimental/RigidPhysics/Public/RigidPhysics/RigidBody.h b/Engine/Source/Runtime/Experimental/RigidPhysics/Public/RigidPhysics/RigidBody.h
--- a/Engine/Source/Runtime/Experimental/RigidPhysics/Public/RigidPhysics/RigidBody.h
+++ b/Engine/Source/Runtime/Experimental/RigidPhysics/Public/RigidPhysics/RigidBody.h
@@ -68,12 +68,12 @@ namespace UE::Physics
 			return BodyRaw->IsActive();
 		}
 
-		void Activate()
+		void Activate() requires CIsGameContext<FContext>
 		{
 			BodyRaw->Activate();
 		}
 
-		void Deactivate()
+		void Deactivate() requires CIsGameContext<FContext>
 		{
 			BodyRaw->Deactivate();
 		}
@@ -103,17 +103,17 @@ namespace UE::Physics
 			return BodyRaw->GetShape(InShapeIndex);
 		}
 
-		TRigidShapeInstancePtr<FContext> CreateShape(const FRigidShapeInstanceSetup& InSetup)
+		TRigidShapeInstancePtr<FContext> CreateShape(const FRigidShapeInstanceSetup& InSetup) requires CIsGameContext<FContext>
 		{
 			return BodyRaw->CreateShape(InSetup);
 		}
 
-		UE_RIGIDPHYSICS_LEGACY_API TRigidShapeInstancePtr<FContext> CreateShape(TUniquePtr<Chaos::FPerShapeData>&& InShapeData)
+		UE_RIGIDPHYSICS_LEGACY_API TRigidShapeInstancePtr<FContext> CreateShape(TUniquePtr<Chaos::FPerShapeData>&& InShapeData) requires CIsGameContext<FContext>
 		{
 			return BodyRaw->CreateShape(MoveTemp(InShapeData));
 		}
 
-		void DestroyShape(const TRigidShapeInstancePtr<FContext>& InShape, bool bWakeTouching = true)
+		void DestroyShape(const TRigidShapeInstancePtr<FContext>& InShape, bool bWakeTouching = true) requires CIsGameContext<FContext>
 		{
 			BodyRaw->DestroyShape(InShape.Get(), bWakeTouching);
 		}
diff --git a/Engine/Source/Runtime/Experimental/RigidPhysics/Public/RigidPhysics/RigidFwd.h b/Engine/Source/Runtime/Experimental/RigidPhysics/Public/RigidPhysics/RigidFwd.h
--- a/Engine/Source/Runtime/Experimental/RigidPhysics/Public/RigidPhysics/RigidFwd.h
+++ b/Engine/Source/Runtime/Experimental/RigidPhysics/Public/RigidPhysics/RigidFwd.h
@@ -149,6 +149,9 @@ namespace UE::Physics
 	using FRigidContextGameRO = TRigidContext<ERigidContextType::Game, ERigidLockType::ReadOnly>;
 	using FRigidContextGameRW = TRigidContext<ERigidContextType::Game, ERigidLockType::ReadWrite>;
 	using FRigidContextSimRW = TRigidContext<ERigidContextType::Sim, ERigidLockType::ReadWrite>;
+
+	template <typename ContextType>
+	concept CIsGameContext = ContextType::ContextType == ERigidContextType::Game;
 } // namespace UE::Physics
 
 //
diff --git a/Engine/Source/Runtime/Experimental/RigidPhysics/Public/RigidPhysics/RigidScene.h b/Engine/Source/Runtime/Experimental/RigidPhysics/Public/RigidPhysics/RigidScene.h
--- a/Engine/Source/Runtime/Experimental/RigidPhysics/Public/RigidPhysics/RigidScene.h
+++ b/Engine/Source/Runtime/Experimental/RigidPhysics/Public/RigidPhysics/RigidScene.h
@@ -54,22 +54,22 @@ namespace UE::Physics
 			return SceneRaw->GetGeometryCollection(InGeometryCollectionId);
 		}
 
-		void RegisterModifier(IRigidSceneModifier* Modifier)
+		void RegisterModifier(IRigidSceneModifier* Modifier) requires CIsGameContext<FContext>
 		{
 			SceneRaw->RegisterModifier(Modifier);
 		}
 
-		void UnregisterModifier(IRigidSceneModifier* Modifier)
+		void UnregisterModifier(IRigidSceneModifier* Modifier) requires CIsGameContext<FContext>
 		{
 			SceneRaw->UnregisterModifier(Modifier);
 		}
 
-		TRigidBodyPtr<FContext> CreateBody(const FRigidDebugName& InName, ERigidMovementType InMovementType)
+		TRigidBodyPtr<FContext> CreateBody(const FRigidDebugName& InName, ERigidMovementType InMovementType) requires CIsGameContext<FContext>
 		{
 			return SceneRaw->CreateBody(InName, InMovementType);
 		}
 
-		void DestroyBody(TRigidBodyPtr<FContext>&& InBody)
+		void DestroyBody(TRigidBodyPtr<FContext>&& InBody) requires CIsGameContext<FContext>
 		{
 			if (InBody.IsValid())
 			{
@@ -78,23 +78,23 @@ namespace UE::Physics
 			}
 		}
 
-		TRigidGeometryCollectionPtr<FContext> CreateGeometryCollection(const FRigidDebugName& InName, const FSimulationParameters& InParameters)
+		TRigidGeometryCollectionPtr<FContext> CreateGeometryCollection(const FRigidDebugName& InName, const FSimulationParameters& InParameters) requires CIsGameContext<FContext>
 		{
 			return SceneRaw->CreateGeometryCollection(InName, InParameters);
 		}
 
-		void DestroyGeometryCollection(TRigidGeometryCollectionPtr<ContextType>&& InGC)
+		void DestroyGeometryCollection(TRigidGeometryCollectionPtr<ContextType>&& InGC) requires CIsGameContext<FContext>
 		{
 			SceneRaw->DestroyGeometryCollection(InGC.Get());
 			InGC.Reset();
 		}
 
-		TJointConstraint6DOFPtr<FContext> CreateJointConstraint6DOF()
+		TJointConstraint6DOFPtr<FContext> CreateJointConstraint6DOF() requires CIsGameContext<FContext>
 		{
 			return SceneRaw->CreateJointConstraint6DOF();
 		}
 
-		void DestroyJointConstraint(TJointConstraint6DOFPtr<FContext>&& InConstraint)
+		void DestroyJointConstraint(TJointConstraint6DOFPtr<FContext>&& InConstraint) requires CIsGameContext<FContext>
 		{
 			if (InConstraint.IsValid())
 			{
@@ -113,7 +113,7 @@ namespace UE::Physics
 			return SceneRaw->Query(QueryDesc, OutHits);
 		}
 
-		void Tick(float Dt)
+		void Tick(float Dt) requires CIsGameContext<FContext>
 		{
 			SceneRaw->Tick(Dt);
 		}
