Fix issue when using Legacy SKM cloth with Softs::FEvolution solver where the Kinematic Updates were reset when RefreshClothConfig is called (e.g., when updating properties such as Mass in the editor).

#Synced-CL 51549410

diff --git a/Engine/Plugins/ChaosCloth/Source/ChaosCloth/Private/ChaosCloth/ChaosClothingSimulationSolver.cpp b/Engine/Plugins/ChaosCloth/Source/ChaosCloth/Private/ChaosCloth/ChaosClothingSimulationSolver.cpp
--- a/Engine/Plugins/ChaosCloth/Source/ChaosCloth/Private/ChaosCloth/ChaosClothingSimulationSolver.cpp
+++ b/Engine/Plugins/ChaosCloth/Source/ChaosCloth/Private/ChaosCloth/ChaosClothingSimulationSolver.cpp
@@ -547,7 +547,8 @@ void FClothingSimulationSolver::Reset()
 {
 	if (Evolution)
 	{
-		Evolution->Reset();
+		constexpr bool bResetGlobalRules = false;
+		Evolution->Reset(bResetGlobalRules);
 		ClothsConstraints.Reset();
 	}
 	else
diff --git a/Engine/Source/Runtime/Experimental/Chaos/Private/Chaos/SoftsEvolution.cpp b/Engine/Source/Runtime/Experimental/Chaos/Private/Chaos/SoftsEvolution.cpp
--- a/Engine/Source/Runtime/Experimental/Chaos/Private/Chaos/SoftsEvolution.cpp
+++ b/Engine/Source/Runtime/Experimental/Chaos/Private/Chaos/SoftsEvolution.cpp
@@ -175,7 +175,7 @@ FEvolution::FEvolution(const FCollectionPropertyConstFacade& Properties)
 	Particles.AddArray(&ParticleDampings);
 }
 
-void FEvolution::Reset()
+void FEvolution::Reset(bool bResetGlobalRules)
 {
 	CollisionParticles.Resize(0);
 	Particles.Resize(0);
@@ -185,8 +185,11 @@ void FEvolution::Reset()
 
 	Groups.Reset();
 
-	KinematicUpdate.Reset();
-	CollisionKinematicUpdate.Reset();
+	if (bResetGlobalRules)
+	{
+		KinematicUpdate.Reset();
+		CollisionKinematicUpdate.Reset();
+	}
 }
 
 TArray<uint32> FEvolution::GetActiveGroupsArray() const
diff --git a/Engine/Source/Runtime/Experimental/Chaos/Public/Chaos/SoftsEvolution.h b/Engine/Source/Runtime/Experimental/Chaos/Public/Chaos/SoftsEvolution.h
--- a/Engine/Source/Runtime/Experimental/Chaos/Public/Chaos/SoftsEvolution.h
+++ b/Engine/Source/Runtime/Experimental/Chaos/Public/Chaos/SoftsEvolution.h
@@ -79,8 +79,9 @@ public:
 	CHAOS_API FEvolution(const FCollectionPropertyConstFacade& Properties);
 	~FEvolution() = default;
 
-	/** Reset/empty everything.*/
-	CHAOS_API void Reset();
+	/** Reset/empty everything.
+	 * @param bResetGlobalRules whether or not to reset global rules (e.g., KinematicUpdate) */
+	CHAOS_API void Reset(bool bResetGlobalRules = true);
 
 	/** Move forward in time */
 	CHAOS_API void AdvanceOneTimeStep(const FSolverReal Dt, const FSolverReal TimeDependentIterationMultiplier);
