IAX will once again run an initial connection test against the first hostgroup provided.

#rnx

- When the system used to resolve it's own hostgroup we would run an initial connection test against that hostgroup to make sure that it could be connected too.
- This allowed us to log vital connection errors early so that they would be preserved in the log header. It also allows us to provide enough info to the analytics system that we could accurately track how many users fail to connect to the hostgroups entirely.
- This change aims to restore the functionality in a limited fashion by checking the first hostgroup provided to the system. We will extend this to all hostgroups in a later work item.
- This feature can be disabled by setting the cvar 'ias.IaxHttpForceInitialConnectionCheck' to false, although this does occur very early in start up and likely is done before hotfixes are applied which is worth keeping in mind.
- When the feature is active we will ask theFHostGroupManager to run connection tests on it's hostgroups before FHttpDispatcher::Run can process anything.
- If the FHostGroupManager does not have any hostgroups then we will wait until one is.
-- This shouldn't really change anything as requests cannot be queued unless they have a valid hostgroup anyway.
-- This waiting loop will early exit if the process is shutdown before it completes.

#Synced-CL 51517438

diff --git a/Engine/Source/Runtime/Experimental/IoStore/OnDemand/Private/IasHostGroup.cpp b/Engine/Source/Runtime/Experimental/IoStore/OnDemand/Private/IasHostGroup.cpp
--- a/Engine/Source/Runtime/Experimental/IoStore/OnDemand/Private/IasHostGroup.cpp
+++ b/Engine/Source/Runtime/Experimental/IoStore/OnDemand/Private/IasHostGroup.cpp
@@ -244,9 +244,9 @@ bool FIASHostGroup::Disconnect()
 	return true;
 }
 
-FIASHostGroup::EReconnectionResult FIASHostGroup::AttemptReconnection(uint32 TimeoutMs, std::atomic_bool& CancellationToken)
+FIASHostGroup::EReconnectionResult FIASHostGroup::AttemptReconnection(bool bForce, uint32 TimeoutMs, std::atomic_bool& CancellationToken)
 {
-	if (!IsConnected())
+	if (!IsConnected() || bForce)
 	{
 		if (Impl->TestPath.IsEmpty())
 		{
@@ -388,6 +388,27 @@ FIASHostGroup FHostGroupManager::Find(FName Name)
 	}
 }
 
+bool FHostGroupManager::ResolveInitialConnection(uint32 TimeoutMs, std::atomic_bool& CancellationToken)
+{
+	UE::TReadScopeLock _(Mutex);
+	if (!HostGroups.IsEmpty())
+	{
+		for (FIASHostGroup& HostGroup : HostGroups)
+		{
+			if (HostGroup.AttemptReconnection(true, TimeoutMs, CancellationToken) == FIASHostGroup::EReconnectionResult::Reconnected)
+			{
+				FOnDemandIoBackendStats::Get()->OnHttpConnected();
+			}
+		}
+
+		return true;
+	}
+	else
+	{
+		return false;
+	}
+}
+
 void FHostGroupManager::ForEachHostGroup(TFunctionRef<void(const FIASHostGroup&)> Callback) const
 {
 	UE::TReadScopeLock _(Mutex);
@@ -406,7 +427,7 @@ void FHostGroupManager::Tick(uint32 TimeoutMs, std::atomic_bool& CancellationTok
 
 	for (FIASHostGroup& HostGroup : HostGroups)
 	{
-		if (HostGroup.AttemptReconnection(TimeoutMs, CancellationToken) == FIASHostGroup::EReconnectionResult::Reconnected)
+		if (HostGroup.AttemptReconnection(false, TimeoutMs, CancellationToken) == FIASHostGroup::EReconnectionResult::Reconnected)
 		{
 			FOnDemandIoBackendStats::Get()->OnHttpConnected(); // TODO: Try to avoid singleton access somehow
 		}
diff --git a/Engine/Source/Runtime/Experimental/IoStore/OnDemand/Private/IasHostGroup.h b/Engine/Source/Runtime/Experimental/IoStore/OnDemand/Private/IasHostGroup.h
--- a/Engine/Source/Runtime/Experimental/IoStore/OnDemand/Private/IasHostGroup.h
+++ b/Engine/Source/Runtime/Experimental/IoStore/OnDemand/Private/IasHostGroup.h
@@ -46,7 +46,7 @@ public:
 		FailedToConnect,
 	};
 
-	EReconnectionResult AttemptReconnection(uint32 TimeoutMs, std::atomic_bool& CancellationToken);
+	EReconnectionResult AttemptReconnection(bool bForce, uint32 TimeoutMs, std::atomic_bool& CancellationToken);
 
 	void OnSuccessfulResponse();
 	bool OnFailedResponse();
@@ -83,6 +83,9 @@ public:
 
 	FIASHostGroup Find(FName Name);
 
+	/** Returns true if we have tried to resolve at least one hostgroup and false if there are currently no hostgroups */
+	bool ResolveInitialConnection(uint32 TimeoutMs, std::atomic_bool& CancellationToken);
+
 	void ForEachHostGroup(TFunctionRef<void (const FIASHostGroup&)> Callback) const;
 
 	void Tick(uint32 TimeoutMs, std::atomic_bool& CancellationToken);
diff --git a/Engine/Source/Runtime/Experimental/IoStore/OnDemand/Private/OnDemandHttpIoDispacher.cpp b/Engine/Source/Runtime/Experimental/IoStore/OnDemand/Private/OnDemandHttpIoDispacher.cpp
--- a/Engine/Source/Runtime/Experimental/IoStore/OnDemand/Private/OnDemandHttpIoDispacher.cpp
+++ b/Engine/Source/Runtime/Experimental/IoStore/OnDemand/Private/OnDemandHttpIoDispacher.cpp
@@ -342,6 +342,14 @@ static FAutoConsoleVariableRef CVar_HttpCacheEnabled(
 	TEXT("")
 );
 
+/** Note that this will most likely be used before config hotfixs are applied */
+bool GIaxHttpForceInitialConnectionCheck = true;
+static FAutoConsoleVariableRef CVar_IaxHttpForceInitialConnectionCheck(
+	TEXT("ias.IaxHttpForceInitialConnectionCheck"),
+	GIaxHttpForceInitialConnectionCheck,
+	TEXT("When true the IAX dispatcher will be forced to attempt to connect to at least one hostgroup before servicing requests")
+);
+
 ////////////////////////////////////////////////////////////////////////////////
 int32 CalculateRequestQueueLengthMax()
 {
@@ -1470,6 +1478,30 @@ uint32 FHttpDispatcher::Run()
 		DBManager.SetActivationForService(BandwidthManagerServiceHandle, true);
 	}
 
+	if (GIaxHttpForceInitialConnectionCheck)
+	{
+		bool bFirstAttempt = false;
+		if (FIoStatus CertStatus = LoadDefaultHttpCertificates(bFirstAttempt); !CertStatus.IsOk())
+		{
+			UE_LOG(LogHttpIoDispatcher, Error, TEXT("Failed to load certificates, reason '%s'"), *CertStatus.ToString());
+		}
+	
+		for (;;)
+		{
+			if (FHostGroupManager::Get().ResolveInitialConnection(GIasHttpTimeOutMs, bStopRequested))
+			{
+				break;
+			}
+
+			if (bStopRequested.load(std::memory_order_relaxed))
+			{
+				break;
+			}
+
+			WakeUp->Wait(GIasHttpHealthCheckWaitTime);
+		}
+	}
+
 	for (;;)
 	{
 		UpdateThreadPriorityIfNeeded();
@@ -1498,7 +1530,7 @@ uint32 FHttpDispatcher::Run()
 		{
 			if ((FPlatformTime::Seconds() - StartTimeSeconds) > WaitTimeoutSeconds)
 			{
-				UE_LOG(LogIoStoreOnDemand, Warning, TEXT("Shutting down with %d inflight request(s) after waiting %.2lf second(s)"), Remaining, WaitTimeoutSeconds);
+				UE_LOG(LogHttpIoDispatcher, Warning, TEXT("Shutting down with %d inflight request(s) after waiting %.2lf second(s)"), Remaining, WaitTimeoutSeconds);
 				break;
 			}
 
@@ -1615,6 +1647,9 @@ FIoStatus FHttpDispatcher::RegisterHostGroup(const FName& HostGroupName, TConstA
 		HostGroupRegisterd.Broadcast(HostGroupName);
 	}
 
+	// Trigger a wakeup if we're waiting on FHostGroupManager::ResolveInitialConnection in ::Run()
+	WakeUp->Trigger();
+
 	return FIoStatus::Ok;
 }
 
