[HWRT] Fix a few issues in ray tracing support geometry collections.
- r.RayTracing.Geometry.GeometryCollection.WPO was incorrectly influencing bUseStaticRayTracingInstances.
- Segment.MaxVertices was off-by-one.
- also clean up VertexCount used in AddRayTracingGeometryUpdate(...) to use MeshDescription.NumVertices instead of being calculated based on mesh batches.

#Synced-CL 51569112

diff --git a/Engine/Source/Runtime/Experimental/GeometryCollectionEngine/Private/GeometryCollection/GeometryCollectionSceneProxy.cpp b/Engine/Source/Runtime/Experimental/GeometryCollectionEngine/Private/GeometryCollection/GeometryCollectionSceneProxy.cpp
--- a/Engine/Source/Runtime/Experimental/GeometryCollectionEngine/Private/GeometryCollection/GeometryCollectionSceneProxy.cpp
+++ b/Engine/Source/Runtime/Experimental/GeometryCollectionEngine/Private/GeometryCollection/GeometryCollectionSceneProxy.cpp
@@ -550,7 +550,7 @@ void FGeometryCollectionSceneProxyBase::UpdatingRayTracingGeometry_RenderingThre
 				Segment.FirstPrimitive = Section.TriangleStart;
 				Segment.VertexBuffer = MeshResource.PositionVertexBuffer.VertexBufferRHI;
 				Segment.NumPrimitives = Section.TriangleCount;
-				Segment.MaxVertices = Section.VertexEnd;
+				Segment.MaxVertices = Section.VertexEnd + 1;
 
 				Initializer.Segments.Add(Segment);
 
@@ -620,7 +620,6 @@ void FGeometryCollectionSceneProxyBase::GetDynamicRayTracingInstances(FRayTracin
 		RayTracingInstance.Geometry = &RayTracingGeometry;
 		RayTracingInstance.InstanceTransforms.Emplace(LocalToWorld);
 
-		uint32 MaxVertexIndex = 0;
 		for (int32 SectionIndex = 0; SectionIndex < SectionArray.Num(); ++SectionIndex)
 		{
 			const FGeometryCollectionMeshElement& Section = SectionArray[SectionIndex];
@@ -628,8 +627,6 @@ void FGeometryCollectionSceneProxyBase::GetDynamicRayTracingInstances(FRayTracin
 			FMeshBatch& Mesh = RayTracingInstance.Materials.AddDefaulted_GetRef();
 			GetMeshElement(LODIndex, SectionIndex, Section, &VertexFactory, MaterialProxies[SectionIndex], SDPG_World, UniformBuffer, bWireframe, bReverseCulling, Mesh);
 
-			MaxVertexIndex = std::max(Mesh.Elements[0].MaxVertexIndex, MaxVertexIndex);
-
 			//#TODO: bone color, bone selection and render bound?
 		}
 
@@ -650,15 +647,14 @@ void FGeometryCollectionSceneProxyBase::GetDynamicRayTracingInstances(FRayTracin
 
 		if (bNeedsUpdate)
 		{
-			const uint32 VertexCount = MaxVertexIndex + 1;
 			Collector.AddRayTracingGeometryUpdate(
 				FirstActiveViewIndex,
 				FRayTracingDynamicGeometryUpdateParams
 				{
 					RayTracingInstance.Materials,
 					false,
-					VertexCount,
-					VertexCount * (uint32)sizeof(FVector3f),
+					MeshDescription.NumVertices,
+					MeshDescription.NumVertices * (uint32)sizeof(FVector3f),
 					RayTracingGeometry.GetInitializer().TotalPrimitiveCount,
 					&RayTracingGeometry,
 					VertexBuffer,
@@ -1469,7 +1465,7 @@ FNaniteGeometryCollectionSceneProxy::FNaniteGeometryCollectionSceneProxy(UGeomet
 #if RHI_RAYTRACING
 	if (HasRayTracingRepresentation())
 	{
-		bUseStaticRayTracingInstances = !MeshDescription.SubSections.IsEmpty() && (!bAnyMaterialHasWorldPositionOffset || GRayTracingGeometryCollectionWPO) && CVarRayTracingGeometryCollectionCombinedBLAS.GetValueOnGameThread() == 0;
+		bUseStaticRayTracingInstances = !MeshDescription.SubSections.IsEmpty() && (!bAnyMaterialHasWorldPositionOffset || !GRayTracingGeometryCollectionWPO) && CVarRayTracingGeometryCollectionCombinedBLAS.GetValueOnGameThread() == 0;
 	}
 #endif
 }
