Dataflow:
- Updated GetDataflowColorMaterial() and GetDataflowColorWireframeMaterial() to handle Two sided materials

#Synced-CL 51614806

diff --git a/Engine/Source/Runtime/Experimental/Dataflow/Engine/Private/Dataflow/DataflowEngineUtil.cpp b/Engine/Source/Runtime/Experimental/Dataflow/Engine/Private/Dataflow/DataflowEngineUtil.cpp
--- a/Engine/Source/Runtime/Experimental/Dataflow/Engine/Private/Dataflow/DataflowEngineUtil.cpp
+++ b/Engine/Source/Runtime/Experimental/Dataflow/Engine/Private/Dataflow/DataflowEngineUtil.cpp
@@ -214,23 +214,33 @@ namespace UE::Dataflow
 		
 			return LoadObject<UStaticMesh>(nullptr, Path);
 		}
+
+		UStaticMesh* GetDataflowPlane()
+		{
+			const FStringView Path = TEXT("/Engine/EditorMeshes/Dataflow/SM_Dataflow_Plane.SM_Dataflow_Plane");
+
+			return LoadObject<UStaticMesh>(nullptr, Path);
+		}
 	}
 
 	namespace RenderMaterial
 	{
-		UMaterialInterface* GetDataflowColorMaterial()
+		UMaterialInterface* GetDataflowColorMaterial(const bool bTwoSided)
 		{
-			const FStringView Path = TEXT("/Engine/EditorMaterials/Dataflow/M_Dataflow_Color.M_Dataflow_Color");
+			const FStringView Path = bTwoSided
+				? TEXT("/Engine/EditorMaterials/Dataflow/M_Dataflow_Color_TwoSided.M_Dataflow_Color_TwoSided")
+				: TEXT("/Engine/EditorMaterials/Dataflow/M_Dataflow_Color.M_Dataflow_Color");
 
 			return LoadObject<UMaterialInterface>(nullptr, Path);
 		}
 
-		UMaterialInterface* GetDataflowColorWireframeMaterial()
+		UMaterialInterface* GetDataflowColorWireframeMaterial(const bool bTwoSided)
 		{
-			const FStringView Path = TEXT("/Engine/EditorMaterials/Dataflow/M_Dataflow_Color_Wireframe.M_Dataflow_Color_Wireframe");
+			const FStringView Path = bTwoSided
+				? TEXT("/Engine/EditorMaterials/Dataflow/M_Dataflow_Color_Wireframe_TwoSided.M_Dataflow_Color_Wireframe_TwoSided")
+				: TEXT("/Engine/EditorMaterials/Dataflow/M_Dataflow_Color_Wireframe.M_Dataflow_Color_Wireframe");
 
 			return LoadObject<UMaterialInterface>(nullptr, Path);
 		}
 	}
 }
-
diff --git a/Engine/Source/Runtime/Experimental/Dataflow/Engine/Public/Dataflow/DataflowEngineUtil.h b/Engine/Source/Runtime/Experimental/Dataflow/Engine/Public/Dataflow/DataflowEngineUtil.h
--- a/Engine/Source/Runtime/Experimental/Dataflow/Engine/Public/Dataflow/DataflowEngineUtil.h
+++ b/Engine/Source/Runtime/Experimental/Dataflow/Engine/Public/Dataflow/DataflowEngineUtil.h
@@ -100,12 +100,13 @@ namespace UE::Dataflow
 	{
 		UE_API UStaticMesh* GetDataflowBox();
 		UE_API UStaticMesh* GetDataflowSphere();
+		UE_API UStaticMesh* GetDataflowPlane();
 	}
 
 	namespace RenderMaterial
 	{
-		UE_API UMaterialInterface* GetDataflowColorMaterial();
-		UE_API UMaterialInterface* GetDataflowColorWireframeMaterial();
+		UE_API UMaterialInterface* GetDataflowColorMaterial(const bool bTwoSided = false);
+		UE_API UMaterialInterface* GetDataflowColorWireframeMaterial(const bool bTwoSided = false);
 	}
 }
 
