Broader error messages on verification failure

#rnx

#Synced-CL 51536473

diff --git a/Engine/Source/Runtime/Experimental/IoStore/HttpClient/Private/Peer.inl b/Engine/Source/Runtime/Experimental/IoStore/HttpClient/Private/Peer.inl
--- a/Engine/Source/Runtime/Experimental/IoStore/HttpClient/Private/Peer.inl
+++ b/Engine/Source/Runtime/Experimental/IoStore/HttpClient/Private/Peer.inl
@@ -386,12 +386,22 @@ static FOutcome Ssl_Handshake(SSL* Ssl)
 	if (Result == 0) return FOutcome::Error("unsuccessful tls handshake");
 	if (Result != 1) return Ssl_GetOutcome(Ssl, Result, "tls handshake error");
 
-	if (Result = SSL_get_verify_result(Ssl); Result != X509_V_OK)
+	Result = SSL_get_verify_result(Ssl);
+	if (Result == X509_V_OK)
 	{
-		return FOutcome::Error("x509 verification error", Result);
+		return FOutcome::Ok();
 	}
 
-	return FOutcome::Ok();
+	const char* Msg;
+	switch (Result)
+	{
+	case X509_V_ERR_CERT_HAS_EXPIRED:			Msg = "certificate has expired (check local date)";		break;
+	case X509_V_ERR_CERT_NOT_YET_VALID:			Msg = "certificate not yet valid (check local date)";	break;
+	case X509_V_ERR_DEPTH_ZERO_SELF_SIGNED_CERT:Msg = "self-signed certificate";						break;
+	case X509_V_ERR_SELF_SIGNED_CERT_IN_CHAIN:	Msg = "self-signed certificate in chain";				break;
+	default:									Msg = "x509 verification error";						break;
+	}
+	return FOutcome::Error(Msg, Result);
 }
 
 ////////////////////////////////////////////////////////////////////////////////
