Chaos Cloth:
- Added an experimental hybrid CCD collision mode.
Currently only works on the first iteration and for non complex colliders.
- Fixed CCD friction code path.
- Fixed non ISPC friction.

#rnx
#Synced-CL 51543885

diff --git a/Engine/Plugins/ChaosCloth/Source/ChaosCloth/Private/ChaosCloth/ChaosClothConstraints.cpp b/Engine/Plugins/ChaosCloth/Source/ChaosCloth/Private/ChaosCloth/ChaosClothConstraints.cpp
--- a/Engine/Plugins/ChaosCloth/Source/ChaosCloth/Private/ChaosCloth/ChaosClothConstraints.cpp
+++ b/Engine/Plugins/ChaosCloth/Source/ChaosCloth/Private/ChaosCloth/ChaosClothConstraints.cpp
@@ -1568,7 +1568,8 @@ void FClothConstraints::CreateCollisionConstraint(
 			CollisionContacts,
 			CollisionNormals,
 			CollisionPhis,
-			ComplexBoneData
+			ComplexBoneData,
+			WeightMaps
 			);
 		++NumCollisionConstraintRules;
 		++NumUpdateLinearSystemCollisionsRules;
diff --git a/Engine/Plugins/ChaosClothAssetDataflowNodes/Source/ChaosClothAssetDataflowNodes/Private/ChaosClothAsset/SimulationCollisionConfigNode.cpp b/Engine/Plugins/ChaosClothAssetDataflowNodes/Source/ChaosClothAssetDataflowNodes/Private/ChaosClothAsset/SimulationCollisionConfigNode.cpp
--- a/Engine/Plugins/ChaosClothAssetDataflowNodes/Source/ChaosClothAssetDataflowNodes/Private/ChaosClothAsset/SimulationCollisionConfigNode.cpp
+++ b/Engine/Plugins/ChaosClothAssetDataflowNodes/Source/ChaosClothAssetDataflowNodes/Private/ChaosClothAsset/SimulationCollisionConfigNode.cpp
@@ -16,12 +16,17 @@ FChaosClothAssetSimulationCollisionConfigNode::FChaosClothAssetSimulationCollisi
 	RegisterInputConnection(&ClothCollisionThicknessImported.WeightMap)
 		.SetCanHidePin(true)
 		.SetPinIsHidden(true);
+	RegisterInputConnection(&InnerCollisionThickness.WeightMap)
+		.SetCanHidePin(true)
+		.SetPinIsHidden(true);
 }
 PRAGMA_ENABLE_DEPRECATION_WARNINGS
 
 void FChaosClothAssetSimulationCollisionConfigNode::AddProperties(FPropertyHelper& PropertyHelper) const
 {
 	PropertyHelper.SetPropertyBool(this, &bUseCCD);
+	PropertyHelper.SetPropertyWeighted(this, &InnerCollisionThickness);
+
 	PropertyHelper.SetProperty(this, &ProximityStiffness);
 	
 	PropertyHelper.SetFabricProperty(FName(TEXT("CollisionThickness")), CollisionThicknessImported,
diff --git a/Engine/Plugins/ChaosClothAssetDataflowNodes/Source/ChaosClothAssetDataflowNodes/Public/ChaosClothAsset/SimulationCollisionConfigNode.h b/Engine/Plugins/ChaosClothAssetDataflowNodes/Source/ChaosClothAssetDataflowNodes/Public/ChaosClothAsset/SimulationCollisionConfigNode.h
--- a/Engine/Plugins/ChaosClothAssetDataflowNodes/Source/ChaosClothAssetDataflowNodes/Public/ChaosClothAsset/SimulationCollisionConfigNode.h
+++ b/Engine/Plugins/ChaosClothAssetDataflowNodes/Source/ChaosClothAssetDataflowNodes/Public/ChaosClothAsset/SimulationCollisionConfigNode.h
@@ -64,6 +64,17 @@ public:
 	UPROPERTY(EditAnywhere, Category = "Collision Properties", Meta = (InteractorName = "UseCCD"))
 	bool bUseCCD = false;
 
+	/**
+	 * Experimental feature, use at your own risk.
+	 * Improve collision behavior by running the CCD code path, but only where particles penetration exceeds the given inner thickness value.
+	 * This is using a faster collision code path than CCD, but unlike CCD it does not resolve fast moving objects.
+	 * Use an inner thickness that is about a few centimeters to start with, up to a quarter of your collider diameter for best results.
+	 * If the inner thickness is too large (e.g. more than half the collider's diameter), the CCD code path won't ever be activated.
+	 * This property can be used in conjunction with a weight map to selectively allow certain particles to use the extra calculations or to locally adjust the effect.
+	 * Set both the low and high values to 0 to disable.
+	 */
+	UPROPERTY(EditAnywhere, Category = "Collision Properties", DisplayName = "Inner Collision Thickness (Experimental)", Meta = (UIMin = "0", UIMax = "10", ClampMin = "0", EditCondition = "bUseCCD == false"))
+	FChaosClothAssetWeightedValueNonAnimatable InnerCollisionThickness = { 0.f, 0.f, TEXT("InnerCollisionThickness"), false };
 
 	UE_DEPRECATED(5.6, "Use FrictionCoefficientWeighted instead.")
 	UPROPERTY()
diff --git a/Engine/Source/Runtime/Experimental/Chaos/Private/Chaos/PBDSoftBodyCollisionConstraint.cpp b/Engine/Source/Runtime/Experimental/Chaos/Private/Chaos/PBDSoftBodyCollisionConstraint.cpp
--- a/Engine/Source/Runtime/Experimental/Chaos/Private/Chaos/PBDSoftBodyCollisionConstraint.cpp
+++ b/Engine/Source/Runtime/Experimental/Chaos/Private/Chaos/PBDSoftBodyCollisionConstraint.cpp
@@ -26,6 +26,8 @@ static int32 Chaos_SoftBodyCollision_ISPC_ParallelBatchSize = 128;
 FAutoConsoleVariableRef CVarChaosSoftBodyCollisionISPCParallelBatchSize(TEXT("p.Chaos.SoftBodyCollision.ISPC.ParallelBatchSize"), Chaos_SoftBodyCollision_ISPC_ParallelBatchSize, TEXT("Parallel batch size for ISPC"));
 #endif
 
+DECLARE_CYCLE_STAT(TEXT("Chaos PBD Collision Constraint Apply"), STAT_ChaosPBDCollisionConstraintApply, STATGROUP_Chaos);
+
 namespace Chaos::Softs {
 
 	namespace Private
@@ -51,7 +53,7 @@ namespace Chaos::Softs {
 			return true;
 		}
 
-		static void ApplyFriction(FSolverVec3& P, const FSolverVec3& X, const FSolverVec3& NormalWorld, const FSolverReal MaxFrictionCorrection, const FSolverReal Dt, const FSolverVec3& CollisionX, const FSolverVec3& CollisionV, const FSolverVec3& CollisionW, FSolverVec3& ColliderVelocityAtPoint)
+		static void ApplyFrictionAndReturnVelocity(FSolverVec3& P, const FSolverVec3& X, const FSolverVec3& NormalWorld, const FSolverReal MaxFrictionCorrection, const FSolverReal Dt, const FSolverVec3& CollisionX, const FSolverVec3& CollisionV, const FSolverVec3& CollisionW, FSolverVec3& ColliderVelocityAtPoint)
 		{
 			const FSolverVec3 VectorToPoint = P - CollisionX;
 			ColliderVelocityAtPoint = CollisionV + FSolverVec3::CrossProduct(CollisionW, VectorToPoint);
@@ -70,12 +72,150 @@ namespace Chaos::Softs {
 		static void ApplyFriction(FSolverVec3& P, const FSolverVec3& X, const FSolverVec3& NormalWorld, const FSolverReal MaxFrictionCorrection, const FSolverReal Dt, const FSolverVec3& CollisionX, const FSolverVec3& CollisionV, const FSolverVec3& CollisionW)
 		{
 			FSolverVec3 ColliderVelocityAtPoint;
-			ApplyFriction(P, X, NormalWorld, MaxFrictionCorrection, Dt, CollisionX, CollisionV, CollisionW, ColliderVelocityAtPoint);
+			ApplyFrictionAndReturnVelocity(P, X, NormalWorld, MaxFrictionCorrection, Dt, CollisionX, CollisionV, CollisionW, ColliderVelocityAtPoint);
 		}
 
 		// Defined in PerParticlePBDCollisionConstraint.cpp
 		extern void ReflectOneSidedCollision(const FSolverVec3& P, const FSolverVec3& OneSidedPlaneNormal, const FSolverVec3& SplitOrigin, FSolverReal& Penetration, FSolverVec3& ImplicitNormal);
+
+		template<bool bWithFriction>
+		static inline bool ApplySimpleCCD(
+			const FSolverCollisionParticlesRange& CollisionParticlesRange,
+			const int32 CollisionIndex,
+			const FSolverRigidTransform3& LastCollisionTransform,
+			const FSolverVec3& X,
+			const FSolverReal Thickness,
+			[[maybe_unused]] const FSolverReal FrictionCoefficient,
+			[[maybe_unused]] const FSolverReal Dt,
+			FSolverVec3& P,
+			FSolverVec3& NormalWorld,
+			FSolverVec3& ContactWorld,
+			FSolverReal& Penetration)
+		{
+			const FSolverRigidTransform3 Frame(CollisionParticlesRange.GetX(CollisionIndex), CollisionParticlesRange.R(CollisionIndex));
+			const FVec3 StartPoint(LastCollisionTransform.InverseTransformPositionNoScale(X));  // Geometry operates in FReal
+			const FVec3 EndPoint(Frame.InverseTransformPositionNoScale(P));
+
+			const FImplicitObjectPtr& Geometry = CollisionParticlesRange.GetGeometry(CollisionIndex);
+			const Pair<FVec3, bool> PointPair = Geometry->FindClosestIntersection(StartPoint, EndPoint, (FReal)Thickness);
+
+			if (PointPair.Second)
+			{
+				const FSolverVec3 Normal = FSolverVec3(Geometry->Normal(PointPair.First));
+				NormalWorld = Frame.TransformVectorNoScale(Normal);
+				ContactWorld = Frame.TransformPositionNoScale(UE::Math::TVector<FSolverReal>(PointPair.First));
+
+				const FSolverVec3 Direction = ContactWorld - P;
+				Penetration = FMath::Max((FSolverReal)0., FSolverVec3::DotProduct(NormalWorld, Direction)) + (FSolverReal)UE_THRESH_POINT_ON_PLANE;
+
+				P += Penetration * NormalWorld;
+
+				if constexpr (bWithFriction)
+				{
+					// Friction
+					const FSolverVec3 CollisionV = CollisionParticlesRange.V(CollisionIndex);
+					const FSolverVec3 CollisionX = CollisionParticlesRange.X(CollisionIndex);
+					const FSolverVec3 CollisionW = CollisionParticlesRange.W(CollisionIndex);
+					Private::ApplyFriction(P, X, NormalWorld, Penetration * FrictionCoefficient, Dt, CollisionX, CollisionV, CollisionW);
+				}
+				return true;
+			}
+			return false;
+		}
+	}
+
+// Helper function to call ApplySimpleCCD and return data to ISPC
+extern "C" void ApplySimpleCCDNoFriction(
+	const uint8* CollisionParticlesRange,
+	const int32 CollisionIndex,
+	const FSolverRigidTransform3& LastCollisionTransform,
+	const FSolverReal* XSoa,
+	const FSolverReal Thickness,
+	FSolverReal* PSoa,
+	const int32 ProgramCount,
+	const int32 Mask)
+{
+	for (int32 Index = 0; Index < ProgramCount; ++Index)
+	{
+		if (Mask & (1 << Index))
+		{
+			const FSolverVec3 X(
+				XSoa[Index],
+				XSoa[Index + ProgramCount],
+				XSoa[Index + 2 * ProgramCount]);
+			FSolverVec3 P(
+				PSoa[Index],
+				PSoa[Index + ProgramCount],
+				PSoa[Index + 2 * ProgramCount]);
+			constexpr bool bWithFriction = false;
+			FSolverVec3 Normal;
+			FSolverVec3 Contact;
+			FSolverReal Penetration;
+			Private::ApplySimpleCCD<bWithFriction>(
+				*(const FSolverCollisionParticlesRange*)CollisionParticlesRange,
+				CollisionIndex,
+				LastCollisionTransform,
+				X,
+				Thickness,
+				(FSolverReal)0., // No friction, no FrictionCoefficient needed
+				(FSolverReal)0., // No friction, no Dt needed
+				P,
+				Normal,
+				Contact,
+				Penetration);
+			PSoa[Index] = P.X;
+			PSoa[Index + ProgramCount] = P.Y;
+			PSoa[Index + 2 * ProgramCount] = P.Z;
+		}
+	}
+}
+
+extern "C" void ApplySimpleCCDWithFriction(
+	const uint8* CollisionParticlesRange,
+	const int32 CollisionIndex,
+	const FSolverRigidTransform3& LastCollisionTransform,
+	const FSolverReal* XSoa,
+	const FSolverReal Thickness,
+	const FSolverReal FrictionCoefficient,
+	const FSolverReal Dt,
+	FSolverReal* PSoa,
+	const int32 ProgramCount,
+	const int32 Mask)
+{
+	for (int32 Index = 0; Index < ProgramCount; ++Index)
+	{
+		if (Mask & (1 << Index))
+		{
+			const FSolverVec3 X(
+				XSoa[Index],
+				XSoa[Index + ProgramCount],
+				XSoa[Index + 2 * ProgramCount]);
+			FSolverVec3 P(
+				PSoa[Index],
+				PSoa[Index + ProgramCount],
+				PSoa[Index + 2 * ProgramCount]);
+			constexpr bool bWithFriction = true;
+			FSolverVec3 Normal;
+			FSolverVec3 Contact;
+			FSolverReal Penetration;
+			Private::ApplySimpleCCD<bWithFriction>(
+				*(const FSolverCollisionParticlesRange*)CollisionParticlesRange,
+				CollisionIndex,
+				LastCollisionTransform,
+				X,
+				Thickness,
+				FrictionCoefficient,
+				Dt,
+				P,
+				Normal,
+				Contact,
+				Penetration);
+			PSoa[Index] = P.X;
+			PSoa[Index + ProgramCount] = P.Y;
+			PSoa[Index + 2 * ProgramCount] = P.Z;
+		}
 	}
+}
 
 // Helper function to call PhiWithNormal and return data to ISPC
 extern "C" void GetPhiWithNormalCollisionParticleRange(const uint8 * CollisionParticlesRange, const FSolverReal * InV, FSolverReal * Normal, FSolverReal * Phi, const int32 i, const int32 ProgramCount, const int32 Mask)
@@ -125,14 +265,14 @@ void FPBDSoftBodyCollisionConstraintBase::ApplyWithPlanarConstraints(FSolverPart
 			}
 			else
 			{
-				ApplyInternalCCD<true, false>(Particles, Dt, CollisionParticles);
+				ApplyInternalCCD<false, true>(Particles, Dt, CollisionParticles);
 			}
 		}
 		else
 		{
 			if (bLockAndWriteContacts)
 			{
-				ApplyInternalCCD<false, true>(Particles, Dt, CollisionParticles);
+				ApplyInternalCCD<true, false>(Particles, Dt, CollisionParticles);
 			}
 			else
 			{
@@ -334,14 +474,14 @@ void FPBDSoftBodyCollisionConstraintBase::Apply(FSolverParticlesRange& Particles
 			}
 			else
 			{
-				ApplyInternalCCD<true, false>(Particles, Dt, CollisionParticles);
+				ApplyInternalCCD<false, true>(Particles, Dt, CollisionParticles);
 			}
 		}
 		else
 		{
 			if (bLockAndWriteContacts)
 			{
-				ApplyInternalCCD<false, true>(Particles, Dt, CollisionParticles);
+				ApplyInternalCCD<true, false>(Particles, Dt, CollisionParticles);
 			}
 			else
 			{
@@ -422,9 +562,12 @@ void FPBDSoftBodyCollisionConstraintBase::ApplySimpleInternal(FSolverParticlesRa
 			{
 				return;
 			}
+			const FSolverReal InnerCollisionThicknessValue = InnerCollisionThickness.HasWeightMap() ? InnerCollisionThickness[Index] : (FSolverReal)InnerCollisionThickness;
 
 			for (const FSolverCollisionParticlesRange& CollisionParticlesRange : CollisionParticles)
 			{
+				const TConstArrayView<FSolverRigidTransform3> CollisionTransforms = CollisionParticlesRange.GetConstArrayView(LastCollisionTransforms);
+
 				for (int32 CollisionIndex = 0; CollisionIndex < CollisionParticlesRange.GetRangeSize(); ++CollisionIndex)
 				{
 					const FImplicitObjectPtr& Geometry = CollisionParticlesRange.GetGeometry(CollisionIndex);
@@ -439,7 +582,8 @@ void FPBDSoftBodyCollisionConstraintBase::ApplySimpleInternal(FSolverParticlesRa
 					FVec3 ImplicitNormalReal;
 					const FSolverReal Phi = (FSolverReal)Geometry->PhiWithNormal(RigidSpacePosition, ImplicitNormalReal);
 					FSolverVec3 ImplicitNormal(ImplicitNormalReal);
-					FSolverReal Penetration = SoftBodyCollisionThickness + CollisionThickness - Phi;
+					const FSolverReal Thickness = SoftBodyCollisionThickness + CollisionThickness;
+					FSolverReal Penetration = Thickness - Phi;
 					if (Penetration > (FSolverReal)0.)
 					{
 						// Split capsules always push out in the OneSidedPlaneNormal direction.
@@ -452,6 +596,36 @@ void FPBDSoftBodyCollisionConstraintBase::ApplySimpleInternal(FSolverParticlesRa
 							}
 						}
 
+						// Fallback to CCD when the penetration is deep enough to cause large projection errors
+						if (InnerCollisionThicknessValue > (FSolverReal)0. && Penetration > InnerCollisionThicknessValue)
+						{
+							FSolverVec3 NormalWorld;
+							FSolverVec3 ContactWorld;
+							if (Private::ApplySimpleCCD<bWithFriction>(
+								CollisionParticlesRange,
+								CollisionIndex,
+								CollisionTransforms[CollisionIndex],
+								X[Index],
+								Thickness,
+								FrictionCoefficient,
+								Dt,
+								PAndInvM[Index].P,
+								NormalWorld,
+								ContactWorld,
+								Penetration))
+							{
+								if constexpr (bLockAndWriteContacts)
+								{
+									FScopeLock Lock(&DebugMutex);
+									CollisionParticlesRange.GetArrayView(*CollisionParticleCollided)[CollisionIndex] = true;
+									Contacts->Emplace(ContactWorld);
+									Normals->Emplace(NormalWorld);
+									Phis->Emplace(Penetration);
+								}
+							}
+							continue;
+						}
+
 						const FSolverVec3 NormalWorld = Frame.TransformVector(ImplicitNormal);
 						if constexpr (bLockAndWriteContacts)
 						{
@@ -475,7 +649,16 @@ void FPBDSoftBodyCollisionConstraintBase::ApplySimpleInternal(FSolverParticlesRa
 						if constexpr (bWithFriction)
 						{
 							FSolverVec3 ColliderVelocityAtPoint;
-							Private::ApplyFriction(PAndInvM[Index].P, X[Index], NormalWorld, Penetration * FrictionCoefficient, Dt, CollisionParticlesRange.X(CollisionIndex), CollisionParticlesRange.W(CollisionIndex), ColliderVelocityAtPoint);
+							Private::ApplyFrictionAndReturnVelocity(
+								PAndInvM[Index].P,
+								X[Index],
+								NormalWorld,
+								Penetration * FrictionCoefficient,
+								Dt,
+								CollisionParticlesRange.X(CollisionIndex),
+								CollisionParticlesRange.V(CollisionIndex),
+								CollisionParticlesRange.W(CollisionIndex),
+								ColliderVelocityAtPoint);
 
 							if constexpr (bGeneratePlanarConstraints)
 							{
@@ -567,11 +750,11 @@ void FPBDSoftBodyCollisionConstraintBase::ApplyComplexInternal(FSolverParticlesR
 								if (ColliderBoneData->MappedBoneIndices.IsValidIndex(StrongestBone))
 								{
 									const int32 MappedIndex = ColliderBoneData->MappedBoneIndices[StrongestBone];
-									Private::ApplyFriction(PAndInvM[Index].P, X[Index], NormalWorld, Penetration * FrictionCoefficient, Dt, ColliderBoneData->X[MappedIndex], ColliderBoneData->V[MappedIndex], ColliderBoneData->W[MappedIndex], ColliderVelocityAtPoint);
+									Private::ApplyFrictionAndReturnVelocity(PAndInvM[Index].P, X[Index], NormalWorld, Penetration * FrictionCoefficient, Dt, ColliderBoneData->X[MappedIndex], ColliderBoneData->V[MappedIndex], ColliderBoneData->W[MappedIndex], ColliderVelocityAtPoint);
 								}
 								else
 								{
-									Private::ApplyFriction(PAndInvM[Index].P, X[Index], NormalWorld, Penetration * FrictionCoefficient, Dt, CollisionX, CollisionV, CollisionW, ColliderVelocityAtPoint);
+									Private::ApplyFrictionAndReturnVelocity(PAndInvM[Index].P, X[Index], NormalWorld, Penetration * FrictionCoefficient, Dt, CollisionX, CollisionV, CollisionW, ColliderVelocityAtPoint);
 								}
 
 								if constexpr (bGeneratePlanarConstraints)
@@ -634,7 +817,7 @@ void FPBDSoftBodyCollisionConstraintBase::ApplyComplexInternal(FSolverParticlesR
 								}
 
 								FSolverVec3 ColliderVelocityAtPoint;
-								Private::ApplyFriction(PAndInvM[Index].P, X[Index], NormalWorld, Penetration * FrictionCoefficient, Dt, CollisionX, CollisionV, CollisionW, ColliderVelocityAtPoint);
+								Private::ApplyFrictionAndReturnVelocity(PAndInvM[Index].P, X[Index], NormalWorld, Penetration * FrictionCoefficient, Dt, CollisionX, CollisionV, CollisionW, ColliderVelocityAtPoint);
 
 								if (bGeneratePlanarConstraints)
 								{
@@ -868,6 +1051,7 @@ void FPBDSoftBodyCollisionConstraintBase::ApplySimpleInternalISPC(FSolverParticl
 	}
 
 	TRACE_CPUPROFILER_EVENT_SCOPE(FPBDSoftBodyCollisionConstraint_ApplySimpleInternalISPC);
+	SCOPE_CYCLE_COUNTER(STAT_ChaosPBDCollisionConstraintApply);
 
 	check(bRealTypeCompatibleWithISPC);
 
@@ -933,6 +1117,10 @@ void FPBDSoftBodyCollisionConstraintBase::ApplySimpleInternalISPC(FSolverParticl
 				}
 				else
 				{
+					const bool bHasInnerCollisionThickness = (InnerCollisionThickness.GetWeightedValue() != FSolverVec2::ZeroVector);
+					const bool bHasInnerCollisionThicknessWeightMap = bHasInnerCollisionThickness && InnerCollisionThickness.HasWeightMap();
+
+					const TConstArrayView<FSolverRigidTransform3> CollisionTransforms = CollisionParticlesRange.GetConstArrayView(LastCollisionTransforms);
 					if (bWithFriction)
 					{
 						ispc::ApplyPerParticleSimpleCollisionFastFriction(
@@ -942,8 +1130,11 @@ void FPBDSoftBodyCollisionConstraintBase::ApplySimpleInternalISPC(FSolverParticl
 							(const ispc::FVector3f*)CollisionParticlesRange.XArray().GetData(),
 							(const ispc::FVector3f*)CollisionParticlesRange.GetW().GetData(),
 							(const ispc::FVector4f*)CollisionParticlesRange.GetR().GetData(),
+							(const ispc::FTransform3f*)CollisionTransforms.GetData(),
 							FrictionCoefficient,
 							SoftBodyCollisionThickness + CollisionThickness,
+							bHasInnerCollisionThickness ? InnerCollisionThickness.GetTable().GetData() : nullptr,
+							bHasInnerCollisionThicknessWeightMap ? InnerCollisionThickness.GetIndices().GetData() : nullptr,
 							(const uint8*)&CollisionParticlesRange,
 							(const uint8*)CollisionParticlesRange.GetAllGeometry().GetData(),
 							sizeof(FImplicitObject),
@@ -961,7 +1152,10 @@ void FPBDSoftBodyCollisionConstraintBase::ApplySimpleInternalISPC(FSolverParticl
 							(const ispc::FVector3f*)Particles.XArray().GetData(),
 							(const ispc::FVector3f*)CollisionParticlesRange.XArray().GetData(),
 							(const ispc::FVector4f*)CollisionParticlesRange.GetR().GetData(),
+							(const ispc::FTransform3f*)CollisionTransforms.GetData(),
 							SoftBodyCollisionThickness + CollisionThickness,
+							bHasInnerCollisionThickness ? InnerCollisionThickness.GetTable().GetData() : nullptr,
+							bHasInnerCollisionThicknessWeightMap ? InnerCollisionThickness.GetIndices().GetData() : nullptr,
 							(const uint8*)&CollisionParticlesRange,
 							(const uint8*)CollisionParticlesRange.GetAllGeometry().GetData(),
 							sizeof(FImplicitObject),
@@ -1275,6 +1469,53 @@ void FPBDSoftBodyCollisionConstraintBase::UpdateLinearSystem(const FSolverPartic
 	});
 }
 
+FPBDSoftBodyCollisionConstraint::FPBDSoftBodyCollisionConstraint(
+	const TArray<FSolverRigidTransform3>& InLastCollisionTransforms,
+	const FCollectionPropertyConstFacade& PropertyCollection,
+	FSolverReal InMeshScale,
+	TArray<bool>* InCollisionParticleCollided,
+	TArray<FSolverVec3>* InContacts,
+	TArray<FSolverVec3>* InNormals,
+	TArray<FSolverReal>* InPhis,
+	const TMap<FParticleRangeIndex, FPBDComplexColliderBoneData>& InComplexBoneData,
+	const TMap<FString, TConstArrayView<float>>& WeightMaps)
+	: Base(InLastCollisionTransforms,
+		InMeshScale* GetCollisionThickness(PropertyCollection, DefaultCollisionThickness),
+		GetFrictionCoefficient(PropertyCollection, DefaultFrictionCoefficient),
+		GetUseCCD(PropertyCollection, false),
+		GetProximityStiffness(PropertyCollection, DefaultProximityStiffness),
+		InCollisionParticleCollided, InContacts, InNormals, InPhis,
+		GetSoftBodyCollisionThickness(PropertyCollection, DefaultSoftBodyCollisionThickness),
+		GetEnableSimpleColliders(PropertyCollection, true),
+		GetEnableComplexColliders(PropertyCollection, true),
+		GetUsePlanarConstraintForSimpleColliders(PropertyCollection, bDefaultUsePlanarConstraintForSimpleColliders),
+		GetUsePlanarConstraintForComplexColliders(PropertyCollection, bDefaultUsePlanarConstraintForComplexColliders),
+		InComplexBoneData)
+	, MeshScale(InMeshScale)
+	, CollisionThicknessIndex(PropertyCollection)
+	, SoftBodyCollisionThicknessIndex(PropertyCollection)
+	, FrictionCoefficientIndex(PropertyCollection)
+	, UseCCDIndex(PropertyCollection)
+	, ProximityStiffnessIndex(PropertyCollection)
+	, EnableSimpleCollidersIndex(PropertyCollection)
+	, UsePlanarConstraintForSimpleCollidersIndex(PropertyCollection)
+	, EnableComplexCollidersIndex(PropertyCollection)
+	, UsePlanarConstraintForComplexCollidersIndex(PropertyCollection)
+{
+	// Initialize the base class InnerCollisionThickness map
+	const FSolverVec2 WeightedValue(GetWeightedFloatInnerCollisionThickness(PropertyCollection, FSolverVec2::ZeroVector));
+	if (WeightedValue.X > (FSolverReal)0. || WeightedValue.Y > (FSolverReal)0.)
+	{
+		const FString WeighMapString = GetInnerCollisionThicknessString(PropertyCollection, InnerCollisionThicknessName.ToString());
+		const TConstArrayView<FRealSingle>& WeightMap = WeightMaps.FindRef(WeighMapString);
+		InnerCollisionThickness = FPBDWeightMap(
+			WeightedValue * MeshScale,
+			WeightMap,
+			WeightMap.Num());
+		InnerCollisionThickness.ApplyValues();
+	}
+}
+
 void FPBDSoftBodyCollisionConstraint::SetProperties(const FCollectionPropertyConstFacade& PropertyCollection)
 {
 	if (IsCollisionThicknessMutable(PropertyCollection))
diff --git a/Engine/Source/Runtime/Experimental/Chaos/Private/Chaos/PerParticlePBDCollisionConstraint.ispc b/Engine/Source/Runtime/Experimental/Chaos/Private/Chaos/PerParticlePBDCollisionConstraint.ispc
--- a/Engine/Source/Runtime/Experimental/Chaos/Private/Chaos/PerParticlePBDCollisionConstraint.ispc
+++ b/Engine/Source/Runtime/Experimental/Chaos/Private/Chaos/PerParticlePBDCollisionConstraint.ispc
@@ -722,6 +722,28 @@ static inline void ApplyFriction(FVector3f& P, const FVector3f& X, const FVector
 	ApplyFriction(P, X, NormalWorld, MaxFrictionCorrection, Dt, CollisionX, CollisionV, CollisionW, ColliderVelocityAtPoint);
 }
 
+extern "C" unmasked void ApplySimpleCCDNoFriction(
+	const uniform uint8* uniform CollisionParticlesRange,
+	const uniform int CollisionIndex,
+	const uniform FTransform3f& LastCollisionTransform,
+	const uniform float* uniform X,
+	const uniform float uniform Thickness,
+	uniform float* uniform P,
+	const uniform int ProgramCount, 
+	const uniform int Mask);
+
+extern "C" unmasked void ApplySimpleCCDWithFriction(
+	const uniform uint8* uniform CollisionParticlesRange,
+	const uniform int CollisionIndex,
+	const uniform FTransform3f& LastCollisionTransform,
+	const uniform float* uniform X,
+	const uniform float uniform Thickness,
+	const uniform float uniform FrictionCoefficient,
+	const uniform float Dt,
+	uniform float* uniform P,
+	const uniform int ProgramCount, 
+	const uniform int Mask);
+
 extern "C" unmasked void GetPhiWithNormal(const uniform uint8 *uniform CollisionParticles, const uniform float *uniform V, uniform float *uniform Normal, uniform float *uniform Phi, const uniform int i, const uniform int ProgramCount, const uniform int Mask);
 extern "C" unmasked void GetPhiWithNormalCollisionParticleRange(const uniform uint8 *uniform CollisionParticlesRange, const uniform float *uniform V, uniform float *uniform Normal, uniform float *uniform Phi, const uniform int i, const uniform int ProgramCount, const uniform int Mask);
 
@@ -933,8 +955,11 @@ export void ApplyPerParticleSimpleCollisionFastFriction(uniform FVector4f Partic
 													const uniform FVector3f CollisionX[],
 													const uniform FVector3f CollisionW[],
 													const uniform FVector4f CollisionR[],
+													const uniform FTransform3f LastCollisionTransforms[],
 													const uniform float PerGroupFriction,
 													const uniform float PerGroupThickness,
+													const uniform float *uniform InnerThicknessTable,
+													const uniform uint8 *uniform InnerThicknessIndices,
 													const uniform uint8 *uniform CollisionParticlesRange,
 													const uniform uint8 *uniform Geometry,
 													const uniform int SizeofFImplicitObject,
@@ -945,64 +970,252 @@ export void ApplyPerParticleSimpleCollisionFastFriction(uniform FVector4f Partic
 													const uniform int OuterParticleBatchBegin,
 													const uniform int OuterParticleBatchEnd)
 {
-	foreach(Index = OuterParticleBatchBegin ... OuterParticleBatchEnd)
+	const uniform bool bHasInnerThickness = (InnerThicknessTable != NULL);
+	if(bHasInnerThickness)
 	{
-		const FVector4f PandInvM = VectorLoad(&ParticlesPandInvM[extract(Index, 0)]);
-		varying FVector3f P;
-		varying float IM;
-		UnzipPandInvM(PandInvM, P, IM);
-
-		if(all(IM == 0))
+		const uniform bool bHasInnerThicknessWeighmap = (InnerThicknessIndices != NULL);
+		if(bHasInnerThicknessWeighmap)
 		{
-			continue;
-		}
+			foreach(Index = OuterParticleBatchBegin ... OuterParticleBatchEnd)
+			{
+				const FVector4f PandInvM = VectorLoad(&ParticlesPandInvM[extract(Index, 0)]);
+				varying FVector3f P;
+				varying float IM;
+				UnzipPandInvM(PandInvM, P, IM);
 
-		for(uniform int i = 0; i < InnerCollisionCount; i++) // SequentialFor
+				if(all(IM == 0))
+				{
+					continue;
+				}
+
+				varying uint8 InnerThicknessIndex = InnerThicknessIndices[Index];
+				#pragma ignore warning(perf)
+				varying float InnerThickness = InnerThicknessTable[InnerThicknessIndex];
+
+				for(uniform int i = 0; i < InnerCollisionCount; i++) // SequentialFor
+				{
+					const uniform SerializedPtr *uniform GeometryPtr = (const uniform SerializedPtr *uniform)Geometry;
+					const uniform uint8 *uniform Ptr = GeometryPtr[i].Ptr;
+					const uniform uint8 *uniform TypedPtr = Ptr + SizeofFImplicitObject;
+					const uniform uint8 Type = Ptr[OffsetofGeometryType] & 0x3F;
+
+					if(!IsSimpleCollider(Type))
+					{
+						continue;
+					}
+
+					// InverseTransformPosition. Can skip Scale calculation because it is 1 for xyz.
+					const FVector3f TranslatedVec = P - CollisionX[i];
+					const uniform FVector4f InverseR = QuatInverse(CollisionR[i]);
+					const FVector3f PhiWithNormalInputFloat = VectorQuaternionRotateVector(InverseR, TranslatedVec);
+					FVector3f Normal;
+					float Penetration;
+			
+					uniform int IMMask = 0;
+					if(IM != 0)
+					{
+						IMMask = lanemask();
+					}
+
+					if(!GetSimplePenetrationWithNormal(PhiWithNormalInputFloat, Type, Ptr, TypedPtr, SizeofFImplicitObject, OffsetofGeometryType, OffsetOfMargin, PerGroupThickness, Penetration, Normal))
+					{
+						const uniform int Mask = IMMask & lanemask();
+						if(Mask != 0)  // Do not call c++ when there is no work to do
+						{
+							float Phi = 0.f;
+							GetPhiWithNormalCollisionParticleRange(CollisionParticlesRange, (const uniform float *uniform)&PhiWithNormalInputFloat, (uniform float *uniform)&Normal, (uniform float *uniform)&Phi, i, programCount, Mask);
+							Penetration = PerGroupThickness - Phi;
+						}
+					}
+
+					const FVector3f X = VectorLoad(&ParticlesX[extract(Index, 0)]);
+
+					if(Penetration > InnerThickness)
+					{
+						const uniform int Mask = IMMask & lanemask();
+						if (Mask != 0)  // Do not call c++ when there is no work to do
+						{
+							ApplySimpleCCDWithFriction(
+								CollisionParticlesRange,
+								i,
+								LastCollisionTransforms[i],
+								(uniform float* uniform)&X,
+								PerGroupThickness,
+								PerGroupFriction,
+								Dt,
+								(uniform float *uniform)&P,
+								programCount,
+								Mask);
+						}
+					}
+					else if(Penetration > 0.0f)
+					{
+						// TransformVector. Scale is 1 so that part can be removed.
+						const FVector3f NormalWorld = VectorQuaternionRotateVector(CollisionR[i], Normal);
+						P = P + (Penetration * NormalWorld);
+
+						ApplyFriction(P, X, NormalWorld, Penetration * PerGroupFriction, Dt, CollisionX[i], CollisionV[i], CollisionW[i]);
+					}
+				}
+
+				if(IM != 0)
+				{
+					VectorStore(&ParticlesPandInvM[extract(Index, 0)], SetVector4(P, IM));
+				}
+			}
+		}
+		else  // No weight map
 		{
-			const uniform SerializedPtr *uniform GeometryPtr = (const uniform SerializedPtr *uniform)Geometry;
-			const uniform uint8 *uniform Ptr = GeometryPtr[i].Ptr;
-			const uniform uint8 *uniform TypedPtr = Ptr + SizeofFImplicitObject;
-			const uniform uint8 Type = Ptr[OffsetofGeometryType] & 0x3F;
+			uniform float InnerThickness = InnerThicknessTable[0];
 
-			if(!IsSimpleCollider(Type))
+			foreach(Index = OuterParticleBatchBegin ... OuterParticleBatchEnd)
 			{
-				continue;
-			}
+				const FVector4f PandInvM = VectorLoad(&ParticlesPandInvM[extract(Index, 0)]);
+				varying FVector3f P;
+				varying float IM;
+				UnzipPandInvM(PandInvM, P, IM);
 
-			// InverseTransformPosition. Can skip Scale calculation because it is 1 for xyz.
-			const FVector3f TranslatedVec = P - CollisionX[i];
-			const uniform FVector4f InverseR = QuatInverse(CollisionR[i]);
-			const FVector3f PhiWithNormalInputFloat = VectorQuaternionRotateVector(InverseR, TranslatedVec);
-			FVector3f Normal;
-			float Penetration;
+				if(all(IM == 0))
+				{
+					continue;
+				}
+
+				for(uniform int i = 0; i < InnerCollisionCount; i++) // SequentialFor
+				{
+					const uniform SerializedPtr *uniform GeometryPtr = (const uniform SerializedPtr *uniform)Geometry;
+					const uniform uint8 *uniform Ptr = GeometryPtr[i].Ptr;
+					const uniform uint8 *uniform TypedPtr = Ptr + SizeofFImplicitObject;
+					const uniform uint8 Type = Ptr[OffsetofGeometryType] & 0x3F;
+
+					if(!IsSimpleCollider(Type))
+					{
+						continue;
+					}
+
+					// InverseTransformPosition. Can skip Scale calculation because it is 1 for xyz.
+					const FVector3f TranslatedVec = P - CollisionX[i];
+					const uniform FVector4f InverseR = QuatInverse(CollisionR[i]);
+					const FVector3f PhiWithNormalInputFloat = VectorQuaternionRotateVector(InverseR, TranslatedVec);
+					FVector3f Normal;
+					float Penetration;
 			
-			if(!GetSimplePenetrationWithNormal(PhiWithNormalInputFloat, Type, Ptr, TypedPtr, SizeofFImplicitObject, OffsetofGeometryType, OffsetOfMargin, PerGroupThickness, Penetration, Normal))
-			{
-				uniform int IMMask;
+					uniform int IMMask = 0;
+					if(IM != 0)
+					{
+						IMMask = lanemask();
+					}
+
+					if(!GetSimplePenetrationWithNormal(PhiWithNormalInputFloat, Type, Ptr, TypedPtr, SizeofFImplicitObject, OffsetofGeometryType, OffsetOfMargin, PerGroupThickness, Penetration, Normal))
+					{
+						const uniform int Mask = IMMask & lanemask();
+						if(Mask != 0)  // Do not call c++ when there is no work to do
+						{
+							float Phi = 0.f;
+							GetPhiWithNormalCollisionParticleRange(CollisionParticlesRange, (const uniform float *uniform)&PhiWithNormalInputFloat, (uniform float *uniform)&Normal, (uniform float *uniform)&Phi, i, programCount, Mask);
+							Penetration = PerGroupThickness - Phi;
+						}
+					}
+
+					const FVector3f X = VectorLoad(&ParticlesX[extract(Index, 0)]);
+
+					if(Penetration > InnerThickness)
+					{
+						const uniform int Mask = IMMask & lanemask();
+						if (Mask != 0)  // Do not call c++ when there is no work to do
+						{
+							ApplySimpleCCDWithFriction(
+								CollisionParticlesRange,
+								i,
+								LastCollisionTransforms[i],
+								(uniform float* uniform)&X,
+								PerGroupThickness,
+								PerGroupFriction,
+								Dt,
+								(uniform float *uniform)&P,
+								programCount,
+								Mask);
+						}
+					}
+					else if(Penetration > 0.0f)
+					{
+						// TransformVector. Scale is 1 so that part can be removed.
+						const FVector3f NormalWorld = VectorQuaternionRotateVector(CollisionR[i], Normal);
+						P = P + (Penetration * NormalWorld);
+
+						ApplyFriction(P, X, NormalWorld, Penetration * PerGroupFriction, Dt, CollisionX[i], CollisionV[i], CollisionW[i]);
+					}
+				}
+
 				if(IM != 0)
 				{
-					IMMask = lanemask();
+					VectorStore(&ParticlesPandInvM[extract(Index, 0)], SetVector4(P, IM));
 				}
-				
-				float Phi;
-				GetPhiWithNormalCollisionParticleRange(CollisionParticlesRange, (const uniform float *uniform)&PhiWithNormalInputFloat, (uniform float *uniform)&Normal, (uniform float *uniform)&Phi, i, programCount, IMMask);
-				Penetration = PerGroupThickness - Phi;
+			}
+		}
+	}
+	else  // No inner thickness
+	{
+		foreach(Index = OuterParticleBatchBegin ... OuterParticleBatchEnd)
+		{
+			const FVector4f PandInvM = VectorLoad(&ParticlesPandInvM[extract(Index, 0)]);
+			varying FVector3f P;
+			varying float IM;
+			UnzipPandInvM(PandInvM, P, IM);
+
+			if(all(IM == 0))
+			{
+				continue;
 			}
 
-			if(Penetration > 0.0f)
+			for(uniform int i = 0; i < InnerCollisionCount; i++) // SequentialFor
 			{
-				// TransformVector. Scale is 1 so that part can be removed.
-				const FVector3f NormalWorld = VectorQuaternionRotateVector(CollisionR[i], Normal);
-				P = P + (Penetration * NormalWorld);
+				const uniform SerializedPtr *uniform GeometryPtr = (const uniform SerializedPtr *uniform)Geometry;
+				const uniform uint8 *uniform Ptr = GeometryPtr[i].Ptr;
+				const uniform uint8 *uniform TypedPtr = Ptr + SizeofFImplicitObject;
+				const uniform uint8 Type = Ptr[OffsetofGeometryType] & 0x3F;
 
-				const FVector3f X = VectorLoad(&ParticlesX[extract(Index, 0)]);
-				ApplyFriction(P, X, NormalWorld, Penetration * PerGroupFriction, Dt, CollisionX[i], CollisionV[i], CollisionW[i]);
+				if(!IsSimpleCollider(Type))
+				{
+					continue;
+				}
+
+				// InverseTransformPosition. Can skip Scale calculation because it is 1 for xyz.
+				const FVector3f TranslatedVec = P - CollisionX[i];
+				const uniform FVector4f InverseR = QuatInverse(CollisionR[i]);
+				const FVector3f PhiWithNormalInputFloat = VectorQuaternionRotateVector(InverseR, TranslatedVec);
+				FVector3f Normal;
+				float Penetration;	
+
+				if(!GetSimplePenetrationWithNormal(PhiWithNormalInputFloat, Type, Ptr, TypedPtr, SizeofFImplicitObject, OffsetofGeometryType, OffsetOfMargin, PerGroupThickness, Penetration, Normal))
+				{
+					uniform int IMMask = 0;
+					if(IM != 0)
+					{
+						IMMask = lanemask();
+					}
+					if(IMMask != 0)  // Do not call c++ when there is no work to do
+					{
+						float Phi = 0.f;
+						GetPhiWithNormalCollisionParticleRange(CollisionParticlesRange, (const uniform float *uniform)&PhiWithNormalInputFloat, (uniform float *uniform)&Normal, (uniform float *uniform)&Phi, i, programCount, IMMask);
+						Penetration = PerGroupThickness - Phi;
+					}
+				}
+
+				if(Penetration > 0.0f)
+				{
+					// TransformVector. Scale is 1 so that part can be removed.
+					const FVector3f NormalWorld = VectorQuaternionRotateVector(CollisionR[i], Normal);
+					P = P + (Penetration * NormalWorld);
+
+					const FVector3f X = VectorLoad(&ParticlesX[extract(Index, 0)]);
+					ApplyFriction(P, X, NormalWorld, Penetration * PerGroupFriction, Dt, CollisionX[i], CollisionV[i], CollisionW[i]);
+				}
 			}
-		}
 
-		if(IM != 0)
-		{
-			VectorStore(&ParticlesPandInvM[extract(Index, 0)], SetVector4(P, IM));
+			if(IM != 0)
+			{
+				VectorStore(&ParticlesPandInvM[extract(Index, 0)], SetVector4(P, IM));
+			}
 		}
 	}
 }
@@ -1353,7 +1566,10 @@ export void ApplyPerParticleSimpleCollisionNoFriction(uniform FVector4f Particle
 												const uniform FVector3f ParticlesX[],
 												const uniform FVector3f CollisionX[],
 												const uniform FVector4f CollisionR[],
+												const uniform FTransform3f LastCollisionTransforms[],
 												const uniform float PerGroupThickness,
+												const uniform float *uniform InnerThicknessTable,
+												const uniform uint8 *uniform InnerThicknessIndices,
 												const uniform uint8 *uniform CollisionParticlesRange,
 												const uniform uint8 *uniform Geometry,
 												const uniform int SizeofFImplicitObject,
@@ -1363,63 +1579,241 @@ export void ApplyPerParticleSimpleCollisionNoFriction(uniform FVector4f Particle
 												const uniform int OuterParticleBatchBegin,
 												const uniform int OuterParticleBatchEnd)
 {
-
-	foreach(Index = OuterParticleBatchBegin ... OuterParticleBatchEnd)
+	const uniform bool bHasInnerThickness = (InnerThicknessTable != NULL);
+	if(bHasInnerThickness)
 	{
-		const FVector4f PandInvM = VectorLoad(&ParticlesPandInvM[extract(Index, 0)]);
-		varying FVector3f P;
-		varying float IM;
-		UnzipPandInvM(PandInvM, P, IM);
-
-		if(all(IM == 0))
+		const uniform bool bHasInnerThicknessWeighmap = (InnerThicknessIndices != NULL);
+		if(bHasInnerThicknessWeighmap)
 		{
-			continue;
+			foreach(Index = OuterParticleBatchBegin ... OuterParticleBatchEnd)
+			{
+				const FVector4f PandInvM = VectorLoad(&ParticlesPandInvM[extract(Index, 0)]);
+				varying FVector3f P;
+				varying float IM;
+				UnzipPandInvM(PandInvM, P, IM);
+
+				if(all(IM == 0))
+				{
+					continue;
+				}
+
+				varying uint8 InnerThicknessIndex = InnerThicknessIndices[Index];
+				#pragma ignore warning(perf)
+				varying float InnerThickness = InnerThicknessTable[InnerThicknessIndex];
+
+				for(uniform int i = 0; i < InnerCollisionCount; i++) // SequentialFor
+				{
+					const uniform SerializedPtr *uniform GeometryPtr = (const uniform SerializedPtr *uniform)Geometry;
+					const uniform uint8 *uniform Ptr = GeometryPtr[i].Ptr;
+					const uniform uint8 *uniform TypedPtr = Ptr + SizeofFImplicitObject;
+					const uniform uint8 Type = Ptr[OffsetofGeometryType] & 0x3F;
+
+					if(!IsSimpleCollider(Type))
+					{
+						continue;
+					}
+
+					// InverseTransformPosition. Can skip Scale calculation because it is 1 for xyz.
+					const FVector3f TranslatedVec = P - CollisionX[i];
+					const uniform FVector4f InverseR = QuatInverse(CollisionR[i]);
+					const FVector3f PhiWithNormalInputFloat = VectorQuaternionRotateVector(InverseR, TranslatedVec);
+					FVector3f Normal;
+					float Penetration;
+
+					uniform int IMMask = 0;
+					if(IM != 0)
+					{
+						IMMask = lanemask();
+					}
+			
+					if(!GetSimplePenetrationWithNormal(PhiWithNormalInputFloat, Type, Ptr, TypedPtr, SizeofFImplicitObject, OffsetofGeometryType, OffsetOfMargin, PerGroupThickness, Penetration, Normal))
+					{
+						const uniform int Mask = IMMask & lanemask();
+						if(Mask != 0)  // Do not call c++ when there is no work to do
+						{
+							float Phi;
+							GetPhiWithNormalCollisionParticleRange(CollisionParticlesRange, (const uniform float *uniform)&PhiWithNormalInputFloat, (uniform float *uniform)&Normal, (uniform float *uniform)&Phi, i, programCount, Mask);
+							Penetration = PerGroupThickness - Phi;
+						}
+					}
+
+					if(Penetration > InnerThickness)
+					{
+						const uniform int Mask = IMMask & lanemask();
+						if (Mask != 0)  // Do not call c++ when there is no work to do
+						{
+							const FVector3f X = VectorLoad(&ParticlesX[extract(Index, 0)]);
+							ApplySimpleCCDNoFriction(
+								CollisionParticlesRange,
+								i,
+								LastCollisionTransforms[i],
+								(uniform float *uniform)&X,
+								PerGroupThickness,
+								(uniform float *uniform)&P,
+								programCount,
+								Mask);
+						}
+					}
+					else if(Penetration > 0.0f)
+					{
+						// TransformVector. Scale is 1 so that part can be removed.
+						const FVector3f NormalWorld = VectorQuaternionRotateVector(CollisionR[i], Normal);
+						P = P + (Penetration * NormalWorld);
+					}
+				}
+
+				if(IM != 0)
+				{
+					VectorStore(&ParticlesPandInvM[extract(Index, 0)], SetVector4(P, IM));
+				}
+			}
 		}
+		else  // No weight map
+		{
+			uniform float InnerThickness = InnerThicknessTable[0];
 
-		for(uniform int i = 0; i < InnerCollisionCount; i++) // SequentialFor
+			foreach(Index = OuterParticleBatchBegin ... OuterParticleBatchEnd)
+			{
+				const FVector4f PandInvM = VectorLoad(&ParticlesPandInvM[extract(Index, 0)]);
+				varying FVector3f P;
+				varying float IM;
+				UnzipPandInvM(PandInvM, P, IM);
+
+				if(all(IM == 0))
+				{
+					continue;
+				}
+
+				for(uniform int i = 0; i < InnerCollisionCount; i++) // SequentialFor
+				{
+					const uniform SerializedPtr *uniform GeometryPtr = (const uniform SerializedPtr *uniform)Geometry;
+					const uniform uint8 *uniform Ptr = GeometryPtr[i].Ptr;
+					const uniform uint8 *uniform TypedPtr = Ptr + SizeofFImplicitObject;
+					const uniform uint8 Type = Ptr[OffsetofGeometryType] & 0x3F;
+
+					if(!IsSimpleCollider(Type))
+					{
+						continue;
+					}
+
+					// InverseTransformPosition. Can skip Scale calculation because it is 1 for xyz.
+					const FVector3f TranslatedVec = P - CollisionX[i];
+					const uniform FVector4f InverseR = QuatInverse(CollisionR[i]);
+					const FVector3f PhiWithNormalInputFloat = VectorQuaternionRotateVector(InverseR, TranslatedVec);
+					FVector3f Normal;
+					float Penetration;
+
+					uniform int IMMask = 0;
+					if(IM != 0)
+					{
+						IMMask = lanemask();
+					}
+			
+					if(!GetSimplePenetrationWithNormal(PhiWithNormalInputFloat, Type, Ptr, TypedPtr, SizeofFImplicitObject, OffsetofGeometryType, OffsetOfMargin, PerGroupThickness, Penetration, Normal))
+					{
+						const uniform int Mask = IMMask & lanemask();
+						if(Mask != 0)  // Do not call c++ when there is no work to do
+						{
+							float Phi;
+							GetPhiWithNormalCollisionParticleRange(CollisionParticlesRange, (const uniform float *uniform)&PhiWithNormalInputFloat, (uniform float *uniform)&Normal, (uniform float *uniform)&Phi, i, programCount, Mask);
+							Penetration = PerGroupThickness - Phi;
+						}
+					}
+
+					if(Penetration > InnerThickness)
+					{
+						const uniform int Mask = IMMask & lanemask();
+						if (Mask != 0)  // Do not call c++ when there is no work to do
+						{
+							const FVector3f X = VectorLoad(&ParticlesX[extract(Index, 0)]);
+							ApplySimpleCCDNoFriction(
+								CollisionParticlesRange,
+								i,
+								LastCollisionTransforms[i],
+								(uniform float *uniform)&X,
+								PerGroupThickness,
+								(uniform float *uniform)&P,
+								programCount,
+								Mask);
+						}
+					}
+					else if(Penetration > 0.0f)
+					{
+						// TransformVector. Scale is 1 so that part can be removed.
+						const FVector3f NormalWorld = VectorQuaternionRotateVector(CollisionR[i], Normal);
+						P = P + (Penetration * NormalWorld);
+					}
+				}
+
+				if(IM != 0)
+				{
+					VectorStore(&ParticlesPandInvM[extract(Index, 0)], SetVector4(P, IM));
+				}
+			}
+		}
+	}
+	else  // No inner thickness
+	{
+		foreach(Index = OuterParticleBatchBegin ... OuterParticleBatchEnd)
 		{
-			const uniform SerializedPtr *uniform GeometryPtr = (const uniform SerializedPtr *uniform)Geometry;
-			const uniform uint8 *uniform Ptr = GeometryPtr[i].Ptr;
-			const uniform uint8 *uniform TypedPtr = Ptr + SizeofFImplicitObject;
-			const uniform uint8 Type = Ptr[OffsetofGeometryType] & 0x3F;
+			const FVector4f PandInvM = VectorLoad(&ParticlesPandInvM[extract(Index, 0)]);
+			varying FVector3f P;
+			varying float IM;
+			UnzipPandInvM(PandInvM, P, IM);
 
-			if(!IsSimpleCollider(Type))
+			if(all(IM == 0))
 			{
 				continue;
 			}
 
-			// InverseTransformPosition. Can skip Scale calculation because it is 1 for xyz.
-			const FVector3f TranslatedVec = P - CollisionX[i];
-			const uniform FVector4f InverseR = QuatInverse(CollisionR[i]);
-			const FVector3f PhiWithNormalInputFloat = VectorQuaternionRotateVector(InverseR, TranslatedVec);
-			FVector3f Normal;
-			float Penetration;
-			
-			if(!GetSimplePenetrationWithNormal(PhiWithNormalInputFloat, Type, Ptr, TypedPtr, SizeofFImplicitObject, OffsetofGeometryType, OffsetOfMargin, PerGroupThickness, Penetration, Normal))
+			for(uniform int i = 0; i < InnerCollisionCount; i++) // SequentialFor
 			{
-				uniform int IMMask;
-				if(IM != 0)
+				const uniform SerializedPtr *uniform GeometryPtr = (const uniform SerializedPtr *uniform)Geometry;
+				const uniform uint8 *uniform Ptr = GeometryPtr[i].Ptr;
+				const uniform uint8 *uniform TypedPtr = Ptr + SizeofFImplicitObject;
+				const uniform uint8 Type = Ptr[OffsetofGeometryType] & 0x3F;
+
+				if(!IsSimpleCollider(Type))
 				{
-					IMMask = lanemask();
+					continue;
+				}
+
+				// InverseTransformPosition. Can skip Scale calculation because it is 1 for xyz.
+				const FVector3f TranslatedVec = P - CollisionX[i];
+				const uniform FVector4f InverseR = QuatInverse(CollisionR[i]);
+				const FVector3f PhiWithNormalInputFloat = VectorQuaternionRotateVector(InverseR, TranslatedVec);
+				FVector3f Normal;
+				float Penetration;
+
+			
+				if(!GetSimplePenetrationWithNormal(PhiWithNormalInputFloat, Type, Ptr, TypedPtr, SizeofFImplicitObject, OffsetofGeometryType, OffsetOfMargin, PerGroupThickness, Penetration, Normal))
+				{
+					uniform int IMMask = 0;
+					if(IM != 0)
+					{
+						IMMask = lanemask();
+					}
+					if(IMMask != 0)  // Do not call c++ when there is no work to do
+					{
+						float Phi;
+						GetPhiWithNormalCollisionParticleRange(CollisionParticlesRange, (const uniform float *uniform)&PhiWithNormalInputFloat, (uniform float *uniform)&Normal, (uniform float *uniform)&Phi, i, programCount, IMMask);
+						Penetration = PerGroupThickness - Phi;
+					}
+				}
+
+				if(Penetration > 0.0f)
+				{
+					// TransformVector. Scale is 1 so that part can be removed.
+					const FVector3f NormalWorld = VectorQuaternionRotateVector(CollisionR[i], Normal);
+					P = P + (Penetration * NormalWorld);
 				}
-				
-				float Phi;
-				GetPhiWithNormalCollisionParticleRange(CollisionParticlesRange, (const uniform float *uniform)&PhiWithNormalInputFloat, (uniform float *uniform)&Normal, (uniform float *uniform)&Phi, i, programCount, IMMask);
-				Penetration = PerGroupThickness - Phi;
 			}
 
-			if(Penetration > 0.0f)
+			if(IM != 0)
 			{
-				// TransformVector. Scale is 1 so that part can be removed.
-				const FVector3f NormalWorld = VectorQuaternionRotateVector(CollisionR[i], Normal);
-				P = P + (Penetration * NormalWorld);
+				VectorStore(&ParticlesPandInvM[extract(Index, 0)], SetVector4(P, IM));
 			}
 		}
-
-		if(IM != 0)
-		{
-			VectorStore(&ParticlesPandInvM[extract(Index, 0)], SetVector4(P, IM));
-		}
 	}
 }
 
diff --git a/Engine/Source/Runtime/Experimental/Chaos/Public/Chaos/PBDSoftBodyCollisionConstraint.h b/Engine/Source/Runtime/Experimental/Chaos/Public/Chaos/PBDSoftBodyCollisionConstraint.h
--- a/Engine/Source/Runtime/Experimental/Chaos/Public/Chaos/PBDSoftBodyCollisionConstraint.h
+++ b/Engine/Source/Runtime/Experimental/Chaos/Public/Chaos/PBDSoftBodyCollisionConstraint.h
@@ -1,15 +1,16 @@
 // Copyright Epic Games, Inc. All Rights Reserved.
 #pragma once
 
+#include "Chaos/CollectionPropertyFacade.h"
 #include "Chaos/Core.h"
 #include "Chaos/PBDKinematicTriangleMeshCollisions.h"
 #include "Chaos/PBDSoftsEvolutionFwd.h"
 #include "Chaos/PBDSoftsSolverParticles.h"
+#include "Chaos/PBDWeightMap.h"
 #include "Chaos/PBDPlanarConstraints.h"
 #include "Chaos/SoftsSolverCollisionParticles.h"
 #include "Chaos/SoftsSolverCollisionParticlesRange.h"
 #include "Chaos/SoftsSolverParticlesRange.h"
-#include "Chaos/CollectionPropertyFacade.h"
 
 #if !defined(CHAOS_PER_PARTICLE_COLLISION_ISPC_ENABLED_DEFAULT)
 #define CHAOS_PER_PARTICLE_COLLISION_ISPC_ENABLED_DEFAULT 1
@@ -59,6 +60,7 @@ public:
 		const TMap<FParticleRangeIndex, FPBDComplexColliderBoneData>& InComplexBoneData = TMap<FParticleRangeIndex, FPBDComplexColliderBoneData>()
 		)
 		: LastCollisionTransforms(InLastCollisionTransforms)
+		, InnerCollisionThickness(FSolverVec2::ZeroVector)
 		, CollisionThickness(InCollisionThickness)
 		, SoftBodyCollisionThickness(InSoftBodyCollisionThickness)
 		, FrictionCoefficient(InFrictionCoefficient)
@@ -115,6 +117,7 @@ private:
 
 protected:
 	const TArray<FSolverRigidTransform3>& LastCollisionTransforms; // Used by CCD
+	FPBDWeightMap InnerCollisionThickness;
 	FSolverReal CollisionThickness;
 	FSolverReal SoftBodyCollisionThickness;
 	FSolverReal FrictionCoefficient;
@@ -174,7 +177,7 @@ public:
 	static constexpr FSolverReal DefaultFrictionCoefficient = (FSolverReal)0.8;
 	static constexpr FSolverReal DefaultProximityStiffness = (FSolverReal)100.;
 
-	FPBDSoftBodyCollisionConstraint(
+	CHAOS_API FPBDSoftBodyCollisionConstraint(
 		const TArray<FSolverRigidTransform3>& InLastCollisionTransforms,
 		const FCollectionPropertyConstFacade& PropertyCollection,
 		FSolverReal InMeshScale,
@@ -182,31 +185,8 @@ public:
 		TArray<FSolverVec3>* InContacts = nullptr,
 		TArray<FSolverVec3>* InNormals = nullptr,
 		TArray<FSolverReal>* InPhis = nullptr,
-		const TMap<FParticleRangeIndex, FPBDComplexColliderBoneData>& InComplexBoneData = TMap<FParticleRangeIndex, FPBDComplexColliderBoneData>())
-		: Base(InLastCollisionTransforms,
-			InMeshScale* GetCollisionThickness(PropertyCollection, DefaultCollisionThickness),
-			GetFrictionCoefficient(PropertyCollection, DefaultFrictionCoefficient),
-			GetUseCCD(PropertyCollection, false),
-			GetProximityStiffness(PropertyCollection, DefaultProximityStiffness),
-			InCollisionParticleCollided, InContacts, InNormals, InPhis,
-			GetSoftBodyCollisionThickness(PropertyCollection, DefaultSoftBodyCollisionThickness),
-			GetEnableSimpleColliders(PropertyCollection, true),
-			GetEnableComplexColliders(PropertyCollection, true),
-			GetUsePlanarConstraintForSimpleColliders(PropertyCollection, bDefaultUsePlanarConstraintForSimpleColliders),
-			GetUsePlanarConstraintForComplexColliders(PropertyCollection, bDefaultUsePlanarConstraintForComplexColliders),
-			InComplexBoneData
-		)
-		, MeshScale(InMeshScale)
-		, CollisionThicknessIndex(PropertyCollection)
-		, SoftBodyCollisionThicknessIndex(PropertyCollection)
-		, FrictionCoefficientIndex(PropertyCollection)
-		, UseCCDIndex(PropertyCollection)
-		, ProximityStiffnessIndex(PropertyCollection)
-		, EnableSimpleCollidersIndex(PropertyCollection)
-		, UsePlanarConstraintForSimpleCollidersIndex(PropertyCollection)
-		, EnableComplexCollidersIndex(PropertyCollection)
-		, UsePlanarConstraintForComplexCollidersIndex(PropertyCollection)
-	{}
+		const TMap<FParticleRangeIndex, FPBDComplexColliderBoneData>& InComplexBoneData = TMap<FParticleRangeIndex, FPBDComplexColliderBoneData>(),
+		const TMap<FString, TConstArrayView<float>>& WeightMaps = TMap<FString, TConstArrayView<float>>());
 
 	CHAOS_API void SetProperties(const FCollectionPropertyConstFacade& PropertyCollection);
 
@@ -222,6 +202,8 @@ private:
 	UE_CHAOS_DECLARE_PROPERTYCOLLECTION_NAME(UsePlanarConstraintForSimpleColliders, bool);
 	UE_CHAOS_DECLARE_PROPERTYCOLLECTION_NAME(EnableComplexColliders, bool);
 	UE_CHAOS_DECLARE_PROPERTYCOLLECTION_NAME(UsePlanarConstraintForComplexColliders, bool);
+
+	UE_CHAOS_DECLARE_INDEXLESS_PROPERTYCOLLECTION_NAME(InnerCollisionThickness, float);
 };
 
 }  // End namespace Chaos::Softs
