Removed per HTTP category CSV stats.

#rnx

#Synced-CL 51660154

diff --git a/Engine/Source/Runtime/Experimental/IoStore/OnDemand/Private/Statistics.cpp b/Engine/Source/Runtime/Experimental/IoStore/OnDemand/Private/Statistics.cpp
--- a/Engine/Source/Runtime/Experimental/IoStore/OnDemand/Private/Statistics.cpp
+++ b/Engine/Source/Runtime/Experimental/IoStore/OnDemand/Private/Statistics.cpp
@@ -486,31 +486,19 @@ static uint32			GHttpCdnCacheUnknown = 0;
 ////////////////////////////////////////////////////////////////////////////////
 // CSV STATS
 
-#if CSV_PROFILER && !CSV_PROFILER_MINIMAL
-	#define UE_IAX_CSV_DEFINE_STAT(StatName) FCsvDeclaredStat _GCsvStat_##StatName[(uint8)EHttpRequestType::NUM_SOURCES] = {{TEXT(#StatName),CSV_CATEGORY_INDEX(Ias)},{TEXT(#StatName), CSV_CATEGORY_INDEX(Iad)}}
-	#define UE_IAX_CSV_CUSTOM_STAT_DEFINED(RequestType, StatName, Value, Op) FCsvProfiler::RecordCustomStat(_GCsvStat_##StatName[(uint8)RequestType].Name, _GCsvStat_##StatName[(uint8)RequestType].CategoryIndex, Value, Op);
-#else
-	#define UE_IAX_CSV_DEFINE_STAT(StatName)
-	#define UE_IAX_CSV_CUSTOM_STAT_DEFINED(RequestType, StatName, Value, Op)
-#endif
-
 CSV_DEFINE_CATEGORY(Ias, true);
 CSV_DEFINE_CATEGORY(Iad, true);
-// http stat totals
-UE_IAX_CSV_DEFINE_STAT(HttpGetCount);
-UE_IAX_CSV_DEFINE_STAT(HttpRetryCount);
-UE_IAX_CSV_DEFINE_STAT(HttpCancelCount);
-UE_IAX_CSV_DEFINE_STAT(HttpErrorCount);
-UE_IAX_CSV_DEFINE_STAT(HttpPendingCount);
-UE_IAX_CSV_DEFINE_STAT(HttpDownloadedMB);
-UE_IAX_CSV_DEFINE_STAT(HttpBandwidthMbps);
-UE_IAX_CSV_DEFINE_STAT(HttpDurationMsAvg);
-UE_IAX_CSV_DEFINE_STAT(HttpDurationMsMax);
 
 ////////////////////////////////////////////////////////////////////////////////
 // IAX CSV Counters
 CSV_DEFINE_CATEGORY(Iax, true);
 
+CSV_DEFINE_STAT(Iax, HttpGetCount);
+CSV_DEFINE_STAT(Iax, HttpGetMiB);
+CSV_DEFINE_STAT(Iax, HttpRetryCount);
+CSV_DEFINE_STAT(Iax, HttpErrorCount);
+CSV_DEFINE_STAT(Iax, HttpDurationAvgMs);
+
 CSV_DEFINE_STAT(Iax, InstallerInstallCount);
 CSV_DEFINE_STAT(Iax, InstallerInflightCount);
 CSV_DEFINE_STAT(Iax, InstallerDownloadedMiB);
@@ -527,6 +515,29 @@ CSV_DEFINE_STAT(Iax, StreamingCacheWriteCount);
 CSV_DEFINE_STAT(Iax, StreamingCacheWriteMiB);
 CSV_DEFINE_STAT(Iax, StreamingCacheDecodeErrorCount);
 
+////////////////////////////////////////////////////////////////////////////////
+struct FHttpCounters
+{
+	FHttpCounters()
+		: GetCount(TEXT("Iax/Http/GetCount"), TraceCounterDisplayHint_None)
+		, GetBytes(TEXT("Iax/Http/GetBytes"), TraceCounterDisplayHint_Memory)
+		, CancelCount(TEXT("Iax/Http/CancelCount"), TraceCounterDisplayHint_None)
+		, RetryCount(TEXT("Iax/Http/RetryCount"), TraceCounterDisplayHint_None)
+		, ErrorCount(TEXT("Iax/Http/ErrorCount"), TraceCounterDisplayHint_None)
+		, DurationTotalMs(TEXT("Iax/Http/DurationTotalMs"), TraceCounterDisplayHint_None)
+		, DurationAvgMs(TEXT("Iax/Http/DurationAvgMs"), TraceCounterDisplayHint_None)
+	{ }
+
+	FCounterAtomicInt		GetCount;
+	FCounterAtomicInt 		GetBytes;
+	FCounterAtomicInt 		CancelCount;
+	FCounterAtomicInt 		RetryCount;
+	FCounterAtomicInt 		ErrorCount;
+	FCounterAtomicInt 		DurationTotalMs;
+	FCounterAtomicInt 		DurationAvgMs;
+};
+FHttpCounters HttpCounters;
+
 ////////////////////////////////////////////////////////////////////////////////
 struct FInstallerCounters
 {
@@ -736,21 +747,14 @@ void FOnDemandIoBackendStats::UpdateCSVValues()
 		CSV_CUSTOM_STAT_DEFINED(InstallCacheAvgCacheHitPct, int32(FMath::RoundToInt64(CacheHitAvgPct)), ECsvCustomStatOp::Set);
 	}
 
-	// http stat totals
-	auto HttpCsv = [](EHttpRequestType Type)
-		{
-			UE_IAX_CSV_CUSTOM_STAT_DEFINED(Type, HttpGetCount, (int32)GHttpGetCount[(uint8)Type].Get(), ECsvCustomStatOp::Set);
-			UE_IAX_CSV_CUSTOM_STAT_DEFINED(Type, HttpCancelCount, (int32)GHttpCancelCount[(uint8)Type].Get(), ECsvCustomStatOp::Set);
-			UE_IAX_CSV_CUSTOM_STAT_DEFINED(Type, HttpErrorCount, (int32)GHttpErrorCount[(uint8)Type].Get(), ECsvCustomStatOp::Set);
-			UE_IAX_CSV_CUSTOM_STAT_DEFINED(Type, HttpPendingCount, (int32)GHttpPendingCount[(uint8)Type].Get(), ECsvCustomStatOp::Set);
-			UE_IAX_CSV_CUSTOM_STAT_DEFINED(Type, HttpDownloadedMB, (float)GHttpDownloadedBytes[(uint8)Type].Get() * OneOver1024 * OneOver1024, ECsvCustomStatOp::Set);
-			UE_IAX_CSV_CUSTOM_STAT_DEFINED(Type, HttpBandwidthMbps, GHttpBandwidthMbps[(uint8)Type].Get(), ECsvCustomStatOp::Set);
-			UE_IAX_CSV_CUSTOM_STAT_DEFINED(Type, HttpDurationMsAvg, (int32)GHttpDurationMsAvg[(uint8)Type], ECsvCustomStatOp::Set);
-			UE_IAX_CSV_CUSTOM_STAT_DEFINED(Type, HttpDurationMsMax, GHttpDurationMsMax[(uint8)Type], ECsvCustomStatOp::Set);
-		};
-
-	HttpCsv(EHttpRequestType::Streaming);
-	HttpCsv(EHttpRequestType::Installed);
+	// HTTP stats
+	{
+		CSV_CUSTOM_STAT_DEFINED(HttpGetCount, int32(HttpCounters.GetCount.Get()), ECsvCustomStatOp::Set);
+		CSV_CUSTOM_STAT_DEFINED(HttpGetMiB, BytesToApproxMB(HttpCounters.GetBytes.Get()), ECsvCustomStatOp::Set);
+		CSV_CUSTOM_STAT_DEFINED(HttpRetryCount, int32(HttpCounters.RetryCount.Get()), ECsvCustomStatOp::Set);
+		CSV_CUSTOM_STAT_DEFINED(HttpErrorCount, int32(HttpCounters.ErrorCount.Get()), ECsvCustomStatOp::Set);
+		CSV_CUSTOM_STAT_DEFINED(HttpDurationAvgMs, int32(HttpCounters.DurationAvgMs.Get()), ECsvCustomStatOp::Set);
+	}
 }
 
 void FOnDemandIoBackendStats::PrintPeriodicLogging()
@@ -1508,11 +1512,17 @@ void FOnDemandIoBackendStats::OnHttpCompleted(
 
 	if (RetryCount > 0)
 	{
+		HttpCounters.RetryCount.Add(RetryCount);
 		GHttpRetryCount[Index].Increment();
 	}
 
 	if (bHttpOk)
 	{
+		HttpCounters.GetCount.Increment();
+		HttpCounters.GetBytes.Add(SizeBytes);
+		HttpCounters.DurationTotalMs.Add(DurationMs);
+		HttpCounters.DurationAvgMs.Set(int64(double(HttpCounters.DurationTotalMs.Get()) / double(HttpCounters.GetCount.Get())));
+
 		GHttpGetCount[Index].Increment();
 		GHttpDownloadedBytes[Index].Add(SizeBytes);
 		GHttpDurationMsSum[Index] += DurationMs;
@@ -1534,10 +1544,12 @@ void FOnDemandIoBackendStats::OnHttpCompleted(
 	}
 	else if (bCanceled)
 	{
+		HttpCounters.CancelCount.Increment();
 		GHttpCancelCount[Index].Increment();
 	}
 	else
 	{
+		HttpCounters.ErrorCount.Increment();
 		GHttpErrorCount[Index].Increment();
 	}
 }
