Fix for Ensure related to 0 value normals on degenerate triangles.

diff --git a/Engine/Source/Runtime/Experimental/Chaos/Public/Chaos/ImplicitObjectScaled.h b/Engine/Source/Runtime/Experimental/Chaos/Public/Chaos/ImplicitObjectScaled.h
--- a/Engine/Source/Runtime/Experimental/Chaos/Public/Chaos/ImplicitObjectScaled.h
+++ b/Engine/Source/Runtime/Experimental/Chaos/Public/Chaos/ImplicitObjectScaled.h
@@ -698,6 +698,13 @@ public:
 						OutTime = OutScaledTime;
 						LocalPosition = OutScaledPosition;
 						LocalNormal = OutScaledNormal.GetSafeNormal();
+						if (LocalNormal.IsNearlyZero())
+						{
+							// Fallback: use negated sweep direction when GJK produces degenerate normal
+							UE_LOG(LogChaos, Warning, TEXT("LowLevelSweepGeom: OutScaledNormal produced zero normal after GetSafeNormal. OutScaledNormal: (%f, %f, %f), MScale: (%f, %f, %f)"),
+								OutScaledNormal[0], OutScaledNormal[1], OutScaledNormal[2], MScale[0], MScale[1], MScale[2]);
+							LocalNormal = (-LocalDir).GetSafeNormal();
+						}
 						return true;
 					}
 				}
