Dataflow : Node varioable now have desription and category properties
- category will move the property in the override to a category
- description will set the tooltip on the property

#Synced-CL 51713790

diff --git a/Engine/Source/Runtime/Experimental/Dataflow/Engine/Private/Dataflow/DataflowVariableNodes.cpp b/Engine/Source/Runtime/Experimental/Dataflow/Engine/Private/Dataflow/DataflowVariableNodes.cpp
--- a/Engine/Source/Runtime/Experimental/Dataflow/Engine/Private/Dataflow/DataflowVariableNodes.cpp
+++ b/Engine/Source/Runtime/Experimental/Dataflow/Engine/Private/Dataflow/DataflowVariableNodes.cpp
@@ -17,6 +17,12 @@ namespace UE::Dataflow
 	}
 }
 
+namespace GetDataflowVariableNode::Private
+{
+	static const FName EnableCategoriesName = TEXT("EnableCategories");
+	static const FName CategoryMetaDataName = TEXT("Category");
+	static const FName TooltipMetaDataName = TEXT("ToolTip");
+}
 
 FGetDataflowVariableNode::FGetDataflowVariableNode(const UE::Dataflow::FNodeParameters& InParam, FGuid InGuid)
 	: FDataflowNode(InParam, InGuid)
@@ -110,6 +116,10 @@ void FGetDataflowVariableNode::SetVariable(UDataflow* Dataflow, FName InVariable
 		if (const FPropertyBagPropertyDesc* Desc = Dataflow->Variables.FindPropertyDescByName(InVariableName))
 		{
 			VariableName = InVariableName;
+#if WITH_EDITOR
+			Category = Desc->GetMetaData(GetDataflowVariableNode::Private::CategoryMetaDataName);
+			Description = Desc->GetMetaData(GetDataflowVariableNode::Private::TooltipMetaDataName);
+#endif
 
 			VariablePropertyBag.AddProperties({ *Desc });
 			VariablePropertyBag.SetValue(VariableName, Desc->CachedProperty, Dataflow->Variables.GetValue().GetMemory());
@@ -213,12 +223,60 @@ void FGetDataflowVariableNode::UpdateOutputTypes(const FPropertyBagPropertyDesc&
 	}
 }
 
+void FGetDataflowVariableNode::NotifyAssetVariablesChanged(UDataflow* DataflowAsset)
+{
+	if (DataflowAsset)
+	{
+		DataflowAsset->Modify();
+		FDataflowAssetDelegates::OnVariablesChanged.Broadcast(DataflowAsset, VariableName);
+
+		FPropertyChangedEvent PropertyChangedEvent(nullptr);
+		if (UClass* DataflowClass = DataflowAsset->GetClass())
+		{
+			FProperty* MemberProperty = DataflowClass->FindPropertyByName(GET_MEMBER_NAME_CHECKED(UDataflow, Variables));
+			PropertyChangedEvent.SetActiveMemberProperty(MemberProperty);
+		}
+#if WITH_EDITOR
+		DataflowAsset->PostEditChangeProperty(PropertyChangedEvent);
+#endif
+	}
+}
+
+void FGetDataflowVariableNode::SetVariableMetaData(FName MetaDataName, const FString& InValue)
+{
+#if WITH_EDITOR
+	if (const TStrongObjectPtr<UDataflow> DataflowAsset = WeakDataflowPtr.Pin())
+	{
+		if (const UPropertyBag* PropertyBag = DataflowAsset->Variables.GetPropertyBagStruct())
+		{
+			TArray<FPropertyBagPropertyDesc> Descs(PropertyBag->GetPropertyDescs());
+			if (FPropertyBagPropertyDesc* PropertyDesc = Descs.FindByPredicate([this](const FPropertyBagPropertyDesc& Desc) { return Desc.Name == VariableName; }))
+			{
+				if (MetaDataName == GetDataflowVariableNode::Private::CategoryMetaDataName)
+				{
+					PropertyDesc->SetMetaData(GetDataflowVariableNode::Private::EnableCategoriesName, FString());
+				}
+				PropertyDesc->SetMetaData(MetaDataName, InValue);
+
+				// Create new bag struct
+				const UPropertyBag* NewBagStruct = UPropertyBag::GetOrCreateFromDescs(Descs);
+				DataflowAsset->Variables.MigrateToNewBagStruct(NewBagStruct);
+
+				NotifyAssetVariablesChanged(DataflowAsset.Get());
+			}
+		}
+	}
+#endif
+}
+
 void FGetDataflowVariableNode::OnPropertyChanged(UE::Dataflow::FContext& Context, const FPropertyChangedEvent& InPropertyChangedEvent)
 {
 #if WITH_EDITOR
 	const FName PropertyName = InPropertyChangedEvent.GetPropertyName();
 	const FName MemberPropertyName = InPropertyChangedEvent.GetMemberPropertyName();
 	const FName VariablePropertyBagMemberName = GET_MEMBER_NAME_CHECKED(FGetDataflowVariableNode, VariablePropertyBag);
+	const FName CategoryMemberName = GET_MEMBER_NAME_CHECKED(FGetDataflowVariableNode, Category);
+	const FName DescriptionMemberName = GET_MEMBER_NAME_CHECKED(FGetDataflowVariableNode, Description);
 
 	if (PropertyName == VariablePropertyBagMemberName || MemberPropertyName == VariablePropertyBagMemberName)
 	{
@@ -232,22 +290,21 @@ void FGetDataflowVariableNode::OnPropertyChanged(UE::Dataflow::FContext& Context
 					EPropertyBagResult Result = DataflowAsset->Variables.SetValue(VariableName, SourceDesc->CachedProperty, VariablePropertyBag.GetValue().GetMemory());
 					if (Result == EPropertyBagResult::Success)
 					{
-						DataflowAsset->Modify();
-						FDataflowAssetDelegates::OnVariablesChanged.Broadcast(DataflowAsset.Get(), VariableName);
-
-						FPropertyChangedEvent PropertyChangedEvent(nullptr);
-						if (UClass* DataflowClass = DataflowAsset->GetClass())
-						{
-							FProperty* MemberProperty = DataflowClass->FindPropertyByName(GET_MEMBER_NAME_CHECKED(UDataflow, Variables));
-							PropertyChangedEvent.SetActiveMemberProperty(MemberProperty);
-						}
-						DataflowAsset->PostEditChangeProperty(PropertyChangedEvent);
+						NotifyAssetVariablesChanged(DataflowAsset.Get());
 						Invalidate();
 					}
 				}
 			}
 		}
 	}
+	else if (PropertyName == CategoryMemberName || MemberPropertyName == CategoryMemberName)
+	{
+		SetVariableMetaData(GetDataflowVariableNode::Private::CategoryMetaDataName, Category);
+	}
+	else if (PropertyName == DescriptionMemberName || MemberPropertyName == DescriptionMemberName)
+	{
+		SetVariableMetaData(GetDataflowVariableNode::Private::TooltipMetaDataName, Description);
+	}
 #endif
 }
 
diff --git a/Engine/Source/Runtime/Experimental/Dataflow/Engine/Public/Dataflow/DataflowVariableNodes.h b/Engine/Source/Runtime/Experimental/Dataflow/Engine/Public/Dataflow/DataflowVariableNodes.h
--- a/Engine/Source/Runtime/Experimental/Dataflow/Engine/Public/Dataflow/DataflowVariableNodes.h
+++ b/Engine/Source/Runtime/Experimental/Dataflow/Engine/Public/Dataflow/DataflowVariableNodes.h
@@ -60,6 +60,14 @@ private:
 	UPROPERTY(VisibleAnywhere, Category = Variable )
 	FName VariableName;
 
+	/** Category the variable should be displayed under */
+	UPROPERTY(EditAnywhere, Category = Variable)
+	FString Category;
+
+	/** Decsription of this variable, the description will show in the tooltip of overrides */
+	UPROPERTY(EditAnywhere, Category = Variable)
+	FString Description;
+
 	void RegisterHandlers();
 	void UnregisterHandlers();
 
@@ -69,6 +77,8 @@ private:
 	EPropertyBagContainerType ContainerType = EPropertyBagContainerType::None;
 
 	void OnObjectPropertyChanged(UObject* InObject, FPropertyChangedEvent& InPropertyChangedEvent);
+	void NotifyAssetVariablesChanged(UDataflow* DataflowAsset);
+	void SetVariableMetaData(FName MetaDataName, const FString& Value);
 
 	TWeakObjectPtr<UDataflow> WeakDataflowPtr;
 
