[NetPhysics] Improve logic around resimulation with objects that are pending deletion along with some fixes:
- Don't remove particle from RewindData::DirtyParticles when calling UnregisterObject
- Call the UnregisterObject physics command lambda again during resim to disable the particle during resim since rewind might have enabled it again.
- Don't cache Pre-Resim state for particles that are disabled or pending delete when resim starts since we will not be rendering after the resim
- Add in leniency when checking if a particle should be considered desynced during resimulation, else all networked particles will be considered desynced due to the loss in precision from the quantized FRepMovement state.
- Don't ensure when ActorComponent is not valid during legacy flow, it can happen when being destroyed.
- Update object state in bulk during ApplyTargets and StepNonResimParticles (to update SoA's less)

#Synced-CL 51746833

diff --git a/Engine/Source/Runtime/Engine/Private/PhysicsEngine/NetworkPhysicsComponent.cpp b/Engine/Source/Runtime/Engine/Private/PhysicsEngine/NetworkPhysicsComponent.cpp
--- a/Engine/Source/Runtime/Engine/Private/PhysicsEngine/NetworkPhysicsComponent.cpp
+++ b/Engine/Source/Runtime/Engine/Private/PhysicsEngine/NetworkPhysicsComponent.cpp
@@ -2321,7 +2321,7 @@ void FAsyncNetworkPhysicsComponent::ConsumeAsyncInput_Internal(const int32 Physi
 				{
 					if (bIsUsingLegacyData)
 					{
-						if (!ComponentSettings.GetValidateDataOnGameThread() && ensure(ActorComponent.IsValid() && ActorComponent.Get()->IsBeingDestroyed() == false))
+						if (!ComponentSettings.GetValidateDataOnGameThread() && ActorComponent.IsValid() && ActorComponent.Get()->IsBeingDestroyed() == false)
 						{
 							AsyncInput->InputData->ValidateDataInHistory(ActorComponent.Get());
 						}
@@ -2479,7 +2479,7 @@ void FAsyncNetworkPhysicsComponent::OnPreProcessInputs_Internal(const int32 Phys
 			{
 				if (bIsUsingLegacyData)
 				{
-					if (ensure(ActorComponent.IsValid() && ActorComponent.Get()->IsBeingDestroyed() == false))
+					if (ActorComponent.IsValid() && ActorComponent.Get()->IsBeingDestroyed() == false)
 					{
 						FNetworkPhysicsData* LegacyData = static_cast<FNetworkPhysicsData*>(PhysicsData);
 						LegacyData->ApplyData(ActorComponent.Get());
@@ -2526,7 +2526,7 @@ void FAsyncNetworkPhysicsComponent::OnPreProcessInputs_Internal(const int32 Phys
 #endif
 					if (bIsUsingLegacyData)
 					{
-						if (ensure(ActorComponent.IsValid() && ActorComponent.Get()->IsBeingDestroyed() == false))
+						if (ActorComponent.IsValid() && ActorComponent.Get()->IsBeingDestroyed() == false)
 						{
 							// Apply all inputs in range
 							InputHistory->ApplyDataRange(NewImportantInputFrame, NextExpectedInputFrame - 1, ActorComponent.Get(), /*bOnlyImportant*/false);
@@ -2630,7 +2630,7 @@ void FAsyncNetworkPhysicsComponent::OnPreProcessInputs_Internal(const int32 Phys
 
 								if (bIsUsingLegacyData)
 								{
-									if (ensure(ActorComponent.IsValid() && ActorComponent.Get()->IsBeingDestroyed() == false))
+									if (ActorComponent.IsValid() && ActorComponent.Get()->IsBeingDestroyed() == false)
 									{
 										FNetworkPhysicsData* LegacyData = static_cast<FNetworkPhysicsData*>(PhysicsData);
 										LegacyData->ApplyData(ActorComponent.Get());
@@ -2674,7 +2674,7 @@ void FAsyncNetworkPhysicsComponent::OnPreProcessInputs_Internal(const int32 Phys
 
 				if (bIsUsingLegacyData)
 				{
-					if (ensure(ActorComponent.IsValid() && ActorComponent.Get()->IsBeingDestroyed() == false))
+					if (ActorComponent.IsValid() && ActorComponent.Get()->IsBeingDestroyed() == false)
 					{
 						FNetworkPhysicsData* LegacyData = static_cast<FNetworkPhysicsData*>(PhysicsData);
 						LegacyData->ApplyData(ActorComponent.Get());
@@ -2745,7 +2745,7 @@ void FAsyncNetworkPhysicsComponent::OnPostProcessInputs_Internal(const int32 Phy
 
 		if (bIsUsingLegacyData)
 		{
-			if (ensure(ActorComponent.IsValid() && ActorComponent.Get()->IsBeingDestroyed() == false))
+			if (ActorComponent.IsValid() && ActorComponent.Get()->IsBeingDestroyed() == false)
 			{
 				// Gather input data from implementation
 				FNetworkPhysicsData* LegacyData = static_cast<FNetworkPhysicsData*>(PhysicsData);
@@ -2786,7 +2786,7 @@ void FAsyncNetworkPhysicsComponent::OnPostProcessInputs_Internal(const int32 Phy
 
 			if (bIsUsingLegacyData)
 			{
-				if (ensure(ActorComponent.IsValid() && ActorComponent.Get()->IsBeingDestroyed() == false))
+				if (ActorComponent.IsValid() && ActorComponent.Get()->IsBeingDestroyed() == false)
 				{
 					// Gather state data from implementation
 					FNetworkPhysicsData* LegacyData = static_cast<FNetworkPhysicsData*>(PhysicsData);
@@ -2930,7 +2930,7 @@ void FAsyncNetworkPhysicsComponent::OnPostSolve_Internal()
 
 			if (bIsUsingLegacyData)
 			{
-				if (ensure(ActorComponent.IsValid() && ActorComponent.Get()->IsBeingDestroyed() == false))
+				if (ActorComponent.IsValid() && ActorComponent.Get()->IsBeingDestroyed() == false)
 				{
 					// Gather state data from implementation
 					FNetworkPhysicsData* LegacyData = static_cast<FNetworkPhysicsData*>(PhysicsData);
diff --git a/Engine/Source/Runtime/Experimental/Chaos/Private/PBDRigidsSolver.cpp b/Engine/Source/Runtime/Experimental/Chaos/Private/PBDRigidsSolver.cpp
--- a/Engine/Source/Runtime/Experimental/Chaos/Private/PBDRigidsSolver.cpp
+++ b/Engine/Source/Runtime/Experimental/Chaos/Private/PBDRigidsSolver.cpp
@@ -932,7 +932,7 @@ namespace Chaos
 		}
 
 		// Enqueue a command to remove the particle and delete the proxy
-		EnqueueCommandImmediate([Proxy, UniqueIdx, this]()
+		EnqueueCommandScheduled_External(GetMarshallingManager().GetInternalStep_External(), [Proxy, UniqueIdx, this]()
 		{
 			UE_LOG(LogPBDRigidsSolver, Verbose, TEXT("FPBDRigidsSolver::UnregisterObject() ~ Dequeue"));
 
@@ -941,14 +941,7 @@ namespace Chaos
 			ClearCollisionEventsForProxy(Proxy);
 
 			FGeometryParticleHandle* Handle = Proxy->GetHandle_LowLevel();
-			if (MRewindData)
-			{
-				MRewindData->RemoveObject(Handle);
-			}
-			
-			const int32 OffsetForRewind = (bResimDelayDestroyObject && MRewindData) ? MRewindData->Capacity() : 0;
-			PendingDestroyPhysicsProxy.Add(FPendingDestroyInfo{Proxy, GetCurrentFrame() + OffsetForRewind, Handle, UniqueIdx});
-			
+
 			//If particle was created and destroyed before commands were enqueued just skip. I suspect we can skip entire lambda, but too much code to verify right now
 			if(Handle)
 			{
@@ -956,6 +949,22 @@ namespace Chaos
 				GetEvolution()->DisableParticle(Handle);
 			}
 
+			int32 DelayDestroyFrames = 0;
+			if (MRewindData)
+			{
+				if (!MRewindData->IsResim())
+				{
+					MRewindData->DisableInHistoryFromFrame(Proxy, INDEX_NONE);
+				}
+
+				if (bResimDelayDestroyObject)
+				{
+					DelayDestroyFrames = MRewindData->Capacity();
+				}
+			}
+
+			PendingDestroyPhysicsProxy.AddUnique(FPendingDestroyInfo{Proxy, GetCurrentFrame() + DelayDestroyFrames, Handle, UniqueIdx});
+
 			// Remove the proxy for this particle
 			if (SingleParticlePhysicsProxies_PT.IsValidIndex(UniqueIdx.Idx))
 			{
@@ -1000,8 +1009,7 @@ namespace Chaos
 
 		// Particles are removed from acceleration structure in FPhysScene_Chaos::RemoveObject.
 
-
-		EnqueueCommandImmediate([InProxy, this]()
+		EnqueueCommandScheduled_External(GetMarshallingManager().GetInternalStep_External(), [InProxy, this]()
 		{
 			// Remove stale events for particles that no longer exist
 			// See comments in ClearCollisionEventsForProxy for potential race condition
@@ -1010,8 +1018,21 @@ namespace Chaos
 			GeometryCollectionPhysicsProxies_Internal.RemoveSingle(InProxy);
 			InProxy->SyncBeforeDestroy();
 
-			const int32 OffsetForRewind = (bResimDelayDestroyObject && MRewindData) ? MRewindData->Capacity() : 0;
-			PendingDestroyGeometryCollectionPhysicsProxy.Add(FPendingProxyDestroyInfo<FGeometryCollectionPhysicsProxy>{ InProxy, GetCurrentFrame() + OffsetForRewind });
+			int32 DelayDestroyFrames = 0;
+			if (MRewindData)
+			{
+				if (!MRewindData->IsResim())
+				{
+					MRewindData->DisableInHistoryFromFrame(InProxy, INDEX_NONE);
+				}
+
+				if (bResimDelayDestroyObject)
+				{
+					DelayDestroyFrames = MRewindData->Capacity();
+				}
+			}
+
+			PendingDestroyGeometryCollectionPhysicsProxy.AddUnique(FPendingProxyDestroyInfo<FGeometryCollectionPhysicsProxy>{ InProxy, GetCurrentFrame() + DelayDestroyFrames });
 		});
 	}
 
@@ -1057,8 +1078,7 @@ namespace Chaos
 			GTParticle->SetProxy(nullptr);
 		}
 
-		EnqueueCommandImmediate(
-			[Proxy, this]()
+		EnqueueCommandScheduled_External(GetMarshallingManager().GetInternalStep_External(), [Proxy, this]()
 			{
 				// Remove stale events for particles that no longer exist
 				// See comments in ClearCollisionEventsForProxy for potential race condition
@@ -1071,8 +1091,20 @@ namespace Chaos
 				ClusterUnionPhysicsProxies_Internal.RemoveSingle(Proxy);
 				Proxy->SyncBeforeDestroy();
 
-				const int32 OffsetForRewind = (bResimDelayDestroyObject && MRewindData) ? MRewindData->Capacity() : 0;
-				PendingDestroyClusterUnionProxy.Add(FPendingProxyDestroyInfo<FClusterUnionPhysicsProxy>{ Proxy, GetCurrentFrame() + OffsetForRewind });
+				int32 DelayDestroyFrames = 0;
+				if (MRewindData)
+				{
+					if (!MRewindData->IsResim())
+					{
+						MRewindData->DisableInHistoryFromFrame(Proxy, INDEX_NONE);
+					}
+
+					if (bResimDelayDestroyObject)
+					{
+						DelayDestroyFrames = MRewindData->Capacity();
+					}
+				}
+				PendingDestroyClusterUnionProxy.AddUnique(FPendingProxyDestroyInfo<FClusterUnionPhysicsProxy>{ Proxy, GetCurrentFrame() + DelayDestroyFrames });
 			}
 		);
 	}
diff --git a/Engine/Source/Runtime/Experimental/Chaos/Private/RewindData.cpp b/Engine/Source/Runtime/Experimental/Chaos/Private/RewindData.cpp
--- a/Engine/Source/Runtime/Experimental/Chaos/Private/RewindData.cpp
+++ b/Engine/Source/Runtime/Experimental/Chaos/Private/RewindData.cpp
@@ -120,6 +120,12 @@ bool SimWritablePropsMayChange(const TGeometryParticleHandle<FReal,3>& Handle)
 	return ObjectState == EObjectStateType::Dynamic || ObjectState == EObjectStateType::Sleeping;
 }
 
+namespace ResimulationCVars
+{
+	float DesyncLeniencyValue = 0.5f;
+	FAutoConsoleVariableRef CVarResimDesyncLeniencyValue(TEXT("p.Resim.DesyncLeniencyValue"), DesyncLeniencyValue, TEXT("Value of leniency when comparing current state with cached state to determine if particle is in sync of not, leniency should cover network quantization of state else particles will always be considered desynced after having a server state applied. Default RepMovement quantization is WholeNumber which means a max loss of 0.5 units in precision."));
+}
+
 template <bool bSkipDynamics>
 bool FGeometryParticleStateBase::IsInSync(const FGeometryParticleHandle& Handle, const FFrameAndPhase FrameAndPhase, const FDirtyPropertiesPool& Pool) const
 {
@@ -358,6 +364,11 @@ void FRewindData::ApplyTargets(const int32 Frame, const bool bResetSimulation)
 		}
 	};
 
+	DynamicHandlesUpdateHolder.Reset(0);
+	SleepingHandlesUpdateHolder.Reset(0);
+	StaticHandlesUpdateHolder.Reset(0);
+	KinematicHandlesUpdateHolder.Reset(0);
+
 	for (FDirtyParticleInfo& DirtyParticleInfo : DirtyParticles)
 	{
 		FGeometryParticleHandle* PTParticle = DirtyParticleInfo.GetObjectPtr();
@@ -406,26 +417,46 @@ void FRewindData::ApplyTargets(const int32 Frame, const bool bResetSimulation)
 			{
 				Particle->SetInertiaConditioningDirty();
 			}
-		
-			Particle->SetDisabled(Data.Disabled());
-			Solver->GetEvolution()->SetParticleObjectState(Particle, Data.ObjectState());
 
-			// Todo: EResimType should be set by a resimulation system and ApplyTargets() should only process particles marked for resim
-			if (!ResimConsoleVars::bIsBubbleResimEnabled)
+			switch (Data.ObjectState())
 			{
-				switch (Data.ObjectState())
-				{
 				case EObjectStateType::Dynamic:
+				{
+					DynamicHandlesUpdateHolder.Add(Particle);
+					break;
+				}
 				case EObjectStateType::Sleeping:
-					Particle->SetResimType(EResimType::FullResim);
+				{
+					SleepingHandlesUpdateHolder.Add(Particle);
+					break;
+				}
+				case EObjectStateType::Static:
+				{
+					StaticHandlesUpdateHolder.Add(Particle);
 					break;
-				default:
-					Particle->SetResimType(EResimType::ResimAsFollower);
+				}
+				case EObjectStateType::Kinematic:
+				{
+					KinematicHandlesUpdateHolder.Add(Particle);
 					break;
 				}
 			}
 		});
 
+		if (FPBDRigidsEvolution* Evolution = Solver->GetEvolution())
+		{
+			QUICK_SCOPE_CYCLE_COUNTER(ChaosRewindAndResim_ApplyTargets_UpdateObjectsState);
+
+			// Update the particles object states
+			Evolution->SetParticlesObjectState(DynamicHandlesUpdateHolder, EObjectStateType::Dynamic, false);
+			Evolution->SetParticlesObjectState(SleepingHandlesUpdateHolder, EObjectStateType::Sleeping, false);
+			Evolution->SetParticlesObjectState(StaticHandlesUpdateHolder, EObjectStateType::Static, false);
+			Evolution->SetParticlesObjectState(KinematicHandlesUpdateHolder, EObjectStateType::Kinematic, false);
+
+			// Update the dirty views
+			Evolution->GetParticles().UpdateDirtyViews();
+		}
+
 #if WITH_CHAOS_VISUAL_DEBUGGER
 		if (FChaosVisualDebuggerTrace::IsTracing() && CVDDC_ResimDetails->IsChannelEnabled())
 		{
@@ -784,10 +815,10 @@ void FRewindData::UnfreezeParticlesOutsideResimBubbles()
 		
 		if(Chaos::ResimConsoleVars::BubbleResimEnableCheapFreeze)
 		{
-			FrozenDynamicHandles.Reset(FrozenParticles.Num());
-			FrozenSleepingHandles.Reset(FrozenParticles.Num());
-			FrozenStaticHandles.Reset(FrozenParticles.Num());
-			FrozenKinematicHandles.Reset(FrozenParticles.Num());
+			DynamicHandlesUpdateHolder.Reset(FrozenParticles.Num());
+			SleepingHandlesUpdateHolder.Reset(FrozenParticles.Num());
+			StaticHandlesUpdateHolder.Reset(FrozenParticles.Num());
+			KinematicHandlesUpdateHolder.Reset(FrozenParticles.Num());
 		}
 		for (TPair<Chaos::FPhysicsProxyHandle, EObjectStateType>& FrozenParticle : FrozenParticles)
 		{
@@ -801,21 +832,21 @@ void FRewindData::UnfreezeParticlesOutsideResimBubbles()
 					{
 						if (Chaos::ResimConsoleVars::BubbleResimEnableCheapFreeze)
 						{
-							if(PreResimObjectState == EObjectStateType::Dynamic)
+							if (PreResimObjectState == EObjectStateType::Dynamic)
 							{
-								FrozenDynamicHandles.Add(RigidParticleHandle);
+								DynamicHandlesUpdateHolder.Add(RigidParticleHandle);
 							}
 							else if(PreResimObjectState == EObjectStateType::Sleeping)
 							{
-								FrozenSleepingHandles.Add(RigidParticleHandle);
+								SleepingHandlesUpdateHolder.Add(RigidParticleHandle);
 							}
 							if (PreResimObjectState == EObjectStateType::Static)
 							{
-								FrozenStaticHandles.Add(RigidParticleHandle);
+								StaticHandlesUpdateHolder.Add(RigidParticleHandle);
 							}
 							else if (PreResimObjectState == EObjectStateType::Kinematic)
 							{
-								FrozenKinematicHandles.Add(RigidParticleHandle);
+								KinematicHandlesUpdateHolder.Add(RigidParticleHandle);
 							}
 						}
 						else
@@ -840,10 +871,10 @@ void FRewindData::UnfreezeParticlesOutsideResimBubbles()
 		QUICK_SCOPE_CYCLE_COUNTER(ChaosRewindAndResim_ReportObjectsState);
 		
 		// Update the frozen particles handles with the value pre-resim
-		Evolution->SetParticlesObjectState(FrozenDynamicHandles, EObjectStateType::Dynamic, false);
-		Evolution->SetParticlesObjectState(FrozenSleepingHandles, EObjectStateType::Sleeping, false);
-		Evolution->SetParticlesObjectState(FrozenStaticHandles, EObjectStateType::Static, false);
-		Evolution->SetParticlesObjectState(FrozenKinematicHandles, EObjectStateType::Kinematic, false);
+		Evolution->SetParticlesObjectState(DynamicHandlesUpdateHolder, EObjectStateType::Dynamic, false);
+		Evolution->SetParticlesObjectState(SleepingHandlesUpdateHolder, EObjectStateType::Sleeping, false);
+		Evolution->SetParticlesObjectState(StaticHandlesUpdateHolder, EObjectStateType::Static, false);
+		Evolution->SetParticlesObjectState(KinematicHandlesUpdateHolder, EObjectStateType::Kinematic, false);
 	
 		// Update the dirty views
 		Evolution->GetParticles().UpdateDirtyViews();
@@ -1051,9 +1082,29 @@ void FRewindData::RewindToFrame(int32 RewindFrame)
 		{
 			DirtyParticleInfo.bResimAsFollower = ParticlesToForceResimAsFollower.Contains(DirtyParticleInfo.GetObjectPtr());
 		}
+
 		for (FGeometryParticleHandle* GeometryParticle : ParticlesOverlappingBubble)
 		{
-			GeometryParticle->SetResimType(ParticlesToForceResimAsFollower.Contains(GeometryParticle) ? EResimType::ResimAsFollower : EResimType::FullResim);
+			if (GeometryParticle)
+			{
+				if (FDirtyParticleInfo* DirtyParticleInfo = FindDirtyObj(*GeometryParticle))
+				{
+					if (!DirtyParticleInfo->bResimAsFollower)
+					{
+						const IPhysicsProxyBase* PhysicsProxy = GeometryParticle->PhysicsProxy();
+						if (PhysicsProxy && PhysicsProxy->GetMarkedDeleted())
+						{
+							// Force resim as follower when pending delete and we have history for the particle (we have FDirtyParticleInfo)
+							DirtyParticleInfo->bResimAsFollower = true;
+						}
+					}
+					GeometryParticle->SetResimType(DirtyParticleInfo->bResimAsFollower ? EResimType::ResimAsFollower : EResimType::FullResim);
+				}
+				else
+				{
+					GeometryParticle->SetResimType(ParticlesToForceResimAsFollower.Contains(GeometryParticle) ? EResimType::ResimAsFollower : EResimType::FullResim);
+				}
+			}
 		}
 	}
 	else
@@ -1068,12 +1119,23 @@ void FRewindData::RewindToFrame(int32 RewindFrame)
 				// Force resim as follower
 				DirtyParticleInfo.bResimAsFollower = ParticlesToForceResimAsFollower.Contains(GeometryParticle);
 			}
+
+			if (!DirtyParticleInfo.bResimAsFollower)
+			{
+				const IPhysicsProxyBase* PhysicsProxy = GeometryParticle->PhysicsProxy();
+				if (PhysicsProxy && PhysicsProxy->GetMarkedDeleted())
+				{
+					// Force resim as follower when pending delete
+					DirtyParticleInfo.bResimAsFollower = true;
+				}
+			}
+			GeometryParticle->SetResimType(DirtyParticleInfo.bResimAsFollower ? EResimType::ResimAsFollower : EResimType::FullResim);
 		}
 	}
 
 	{
 		QUICK_SCOPE_CYCLE_COUNTER(ChaosRewindAndResim_RewindParticles);
-		for(FDirtyParticleInfo& DirtyParticleInfo : DirtyParticles)
+		for (FDirtyParticleInfo& DirtyParticleInfo : DirtyParticles)
 		{
 			FGeometryParticleHandle* GeometryParticle = DirtyParticleInfo.GetObjectPtr();
 			FKinematicGeometryParticleHandle* KinematicParticle = GeometryParticle->CastToKinematicParticle();
@@ -1084,13 +1146,13 @@ void FRewindData::RewindToFrame(int32 RewindFrame)
 		
 			FGeometryParticleStateBase& History = DirtyParticleInfo.GetHistory(); //non-const in case we need to record what's at head for a rewind (CurFrame has already been increased to the next frame)
 
-			if (!ResimConsoleVars::bIsBubbleResimEnabled || GeometryParticle->ResimType() != EResimType::FrozenDuringResim)
+			const bool bResimAsFollower = DirtyParticleInfo.bResimAsFollower;
+
+			if (GeometryParticle->ResimType() == EResimType::FullResim)
 			{
 				CachePreResimState(*GeometryParticle);
 			}
 
-			const bool bResimAsFollower = DirtyParticleInfo.bResimAsFollower;
-
 			if (KinematicParticle)
 			{
 				// Always clear the kinematic target
@@ -1101,7 +1163,7 @@ void FRewindData::RewindToFrame(int32 RewindFrame)
 			bool bAnyChange = false;
 
 			// With bubble resim enabled, we only rewind particles in the bubble		
-			bool ShouldRewind = !ResimConsoleVars::bIsBubbleResimEnabled || GeometryParticle->ResimType() != EResimType::FrozenDuringResim;
+			bool ShouldRewind = GeometryParticle->ResimType() != EResimType::FrozenDuringResim;
 			if (ShouldRewind)
 			{
 				bAnyChange = RewindHelper(GeometryParticle, bResimAsFollower, History.ParticlePositionRotation, [&](auto Particle, const auto& Data) {Particle->SetXR(Data); });
@@ -1120,16 +1182,10 @@ void FRewindData::RewindToFrame(int32 RewindFrame)
 						}
 					});
 				bAnyChange |= RewindHelper(RigidParticle, bResimAsFollower, History.Dynamics, [](auto Particle, const auto& Data) {Particle->SetDynamics(Data); });
-				bAnyChange |= RewindHelper(RigidParticle, bResimAsFollower, History.DynamicsMisc, [this](auto Particle, const auto& Data) {Solver->SetParticleDynamicMisc(Particle, Data); });
+				bAnyChange |= RewindHelper(RigidParticle, bResimAsFollower, History.DynamicsMisc, [this](auto Particle, const auto& Data) { Solver->SetParticleDynamicMisc(Particle, Data);	});
 				bAnyChange |= RewindHelper(RigidParticle, bResimAsFollower, History.MassProps, [](auto Particle, const auto& Data) {Particle->SetMassProps(Data); });
 			}
 
-			// Set this after rewinding, since ResimType gets overwritten if NonFrequentData is cached
-			if (!ResimConsoleVars::bIsBubbleResimEnabled)
-			{
-				GeometryParticle->SetResimType(bResimAsFollower ? EResimType::ResimAsFollower : EResimType::FullResim);
-			}
-
 			if (!bResimAsFollower && bAnyChange)
 			{
 				//particle actually changes not just created/streamed so need to update its state
@@ -1237,6 +1293,14 @@ void FRewindData::StepNonResimParticles(const int32 Frame)
 		return false;
 	};
 
+	if (bRewindDataOptimization)
+	{
+		DynamicHandlesUpdateHolder.Reset(0);
+		SleepingHandlesUpdateHolder.Reset(0);
+		StaticHandlesUpdateHolder.Reset(0);
+		KinematicHandlesUpdateHolder.Reset(0);
+	}
+
 	for (FDirtyParticleInfo& DirtyParticleInfo : DirtyParticles)
 	{
 		FGeometryParticleHandle* GeometryParticle = DirtyParticleInfo.GetObjectPtr();
@@ -1287,7 +1351,10 @@ void FRewindData::StepNonResimParticles(const int32 Frame)
 			if (Data.GetMode() != EKinematicTargetMode::None)
 			{
 				Particle->SetKinematicTarget(Data);
-				Solver->GetEvolution()->GetParticles().MarkMovingKinematic(Particle->Handle()); // Update SOA views with this particle as being moving kinematic
+				if (!bRewindDataOptimization)
+				{
+					Solver->GetEvolution()->GetParticles().MarkMovingKinematic(Particle->Handle()); // Update SOA views with this particle as being moving kinematic
+				}
 				return true;
 			}
 			return false;
@@ -1296,13 +1363,14 @@ void FRewindData::StepNonResimParticles(const int32 Frame)
 		if (bRewindDataOptimization)
 		{
 			// Set disabled true/false and object state
-			bool bHasUpdatedSOAs = RewindHelper(RigidParticle, History.DynamicsMisc, [this](auto Particle, const auto& Data) -> bool
+			bHasChanged |= RewindHelper(RigidParticle, History.DynamicsMisc, [this, bHasChanged](auto Particle, const auto& Data) -> bool
 				{
 					if (Particle == nullptr)
 					{
-						return false; // SOAs views have not been updated
+						return false;
 					}
 
+					bool bAppliedChange = false;
 					if (Particle->Disabled() != Data.Disabled())
 					{
 						if (Data.Disabled())
@@ -1313,29 +1381,51 @@ void FRewindData::StepNonResimParticles(const int32 Frame)
 						{
 							Solver->GetEvolution()->EnableParticle(Particle);
 						}
+						bAppliedChange |= true;
 					}
 
 					if (Particle->ObjectState() != Data.ObjectState())
 					{
-						Solver->GetEvolution()->SetParticleObjectState(Particle, Data.ObjectState());
-						return true; // SOA views are updated when calling this function
+						bAppliedChange = true;
 					}
 
-					return false; // SOAs views have not been updated
-				});
+					if (bHasChanged || bAppliedChange)
+					{
+						switch (Data.ObjectState())
+						{
+							case EObjectStateType::Dynamic:
+							{
+								// If the particle is dynamic we must fix the collision anchors so that friction doesn't undo the movement
+								Particle->ParticleCollisions().VisitCollisions(
+									[this, Particle](FPBDCollisionConstraint& Collision)
+									{
+										Collision.UpdateParticleTransform(Particle);
+										return ECollisionVisitorResult::Continue;
+									});
+
+								DynamicHandlesUpdateHolder.Add(Particle);
+								break;
+							}
+							case EObjectStateType::Sleeping:
+							{
+								SleepingHandlesUpdateHolder.Add(Particle);
+								break;
+							}
+							case EObjectStateType::Static:
+							{
+								StaticHandlesUpdateHolder.Add(Particle);
+								break;
+							}
+							case EObjectStateType::Kinematic:
+							{
+								KinematicHandlesUpdateHolder.Add(Particle);
+								break;
+							}
+						}
+					}
 
-			// If not already done, update SOA views else particles might not get updated
-			if (!bHasUpdatedSOAs)
-			{
-				if (RigidParticle)
-				{
-					Solver->GetEvolution()->GetParticles().SetDynamicParticleSOA(RigidParticle->Handle());
-				}
-				else if (FPBDRigidClusteredParticleHandle* Clustered = GeometryParticle->CastToClustered())
-				{
-					Solver->GetEvolution()->GetParticles().SetClusteredParticleSOA(Clustered->Handle());
-				}	
-			}
+					return bAppliedChange;
+				});
 		}
 		else
 		{
@@ -1356,16 +1446,6 @@ void FRewindData::StepNonResimParticles(const int32 Frame)
 			RewindHelper(RigidParticle, History.MassProps, [](auto Particle, const auto& Data) -> bool { Particle->SetMassProps(Data); return true; });
 		}
 
-		// If the particle is dynamic we must fix the collision anchors so that friction doesn't undo the movement
-		if (GeometryParticle->ObjectState() == EObjectStateType::Dynamic)
-		{
-			GeometryParticle->ParticleCollisions().VisitCollisions(
-				[this, GeometryParticle](FPBDCollisionConstraint& Collision)
-				{
-					Collision.UpdateParticleTransform(GeometryParticle);
-					return ECollisionVisitorResult::Continue;
-				});
-		}
 
 #if WITH_CHAOS_VISUAL_DEBUGGER
 		if (bHasChanged && RigidParticle && FChaosVisualDebuggerTrace::IsTracing() && CVDDC_ResimDetails->IsChannelEnabled())
@@ -1374,6 +1454,23 @@ void FRewindData::StepNonResimParticles(const int32 Frame)
 		}
 #endif
 	}
+
+	if (bRewindDataOptimization)
+	{
+		if (FPBDRigidsEvolution* Evolution = Solver->GetEvolution())
+		{
+			QUICK_SCOPE_CYCLE_COUNTER(ChaosRewindAndResim_StepNonResimParticles_UpdateObjectsState);
+
+			// Update the particles object states
+			Evolution->SetParticlesObjectState(DynamicHandlesUpdateHolder, EObjectStateType::Dynamic, false);
+			Evolution->SetParticlesObjectState(SleepingHandlesUpdateHolder, EObjectStateType::Sleeping, false);
+			Evolution->SetParticlesObjectState(StaticHandlesUpdateHolder, EObjectStateType::Static, false);
+			Evolution->SetParticlesObjectState(KinematicHandlesUpdateHolder, EObjectStateType::Kinematic, false);
+
+			// Update the dirty views
+			Evolution->GetParticles().UpdateDirtyViews();
+		}
+	}
 }
 
 template <bool bSkipDynamics, typename TDirtyInfo>
@@ -1547,7 +1644,7 @@ void FRewindData::AdvanceFrameImp(IResimCacheBase* ResimCache)
 			ensure(IsResimAndInSync(*Info.GetObjectPtr()) || Info.GetHistory().IsClean(FrameAndPhase));  //Sim hasn't run yet so PostCallbacks (sim results) should be clean
 	
 			//if hasn't changed in a while stop tracking
-			if (Info.LastDirtyFrame < EarliestFrame)
+			if (Info.LastDirtyFrame < EarliestFrame && !IsResim())
 			{
 				RemoveObject(Info.GetObjectPtr(), EAllowShrinking::No);
 			}
@@ -1558,7 +1655,7 @@ void FRewindData::AdvanceFrameImp(IResimCacheBase* ResimCache)
 
 				if (IsResim())
 				{
-					if (!bRewindDataOptimization && !Info.bResimAsFollower)
+					if (!bRewindDataOptimization && Handle->ResimType() == EResimType::FullResim)
 					{
 						DesyncIfNecessary</*bSkipDynamics=*/false>(Info, FrameAndPhase);
 					}
@@ -1820,16 +1917,118 @@ void FRewindData::SpawnProxyIfNeeded(FSingleParticlePhysicsProxy& Proxy)
 	}
 }
 
+void FRewindData::DisableInHistoryFromFrame(IPhysicsProxyBase* PhysicsProxy, int32 DisableFromFrame)
+{
+	if (!Solver || !Solver->GetEvolution())
+	{
+		return;
+	}
+	
+	// If no from frame is specified, go back the amount of frames we are currently forward predicted
+	if (DisableFromFrame == INDEX_NONE)
+	{
+		DisableFromFrame = LatestFrame - FMath::RoundToInt32(NumPredictedFramesAverage);
+	}
+
+	auto DisableInHistoryHelper = [&, DisableFromFrame](const FGeometryParticleHandle* Handle)
+	{
+		if (FDirtyParticleInfo* Info = FindDirtyObj(*Handle))
+		{
+			FGeometryParticleStateBase& History = Info->GetHistory();
+			FFrameAndPhase FrameAndPhase { DisableFromFrame , FFrameAndPhase::PrePushData };
+			for (int32 Frame = DisableFromFrame; Frame <= LatestFrame; Frame++)
+			{
+				FrameAndPhase.Frame = Frame;
+
+				// Check if we have particle data in history to overwrite
+				if (!History.DynamicsMisc.IsClean(FrameAndPhase))
+				{
+					FParticleDynamicMisc& DynamicsMisc = History.DynamicsMisc.Insert(FrameAndPhase, PropertiesPool);
+					DynamicsMisc.SetDisabled(true);
+				}
+
+				// Check if we have target state data in history to overwrite
+				if (!History.TargetStates.IsClean(FrameAndPhase))
+				{
+					FParticleDynamicMisc& DynamicsMisc = History.TargetStates.Insert(FrameAndPhase, PropertiesPool);
+					DynamicsMisc.SetDisabled(true);
+				}
+			}
+		}
+	};
+
+	if (PhysicsProxy)
+	{
+		switch (PhysicsProxy->GetType())
+		{
+		case EPhysicsProxyType::SingleParticleProxy:
+		{
+			if (const FSingleParticlePhysicsProxy* Proxy = static_cast<const FSingleParticlePhysicsProxy*>(PhysicsProxy))
+			{
+				if (const FGeometryParticleHandle* Handle = Proxy->GetHandle_LowLevel())
+				{
+					DisableInHistoryHelper(Handle);
+				}
+			}
+			break;
+		}
+		case EPhysicsProxyType::ClusterUnionProxy:
+		{
+			if (const FClusterUnionPhysicsProxy* Proxy = static_cast<const FClusterUnionPhysicsProxy*>(PhysicsProxy))
+			{
+				if (const FGeometryParticleHandle* Handle = Proxy->GetParticle_Internal())
+				{
+					DisableInHistoryHelper(Handle);
+
+					// Iterate over and process all owned particles in the union also
+					FClusterUnionManager& ClusterUnionManager =	Solver->GetEvolution()->GetRigidClustering().GetClusterUnionManager();
+					if (const Chaos::FClusterUnion* Union = ClusterUnionManager.FindClusterUnionFromParticle(Handle->CastToRigidParticle()))
+					{
+						for (Chaos::FPBDRigidParticleHandle* Child : Union->GeometryChildParticles)
+						{
+							if (Child)
+							{
+								DisableInHistoryHelper(Child);
+							}
+						}
+					}
+				}
+			}
+			break;
+		}
+		case EPhysicsProxyType::GeometryCollectionType:
+		{
+			if (const FGeometryCollectionPhysicsProxy* Proxy = static_cast<const FGeometryCollectionPhysicsProxy*>(PhysicsProxy))
+			{
+				for (const FGeometryParticleHandle* Handle : Proxy->GetUnorderedParticles_Internal())
+				{
+					if (Handle)
+					{
+						DisableInHistoryHelper(Handle);
+					}
+				}
+			}
+			break;
+		}
+		}
+	}
+}
+
 void FRewindData::CachePreResimState(FGeometryParticleHandle& Handle)
 {
-	const IPhysicsProxyBase* PhysicsProxy = static_cast<const IPhysicsProxyBase*>(Handle.PhysicsProxy());
-	if (!PhysicsProxy)
+	if (!Solver || !Solver->GetEvolution())
+	{
+		return;
+	}
+
+	const IPhysicsProxyBase* PhysicsProxy = Handle.PhysicsProxy();
+	if (!PhysicsProxy || PhysicsProxy->GetMarkedDeleted())
 	{
 		return;
 	}
 
 	// Find or add pre-resim error for the particle that has an error
-	auto GetPreResimState = [this](FGeometryParticleHandle* ParticleHandle) -> FParticleSimpleState&
+	auto FindOrAddPreResimState = [this](FGeometryParticleHandle* ParticleHandle) -> FParticleSimpleState&
 		{
 			if (FParticleSimpleState* Found = DirtyParticlePreResimState.Find(ParticleHandle))
 			{
@@ -1846,7 +2045,15 @@ void FRewindData::CachePreResimState(FGeometryParticleHandle& Handle)
 				return;
 			}
 
-			FParticleSimpleState& PreResimState = GetPreResimState(ParticleHandle);
+			if (const FPBDRigidParticleHandle* RigidParticle = ParticleHandle->CastToRigidParticle())
+			{
+				if (RigidParticle->Disabled())
+				{
+					return;
+				}
+			}
+
+			FParticleSimpleState& PreResimState = FindOrAddPreResimState(ParticleHandle);
 
 			// Cache dirty particle XR before a resimulation
 			PreResimState.CacheXR(ParticleHandle->GetX(), ParticleHandle->GetR());
@@ -1865,46 +2072,14 @@ void FRewindData::CachePreResimState(FGeometryParticleHandle& Handle)
 	// If particle is a Cluster Union, also cache child particles
 	if (PhysicsProxy->GetType() == EPhysicsProxyType::ClusterUnionProxy)
 	{
-		const FClusterUnionPhysicsProxy* ClusterProxy = static_cast<const FClusterUnionPhysicsProxy*>(PhysicsProxy);
-		if (ClusterProxy)
+		FClusterUnionManager& ClusterUnionManager = Solver->GetEvolution()->GetRigidClustering().GetClusterUnionManager();
+		if (const Chaos::FClusterUnion* Union = ClusterUnionManager.FindClusterUnionFromParticle(Handle.CastToRigidParticle()))
 		{
-			for (IPhysicsProxyBase* ChildProxyBase : ClusterProxy->GetParticle_Internal()->PhysicsProxies())
+			for (Chaos::FPBDRigidParticleHandle* Child : Union->GeometryChildParticles)
 			{
-				switch (ChildProxyBase->GetType())
+				if (Child)
 				{
-					case EPhysicsProxyType::SingleParticleProxy:
-					{
-						if (FSingleParticlePhysicsProxy* ChildProxy = static_cast<FSingleParticlePhysicsProxy*>(ChildProxyBase))
-						{
-							if (FGeometryParticleHandle* ChildHandle = ChildProxy->GetHandle_LowLevel())
-							{
-								CacheState(ChildHandle);
-							}
-						}
-						break;
-					}
-					case EPhysicsProxyType::ClusterUnionProxy:
-					{
-						if (FClusterUnionPhysicsProxy* ChildProxy = static_cast<FClusterUnionPhysicsProxy*>(ChildProxyBase))
-						{
-							if (FGeometryParticleHandle* ChildHandle = ChildProxy->GetParticle_Internal())
-							{
-								CacheState(ChildHandle);
-							}
-						}
-						break;
-					}
-					case EPhysicsProxyType::GeometryCollectionType:
-					{
-						if (FGeometryCollectionPhysicsProxy* ChildProxy = static_cast<FGeometryCollectionPhysicsProxy*>(ChildProxyBase))
-						{
-							if (FGeometryParticleHandle* ChildHandle = ChildProxy->GetInitialRootParticle_Internal())
-							{
-								CacheState(ChildHandle);
-							}
-						}
-						break;
-					}
+					CacheState(Child);
 				}
 			}
 		}
@@ -1966,7 +2141,7 @@ void FRewindData::BufferPhysicsResults(TMap<const IPhysicsProxyBase*, struct FDi
 			{
 				for (auto& Info : DirtyObjs)
 				{
-					if (!Info.GetObjectPtr() || Info.GetObjectPtr()->ResimType() == EResimType::ResimAsFollower)
+					if (!Info.GetObjectPtr() || Info.GetObjectPtr()->ResimType() != EResimType::FullResim)
 					{
 						continue;
 					}
@@ -2511,8 +2686,8 @@ void FRewindData::CacheKinematicTarget(TPBDRigidParticleHandle<FReal, 3>& Handle
 
 void FRewindData::CacheDirtyParticleData(FGeometryParticleHandle* Particle, const FFrameAndPhase::EParticleHistoryPhase& CurrentPhase, const bool& bDirty)
 {
-	// Bubble Resim
-	if (ResimConsoleVars::bIsBubbleResimEnabled && IsResim() && Particle->ResimType() == EResimType::FrozenDuringResim)
+	// If we are resimulating, only process particles under EResimType::FullResim
+	if (IsResim() && Particle->ResimType() != EResimType::FullResim)
 	{
 		return;
 	}
@@ -2540,7 +2715,7 @@ void FRewindData::CacheDirtyParticleData(FGeometryParticleHandle* Particle, cons
 
 	if (IsResim())
 	{
-		if (!DirtyParticleInfo.bResimAsFollower)
+		// NOTE, we have already checked that the particle is EResimType::FullResim during resimulation
 		{
 			// Check if particle differ from the currently cached history for this frame, if so, clear the history and mark particle as desynced so it will cache data during resimulation
 			DesyncIfNecessary</*bSkipDynamics=*/false>(DirtyParticleInfo, FrameAndPhase);
@@ -2745,6 +2920,9 @@ FJointState FRewindData::GetPastJointStateAtFrame(const FPBDJointConstraintHandl
 CHAOS_API int32 InterpolateTargetGaps = 5;
 FAutoConsoleVariableRef CVarResimInterpolateTargetGaps(TEXT("p.Resim.InterpolateTargetGaps"), InterpolateTargetGaps, TEXT("How many frame gaps in replicated targets we should fill by interpolating between the previous and the new target received. Value in max number of frames to interpolate, deactivate by setting to 0."));
 
+float AveragePredictedFramesScalar = 0.2f;
+static FAutoConsoleVariableRef CVarResimAveragePredictedFramesScalar(TEXT("p.Resim.AveragePredictedFramesScalar"), AveragePredictedFramesScalar, TEXT("Multiplier for how quickly the average number of predicted frames are updated when changed, clamped between 0.01 and 1.0."));
+
 void FRewindData::SetTargetStateAtFrame(FGeometryParticleHandle& Handle, const int32 Frame, FFrameAndPhase::EParticleHistoryPhase Phase,
 	const FVector& Position, const FQuat& Quaternion, const FVector& LinVelocity, const FVector& AngVelocity, const bool bShouldSleep)
 {
@@ -2787,8 +2965,14 @@ void FRewindData::SetTargetStateAtFrame(FGeometryParticleHandle& Handle, const i
 		}
 	}
 
-	PushStateAtFrame(Handle, Frame, Phase,
-		Position, Quaternion, LinVelocity, AngVelocity, bShouldSleep);
+	PushStateAtFrame(Handle, Frame, Phase, Position, Quaternion, LinVelocity, AngVelocity, bShouldSleep);
+
+	// Cache an average of how many frames we are currently forward predicted
+	if (!IsResim() && Frame > 0)
+	{
+		const float PredictedFrames = static_cast<float>(LatestFrame - Frame);
+		NumPredictedFramesAverage += (PredictedFrames - NumPredictedFramesAverage) * FMath::Clamp(AveragePredictedFramesScalar, 0.01f, 1.0f);
+	}
 }
 
 void FRewindData::ForceResimAsFollower(FGeometryParticleHandle* Handle, bool bForceResim)
diff --git a/Engine/Source/Runtime/Experimental/Chaos/Public/Chaos/KinematicTargets.h b/Engine/Source/Runtime/Experimental/Chaos/Public/Chaos/KinematicTargets.h
--- a/Engine/Source/Runtime/Experimental/Chaos/Public/Chaos/KinematicTargets.h
+++ b/Engine/Source/Runtime/Experimental/Chaos/Public/Chaos/KinematicTargets.h
@@ -170,6 +170,19 @@ namespace Chaos
 			return IsEqual(other.KinematicTarget());
 		}
 
+		bool IsNearlyEqual(const FKinematicTarget& other, FReal Leniency = UE_KINDA_SMALL_NUMBER) const
+		{
+			return Mode == other.Mode &&
+				FVec3::IsNearlyEqual(Position, other.Position, Leniency) &&
+				FRotation3f::IsNearlyEqual(Rotation, other.Rotation, static_cast<FRealSingle>(Leniency));
+		}
+
+		template <typename TOther>
+		bool IsNearlyEqual(const TOther& other, FReal Leniency = UE_KINDA_SMALL_NUMBER) const
+		{
+			return IsNearlyEqual(other.KinematicTarget(), Leniency);
+		}
+
 		bool operator==(const FKinematicTarget& other) const
 		{
 			return IsEqual(other);
diff --git a/Engine/Source/Runtime/Experimental/Chaos/Public/Chaos/PBDJointConstraintTypes.h b/Engine/Source/Runtime/Experimental/Chaos/Public/Chaos/PBDJointConstraintTypes.h
--- a/Engine/Source/Runtime/Experimental/Chaos/Public/Chaos/PBDJointConstraintTypes.h
+++ b/Engine/Source/Runtime/Experimental/Chaos/Public/Chaos/PBDJointConstraintTypes.h
@@ -119,6 +119,12 @@ namespace Chaos
 		// whenever settings change.
 		CHAOS_API void Sanitize();
 
+		// Not currently implemented but needed for templated logic
+		bool IsNearlyEqual(const FPBDJointSettings& Other, FReal Leniency = UE_KINDA_SMALL_NUMBER) const
+		{
+			return *this == Other;
+		}
+
 		bool operator==(const FPBDJointSettings& Other) const
 		{
 			return !FMemory::Memcmp(this, &Other, sizeof(*this));
diff --git a/Engine/Source/Runtime/Experimental/Chaos/Public/Chaos/ParticleDirtyFlags.h b/Engine/Source/Runtime/Experimental/Chaos/Public/Chaos/ParticleDirtyFlags.h
--- a/Engine/Source/Runtime/Experimental/Chaos/Public/Chaos/ParticleDirtyFlags.h
+++ b/Engine/Source/Runtime/Experimental/Chaos/Public/Chaos/ParticleDirtyFlags.h
@@ -90,6 +90,12 @@ public:
 		return MX == Other.X() && MR == FRotation3f(Other.R());
 	}
 
+	template <typename TOther>
+	bool IsNearlyEqual(const TOther& Other, FReal Leniency = UE_KINDA_SMALL_NUMBER) const
+	{
+		return FVec3::IsNearlyEqual(MX, Other.X(), Leniency) && FRotation3f::IsNearlyEqual(MR, FRotation3f(Other.R()), static_cast<FRealSingle>(Leniency));
+	}
+
 	bool operator==(const FParticlePositionRotation& Other) const
 	{
 		return IsEqual(Other);
@@ -150,6 +156,12 @@ public:
 		return MV == FVec3f(Other.GetV()) && MW == FVec3f(Other.GetW());
 	}
 
+	template <typename TOther>
+	bool IsNearlyEqual(const TOther& Other, FReal Leniency = UE_KINDA_SMALL_NUMBER) const
+	{
+		return FVec3::IsNearlyEqual(MV, Other.GetV(), Leniency) && FVec3::IsNearlyEqual(MW, Other.GetW(), Leniency);
+	}
+
 	bool operator==(const FParticleVelocities& Other) const
 	{
 		return IsEqual(Other);
@@ -223,6 +235,15 @@ public:
 			&& AngularImpulseVelocity() == Other.AngularImpulseVelocity();
 	}
 
+	template <typename TOther>
+	bool IsNearlyEqual(const TOther& Other, FReal Leniency = UE_KINDA_SMALL_NUMBER) const
+	{
+		return FVec3::IsNearlyEqual(Acceleration(), Other.Acceleration(), Leniency)
+			&& FVec3::IsNearlyEqual(AngularAcceleration(), Other.AngularAcceleration(), Leniency)
+			&& FVec3::IsNearlyEqual(LinearImpulseVelocity(), Other.LinearImpulseVelocity(), Leniency)
+			&& FVec3::IsNearlyEqual(AngularImpulseVelocity(), Other.AngularImpulseVelocity(), Leniency);
+	}
+
 	bool operator==(const FParticleDynamics& Other) const
 	{
 		return IsEqual(Other);
@@ -325,6 +346,17 @@ public:
 			&& InvM() == Other.InvM();
 	}
 
+	template <typename TOther>
+	bool IsNearlyEqual(const TOther& Other, FReal Leniency = UE_KINDA_SMALL_NUMBER) const
+	{
+		return FVec3::IsNearlyEqual(CenterOfMass(), Other.CenterOfMass(), Leniency)
+			&& FRotation3::IsNearlyEqual(RotationOfMass(), Other.RotationOfMass(), static_cast<FRealSingle>(Leniency))
+			&& FVec3::IsNearlyEqual(I(), Other.I(), Leniency)
+			&& FVec3::IsNearlyEqual(InvI(), Other.InvI(), Leniency)
+			&& FMath::Abs(M() - Other.M()) <= Leniency
+			&& FMath::Abs(InvM() - Other.InvM()) <= Leniency;
+	}
+
 	bool operator==(const FParticleMassProps& Other) const
 	{
 		return IsEqual(Other);
@@ -511,6 +543,24 @@ public:
 			&& Disabled() == Other.Disabled();
 	}
 
+	template <typename TOther>
+	bool IsNearlyEqual(const TOther& Other, FReal Leniency = UE_KINDA_SMALL_NUMBER) const
+	{
+		return ObjectState() == Other.ObjectState()
+			&& FMath::Abs(MLinearEtherDrag - Other.MLinearEtherDrag) <= Leniency
+			&& FMath::Abs(MAngularEtherDrag - Other.MAngularEtherDrag) <= Leniency
+			&& FMath::Abs(MMaxLinearSpeedSq - Other.MMaxLinearSpeedSq) <= Leniency
+			&& FMath::Abs(MMaxAngularSpeedSq - Other.MMaxAngularSpeedSq) <= Leniency
+			&& FMath::Abs(InitialOverlapDepenetrationVelocity() - Other.InitialOverlapDepenetrationVelocity()) <= Leniency
+			&& FMath::Abs(SleepThresholdMultiplier() - Other.SleepThresholdMultiplier()) <= Leniency
+			&& CollisionGroup() == Other.CollisionGroup()
+			&& SleepType() == Other.SleepType()
+			&& CollisionConstraintFlags() == Other.CollisionConstraintFlags()
+			&& ControlFlags() == Other.ControlFlags()
+			&& IterationSettings() == Other.IterationSettings()
+			&& Disabled() == Other.Disabled();
+	}
+
 	bool operator==(const FParticleDynamicMisc& Other) const
 	{
 		return IsEqual(Other);
@@ -653,6 +703,13 @@ public:
 			&& SpatialIdx() == Other.SpatialIdx();
 	}
 
+	template <typename TOther>
+	bool IsNearlyEqual(const TOther& Other, FReal Leniency = UE_KINDA_SMALL_NUMBER) const
+	{
+		// No properties that can be nearly compared
+		return IsEqual(Other);
+	}
+
 	bool operator==(const FParticleNonFrequentData& Other) const
 	{
 		return IsEqual(Other);
diff --git a/Engine/Source/Runtime/Experimental/Chaos/Public/PBDRigidsSolver.h b/Engine/Source/Runtime/Experimental/Chaos/Public/PBDRigidsSolver.h
--- a/Engine/Source/Runtime/Experimental/Chaos/Public/PBDRigidsSolver.h
+++ b/Engine/Source/Runtime/Experimental/Chaos/Public/PBDRigidsSolver.h
@@ -447,6 +447,11 @@ namespace Chaos
 			int32 DestroyOnStep;
 			FGeometryParticleHandle* Handle;
 			FUniqueIdx UniqueIdx;
+
+			const bool operator==(const FPendingDestroyInfo Other) const
+			{
+				return Proxy == Other.Proxy;
+			}
 		};
 
 		template<typename ProxyType>
@@ -454,6 +459,11 @@ namespace Chaos
 		{
 			ProxyType* Proxy;
 			int32 DestroyOnStep;
+
+			const bool operator==(const FPendingProxyDestroyInfo Other) const
+			{
+				return Proxy == Other.Proxy;
+			}
 		};
 
 		TArray<FPendingDestroyInfo> PendingDestroyPhysicsProxy;
diff --git a/Engine/Source/Runtime/Experimental/Chaos/Public/RewindData.h b/Engine/Source/Runtime/Experimental/Chaos/Public/RewindData.h
--- a/Engine/Source/Runtime/Experimental/Chaos/Public/RewindData.h
+++ b/Engine/Source/Runtime/Experimental/Chaos/Public/RewindData.h
@@ -38,6 +38,11 @@ namespace Chaos
 		extern const CHAOS_API FName BubbleResimDebugDrawTag;
 	}
 
+	namespace ResimulationCVars
+	{
+		extern float DesyncLeniencyValue;
+	}
+
 /** Base rewind history used in the rewind data */
 struct FBaseRewindHistory
 {
@@ -806,7 +811,7 @@ public:
 		{
 			T HeadVal;
 			CopyDataFromObject(HeadVal, Handle);
-			return *Val == HeadVal;
+			return Val->IsNearlyEqual(HeadVal, ResimulationCVars::DesyncLeniencyValue);
 		}
 
 		return NoEntryInSync<THandle, T, bNoEntryIsHead>::Helper(Handle);
@@ -1892,6 +1897,8 @@ public:
 
 	void CHAOS_API SpawnProxyIfNeeded(FSingleParticlePhysicsProxy& Proxy);
 
+	void DisableInHistoryFromFrame(IPhysicsProxyBase* PhysicsProxy, int32 DisableParticleOnFrame = INDEX_NONE);
+
 	/** Add input history to the rewind data for future use while resimulating */
 	void AddInputHistory(const TSharedPtr<FBaseRewindHistory>& InputHistory)
 	{
@@ -2276,20 +2283,23 @@ private:
 	TArray<const Private::FPBDIsland*> IslandValidationIslands;
 	TArray<const FGeometryParticleHandle*> IslandValidationIslandParticles;
 	
-	/** Frozen dynamic particle handles to be used when unfreezing  */
-	TArray<FPBDRigidParticleHandle*> FrozenDynamicHandles;
+	/** Dynamic particle handles to be set in bulk */
+	TArray<FPBDRigidParticleHandle*> DynamicHandlesUpdateHolder;
 	
-	/** Frozen sleeping particle handles to be used when unfreezing  */
-	TArray<FPBDRigidParticleHandle*> FrozenSleepingHandles;
+	/** Sleeping particle handles to be set in bulk */
+	TArray<FPBDRigidParticleHandle*> SleepingHandlesUpdateHolder;
 	
-	/** Frozen static particle handles to be used when unfreezing  */
-	TArray<FPBDRigidParticleHandle*> FrozenStaticHandles;
+	/** Static particle handles to be set in bulk */
+	TArray<FPBDRigidParticleHandle*> StaticHandlesUpdateHolder;
 	
-	/** Frozen kinematic particle handles to be used when unfreezing  */
-	TArray<FPBDRigidParticleHandle*> FrozenKinematicHandles;
+	/** Kinematic particle handles to be set in bulk */
+	TArray<FPBDRigidParticleHandle*> KinematicHandlesUpdateHolder;
 
 	TSet<const FGeometryParticleHandle*> ParticlesToForceResimAsFollower;
 
+	/** Average number of predicted frames, overall estimated value */
+	float NumPredictedFramesAverage = 0.0f;
+
 	// BUBBLE RESIM
 	// Bubble resim needs to be enabled, see Chaos::ResimConsoleVars::bIsBubbleResimEnabled
 	// Bubble resim mitigates the cost of resimulations by restricting the amount of particles being resimulated
