[Chaos] Uniform multithreading through the physics engine
When constraints number bellow batch size, not task are spawned anymore
Revert back batch size for collision detection that was causing memory issue on low end platform.

#Synced-CL 51758418

diff --git a/Engine/Source/Runtime/Experimental/Chaos/Private/Chaos/Collision/CollisionConstraintAllocator.cpp b/Engine/Source/Runtime/Experimental/Chaos/Private/Chaos/Collision/CollisionConstraintAllocator.cpp
--- a/Engine/Source/Runtime/Experimental/Chaos/Private/Chaos/Collision/CollisionConstraintAllocator.cpp
+++ b/Engine/Source/Runtime/Experimental/Chaos/Private/Chaos/Collision/CollisionConstraintAllocator.cpp
@@ -6,9 +6,6 @@
 namespace Chaos
 {
 
-int32 MinBatchSizeEndCollision = 40;
-FAutoConsoleVariableRef CVarMinBatchSizeEndCollision(TEXT("p.MinBatchSizeEndCollision"), MinBatchSizeEndCollision, TEXT("Minimum number of constraints to batch tasks in EndCollision."));
-
 	namespace Private
 	{
 		FPBDCollisionConstraintPtr FCollisionContextAllocator::CreateConstraint(
@@ -97,13 +94,14 @@ FAutoConsoleVariableRef CVarMinBatchSizeEndCollision(TEXT("p.MinBatchSizeEndColl
 
 		void FCollisionConstraintAllocator::ProcessNewNonCcdConstraints(const TArray<FPBDCollisionConstraint*>& NewConstraints, TArray<FPBDCollisionConstraint*>& AllCCDConstraints)
 		{
-			const int32 NumTasks = FMath::Max(FMath::Min(FTaskGraphInterface::Get().GetNumWorkerThreads(), Chaos::MaxNumWorkers), 1);
 			const int32 NumConstraints = NewConstraints.Num();
-			const int32 ConstraintByTask = FMath::Max(FMath::DivideAndRoundUp(NumConstraints, NumTasks), MinBatchSizeEndCollision);
-			const int32 NumBatches = FMath::DivideAndRoundUp(NumConstraints, ConstraintByTask);
 
-			if (Chaos::ShouldExecuteTasksInParallel(NumBatches))
+			if (Chaos::ShouldExecuteTasks(NumConstraints))
 			{
+				const int32 NumTasks = FMath::Max(FMath::Min(FTaskGraphInterface::Get().GetNumWorkerThreads(), Chaos::MaxNumWorkers), 1);
+				const int32 ConstraintByTask = FMath::Max(FMath::DivideAndRoundUp(NumConstraints, NumTasks), CollisionSmallBatchSize);
+				const int32 NumBatches = FMath::DivideAndRoundUp(NumConstraints, ConstraintByTask);
+
 				TArray<UE::Tasks::FTask> PendingTasks;
 				PendingTasks.Reserve(NumBatches);
 
@@ -114,7 +112,7 @@ FAutoConsoleVariableRef CVarMinBatchSizeEndCollision(TEXT("p.MinBatchSizeEndColl
 				{
 					const int32 StartIndex = BatchIndex * ConstraintByTask;
 					int32 EndIndex = (BatchIndex + 1) * ConstraintByTask;
-					EndIndex = BatchIndex == NumBatches - 1 ? FMath::Min(NumConstraints, EndIndex) : EndIndex;
+					EndIndex = FMath::Min(NumConstraints, EndIndex);
 
 					TArray<FPBDCollisionConstraint*>& CCDConstraints = CCDConstraintsBatch[BatchIndex];
 
@@ -162,13 +160,16 @@ FAutoConsoleVariableRef CVarMinBatchSizeEndCollision(TEXT("p.MinBatchSizeEndColl
 		void FCollisionConstraintAllocator::ProcessNewCcdConstraints()
 		{
 
-			const int32 NumTasks = FMath::Max(FMath::Min(FTaskGraphInterface::Get().GetNumWorkerThreads(), Chaos::MaxNumWorkers), 1);
+			
 			const int32 NumConstraints = ActiveCCDConstraints.Num();
-			const int32 ConstraintByTask = FMath::Max(FMath::DivideAndRoundUp(NumConstraints, NumTasks), MinBatchSizeEndCollision);
-			const int32 NumBatches = FMath::DivideAndRoundUp(NumConstraints, ConstraintByTask);
+			
 
-			if (Chaos::ShouldExecuteTasksInParallel(NumBatches))
+			if (Chaos::ShouldExecuteTasks(NumConstraints))
 			{
+				const int32 NumTasks = FMath::Max(FMath::Min(FTaskGraphInterface::Get().GetNumWorkerThreads(), Chaos::MaxNumWorkers), 1);
+				const int32 ConstraintByTask = FMath::Max(FMath::DivideAndRoundUp(NumConstraints, NumTasks), CollisionSmallBatchSize);
+				const int32 NumBatches = FMath::DivideAndRoundUp(NumConstraints, ConstraintByTask);
+
 				TArray<UE::Tasks::FTask> PendingTasks;
 				PendingTasks.Reserve(NumBatches);
 
@@ -223,13 +224,15 @@ FAutoConsoleVariableRef CVarMinBatchSizeEndCollision(TEXT("p.MinBatchSizeEndColl
 		void FCollisionConstraintAllocator::ProcessNewItems()
 		{
 
-			const int32 NumTasks = FMath::Max(FMath::Min(FTaskGraphInterface::Get().GetNumWorkerThreads(), Chaos::MaxNumWorkers), 1);
+			
 			const int32 NumConstraints = ContextAllocators.IsEmpty() ? 0 : ContextAllocators[0]->NewActiveConstraints.Num();
-			const int32 ConstraintByTask = FMath::Max(FMath::DivideAndRoundUp(NumConstraints, NumTasks), MinBatchSizeEndCollision);
-			const int32 NumBatches = FMath::DivideAndRoundUp(NumConstraints, ConstraintByTask);
-
-			if (Chaos::ShouldExecuteTasksInParallel(NumBatches))
+			
+			if (Chaos::ShouldExecuteTasks(NumConstraints))
 			{
+				const int32 NumTasks = FMath::Max(FMath::Min(FTaskGraphInterface::Get().GetNumWorkerThreads(), Chaos::MaxNumWorkers), 1);
+				const int32 ConstraintByTask = FMath::Max(FMath::DivideAndRoundUp(NumConstraints, NumTasks), CollisionSmallBatchSize);
+				const int32 NumBatches = FMath::DivideAndRoundUp(NumConstraints, ConstraintByTask);
+
 				UE::Tasks::FTask PendingTask = UE::Tasks::Launch(UE_SOURCE_LOCATION, [this] 
 				{
 					ProcessNewMidPhases();
diff --git a/Engine/Source/Runtime/Experimental/Chaos/Private/Chaos/Framework/Parallel.cpp b/Engine/Source/Runtime/Experimental/Chaos/Private/Chaos/Framework/Parallel.cpp
--- a/Engine/Source/Runtime/Experimental/Chaos/Private/Chaos/Framework/Parallel.cpp
+++ b/Engine/Source/Runtime/Experimental/Chaos/Private/Chaos/Framework/Parallel.cpp
@@ -14,9 +14,12 @@ namespace Chaos
 	CHAOS_API int32 InnerParallelForBatchSize = 0;
 	CHAOS_API int32 MinRangeBatchSize = 0;
 	CHAOS_API int32 MaxNumWorkers = 100;
-	CHAOS_API int32 SmallBatchSize = 1;
+	// Batch size for collision 
+	// can be smaller in general because Narrow Phase between two rigid bodies can be a set of many shapes with height field or mesh. 
+	CHAOS_API int32 CollisionSmallBatchSize = 10;
+	// Batch size  for multi threading, use not to launch tasks if not enough instance
+	CHAOS_API int32 GlobalSmallBatchSize = 40;
 	CHAOS_API int32 LargeBatchSize = 100;
-	CHAOS_API int32 MinParallelTaskSize = 2;
 	
 	CHAOS_API bool bDisablePhysicsParallelFor = false;
 	CHAOS_API bool bDisableParticleParallelFor = false;
@@ -28,9 +31,9 @@ namespace Chaos
 	FAutoConsoleVariableRef CVarInnerPhysicsBatchSize(TEXT("p.Chaos.InnerParallelForBatchSize"), InnerParallelForBatchSize, TEXT("Set the batch size threshold for inner parallel fors"));
 	FAutoConsoleVariableRef CVarMinRangeBatchSize(TEXT("p.Chaos.MinRangeBatchSize"), MinRangeBatchSize, TEXT("Set the min range batch size for parallel for"));
 	FAutoConsoleVariableRef CVarMaxRangeBatchWorkers(TEXT("p.Chaos.MaxNumWorkers"), MaxNumWorkers, TEXT("Set the max number of workers for physics"));
-	FAutoConsoleVariableRef CVarSmallBatchSize(TEXT("p.Chaos.SmallBatchSize"), SmallBatchSize, TEXT("Small batch size for chaos parallel loops"));
-	FAutoConsoleVariableRef CVarLargeBatchSize(TEXT("p.Chaos.LargeBatchSize"), LargeBatchSize, TEXT("Large batch size for chaos parallel loops"));
-	FAutoConsoleVariableRef CVarMinParallelTaskSize(TEXT("p.Chaos.MinParallelTaskSize"), MinParallelTaskSize, TEXT("Minimum number of tasks required to start parallel task execution on worker threads, set this higher to encourage single threaded execution"));
+	FAutoConsoleVariableRef CVarSmallBatchCollisionSize(TEXT("p.Chaos.CollisionSmallBatchSize"), CollisionSmallBatchSize, TEXT("Small batch size used for Collision Detection"));
+	FAutoConsoleVariableRef CVarGlobalSmallBatchSize(TEXT("p.Chaos.GlobalSmallBatchSize"), GlobalSmallBatchSize, TEXT("Minimum batch size needed to launch multi-threaded tasks"));
+	FAutoConsoleVariableRef CVarLargeBatchSize(TEXT("p.Chaos.LargeBatchSize"), LargeBatchSize, TEXT("Large batch size for chaos tasks"));
 	FAutoConsoleVariableRef CVarSingleWorkerPhysics(TEXT("p.Chaos.SingleWorkerPhysics"), bSingleWorkerPhysics, TEXT("Run Chaos on a single thread"));
 }
 
@@ -199,88 +202,19 @@ void Chaos::PhysicsParallelForWithContext(int32 InNum, TFunctionRef<int32 (int32
 	::ParallelForWithTaskContext(TEXT("PhysicsParallelForWithContext"), Contexts, InNum, MinBatchSize, InContextCreator, PassThrough, Flags);
 }
 
-bool Chaos::ShouldExecuteTasksInParallel(int32 InTaskNum)
+bool Chaos::ShouldExecuteTasks(int32 InInstanceNum)
 {
 	if (bSingleWorkerPhysics)
 	{
 		// Single threaded physics should not run in parallel
 		return false;
 	}
-	else if (InTaskNum >= MinParallelTaskSize)
+	else if (InInstanceNum >= GlobalSmallBatchSize)
 	{
 		// Meets threshold so run in parallel
 		return true;
 	}
-	else if (IsInGameThread())
-	{
-		// Game thread should always use workers
-		return true;
-	}
-
 	return false;
 }
 
-//class FRecursiveDivideTask
-//{
-//	TFuture<void> ThisFuture;
-//	TFunctionRef<void(int32)> Callable;
-//
-//	int32 Begin;
-//	int32 End;
-//
-//	FRecursiveDivideTask(int32 InBegin, int32 InEnd, TFunctionRef<void(int32)> InCallable)
-//		: Begin(InBegin)
-//		, End(InEnd)
-//		, Callable(InCallable)
-//	{
-//
-//	}
-//
-//	static FORCEINLINE TStatId GetStatId()
-//	{
-//		return GET_STATID(STAT_ParallelForTask);
-//	}
-//
-//	static FORCEINLINE ENamedThreads::Type GetDesiredThread()
-//	{
-//		return ENamedThreads::AnyHiPriThreadHiPriTask;
-//	}
-//
-//	static FORCEINLINE ESubsequentsMode::Type GetSubsequentsMode()
-//	{
-//		return ESubsequentsMode::FireAndForget;
-//	}
-//
-//	void DoTask(ENamedThreads::Type CurrentThread, const FGraphEventRef& MyCompletionGraphEvent)
-//	{
-//		for(int32 Index = Begin; Index < End; ++Index)
-//		{
-//			Callable(Index);
-//		}
-//
-//		ThisFuture.Share()
-//	}
-//};
-//
-//void Chaos::PhysicsParallelFor_RecursiveDivide(int32 InNum, TFunctionRef<void(int32)> InCallable, bool bForceSingleThreaded /*= false*/)
-//{
-//	const int32 TaskThreshold = 15;
-//	const int32 NumAvailable = FTaskGraphInterface::Get().GetNumWorkerThreads();
-//	const int32 BatchSize = InNum / NumAvailable;
-//	const int32 BatchCount = InNum / BatchSize;
-//
-//	const bool bUseThreads = !bForceSingleThreaded && FApp::ShouldUseThreadingForPerformance() && InNum > (2 * TaskThreshold);
-//	const int32 NumToSpawn = bUseThreads ? FMath::Min<int32>(NumAvailable);
-//
-//	if(!bUseThreads)
-//	{
-//		for(int32 Index = 0; Index < InNum; ++Index)
-//		{
-//			InCallable(Index);
-//		}
-//	}
-//	else
-//	{
-//		const int32 MidPoint = InNum / 2;
-//	}
-//}
+
diff --git a/Engine/Source/Runtime/Experimental/Chaos/Private/Chaos/PBDRigidsEvolution.cpp b/Engine/Source/Runtime/Experimental/Chaos/Private/Chaos/PBDRigidsEvolution.cpp
--- a/Engine/Source/Runtime/Experimental/Chaos/Private/Chaos/PBDRigidsEvolution.cpp
+++ b/Engine/Source/Runtime/Experimental/Chaos/Private/Chaos/PBDRigidsEvolution.cpp
@@ -652,7 +652,7 @@ namespace Chaos
 
 		};
 
-		if (!bSingleWorkerPhysics)
+		if (Chaos::ShouldExecuteTasks(Particles.GetActiveDynamicMovingKinematicParticlesView().Num()))
 		{
 			TaskDispatcher.PruneInternalPendingData(HasToBeUpdated);
 		}
@@ -716,7 +716,7 @@ namespace Chaos
 			AsyncInternalAcceleration->SetSyncTimestamp(LatestExternalTimestampConsumed_Internal);
 			AsyncExternalAcceleration->SetSyncTimestamp(LatestExternalTimestampConsumed_Internal);
 		};
-		if (bSingleWorkerPhysics)
+		if (!Chaos::ShouldExecuteTasks(Particles.GetActiveDynamicMovingKinematicParticlesView().Num()))
 		{
 			FlushAsyncAccelerationQueueLambda();
 		}
diff --git a/Engine/Source/Runtime/Experimental/Chaos/Private/Chaos/PBDRigidsEvolutionGBF.cpp b/Engine/Source/Runtime/Experimental/Chaos/Private/Chaos/PBDRigidsEvolutionGBF.cpp
--- a/Engine/Source/Runtime/Experimental/Chaos/Private/Chaos/PBDRigidsEvolutionGBF.cpp
+++ b/Engine/Source/Runtime/Experimental/Chaos/Private/Chaos/PBDRigidsEvolutionGBF.cpp
@@ -1058,7 +1058,7 @@ void FPBDRigidsEvolutionGBF::Integrate(FReal Dt)
 	};
 	
 
-	if (bSingleWorkerPhysics)
+	if (!Chaos::ShouldExecuteTasks(Particles.GetActiveDynamicMovingKinematicParticlesView().Num()))
 	{
 		// In single threaded mode keep same logic with three batches inside PendingDataArrays: 
 		// 0 Dirty particles from the PushData
@@ -1135,7 +1135,9 @@ void FPBDRigidsEvolutionGBF::ApplyKinematicTargets(const FReal Dt, const FReal S
 		}
 	};
 
-	if (bSingleWorkerPhysics)
+	// At this stage we should be sure if we have launched tasks to linked them
+	// or if we are running on a single thread
+	if (!Chaos::ShouldExecuteTasks(Particles.GetActiveDynamicMovingKinematicParticlesView().Num()))
 	{
 		DirtyKinAsync();
 		// done with update, let's clear the tracking structures
@@ -1283,7 +1285,7 @@ void FPBDRigidsEvolutionGBF::DispatchKinematicTargetsTasks(int32 DispatchBatchIn
 		}
 	};
 
-	if (bSingleWorkerPhysics)
+	if (!Chaos::ShouldExecuteTasks(Particles.GetActiveDynamicMovingKinematicParticlesView().Num()))
 	{
 		KinematicTargetsLambda(0, NumParticles, 0);
 	}
diff --git a/Engine/Source/Runtime/Experimental/Chaos/Private/EventDefaults.cpp b/Engine/Source/Runtime/Experimental/Chaos/Private/EventDefaults.cpp
--- a/Engine/Source/Runtime/Experimental/Chaos/Private/EventDefaults.cpp
+++ b/Engine/Source/Runtime/Experimental/Chaos/Private/EventDefaults.cpp
@@ -291,7 +291,7 @@ namespace Chaos
 #endif
 								}
 							}
-						}, Chaos::SmallBatchSize);
+						}, Chaos::GlobalSmallBatchSize);
 
 						FCollisionDataArray& AllCollisionsDataArray = CollisionEventData.CollisionData.AllCollisionsArray;
 						TMap<IPhysicsProxyBase*, TArray<int32>>& AllCollisionsIndicesByPhysicsProxy = CollisionEventData.PhysicsProxyToCollisionIndices.PhysicsProxyToIndicesMap;
diff --git a/Engine/Source/Runtime/Experimental/Chaos/Private/PBDRigidsSolver.cpp b/Engine/Source/Runtime/Experimental/Chaos/Private/PBDRigidsSolver.cpp
--- a/Engine/Source/Runtime/Experimental/Chaos/Private/PBDRigidsSolver.cpp
+++ b/Engine/Source/Runtime/Experimental/Chaos/Private/PBDRigidsSolver.cpp
@@ -2744,7 +2744,8 @@ namespace Chaos
 	{
 		FPullPhysicsData* PullData = MarshallingManager.GetCurrentPullData_Internal();
 
-		if (!(bSingleWorkerPhysics || CVars::bSingleThreadPullData))
+		const bool bRunMultithreading = ShouldExecuteTasks(GetParticles().GetDirtyParticlesView().Num()) && !CVars::bSingleThreadPullData;
+		if (bRunMultithreading)
 		{
 			using TSOA = typename TParticleView<FPBDRigidParticles>::TSOA;
 			using THandle = typename TSOA::THandleType;
@@ -2861,7 +2862,7 @@ namespace Chaos
 			RewindData->BufferPhysicsResults(PullData->DirtyRigidErrors);
 		}
 
-		if (!(bSingleWorkerPhysics || CVars::bSingleThreadPullData))
+		if (bRunMultithreading)
 		{
 			TaskDispatcher->WaitTaskEndBufferResults();
 		}
diff --git a/Engine/Source/Runtime/Experimental/Chaos/Public/Chaos/Collision/SpatialAccelerationBroadPhase.h b/Engine/Source/Runtime/Experimental/Chaos/Public/Chaos/Collision/SpatialAccelerationBroadPhase.h
--- a/Engine/Source/Runtime/Experimental/Chaos/Public/Chaos/Collision/SpatialAccelerationBroadPhase.h
+++ b/Engine/Source/Runtime/Experimental/Chaos/Public/Chaos/Collision/SpatialAccelerationBroadPhase.h
@@ -656,7 +656,7 @@ namespace Chaos
 			using THandleBase = typename THandle::THandleBase;
 			using TTransientHandle = typename THandle::TTransientHandle;
 
-			const int32 ParticleByContext = FMath::Max(FMath::DivideAndRoundUp(NumParticles, NumActiveBroadphaseContexts), SmallBatchSize);
+			const int32 ParticleByContext = FMath::Max(FMath::DivideAndRoundUp(NumParticles, NumActiveBroadphaseContexts), CollisionSmallBatchSize);
 			const int32 NumBatch = FMath::DivideAndRoundUp(NumParticles, ParticleByContext);
 
 			for (int32 BatchIndex = 0; BatchIndex < NumBatch; ContextIndex++, BatchIndex++)
diff --git a/Engine/Source/Runtime/Experimental/Chaos/Public/Chaos/Framework/Parallel.h b/Engine/Source/Runtime/Experimental/Chaos/Public/Chaos/Framework/Parallel.h
--- a/Engine/Source/Runtime/Experimental/Chaos/Public/Chaos/Framework/Parallel.h
+++ b/Engine/Source/Runtime/Experimental/Chaos/Public/Chaos/Framework/Parallel.h
@@ -14,17 +14,17 @@ namespace Chaos
 	//void CHAOS_API PhysicsParallelFor_RecursiveDivide(int32 InNum, TFunctionRef<void(int32)> InCallable, bool bForceSingleThreaded = false);
 
 	/**
-	 * Helper to call from task execution logic that does not use ParallelFor, to know if it should use parallel execution or not.
-	 * Returns true if the tasks should be run in parallel based on the current context.
-	 * @param InTaskNum The total number of tasks that need to run, will be compared to MinParallelTaskSize
+	 * Helper to call from task execution logic.
+	 * Returns true if the tasks should be run, or if it is better to stay on a single thread and not spawn tasks.
+	 * @param InInstanceNum The total number of instance that should be spread through the tasks.  It will be compared to GlobalSmallBatchSize.
 	 */
-	bool CHAOS_API ShouldExecuteTasksInParallel(int32 InTaskNum);
+	bool CHAOS_API ShouldExecuteTasks(int32 InInstanceNum);
 
 	CHAOS_API extern bool bSingleWorkerPhysics;
 	CHAOS_API extern int32 MaxNumWorkers;
-	CHAOS_API extern int32 SmallBatchSize;
+	CHAOS_API extern int32 CollisionSmallBatchSize;
+	CHAOS_API extern int32 GlobalSmallBatchSize;
 	CHAOS_API extern int32 LargeBatchSize;
-	CHAOS_API extern int32 MinParallelTaskSize;
 	CHAOS_API extern bool bDisablePhysicsParallelFor;
 	CHAOS_API extern bool bDisableParticleParallelFor;
 	CHAOS_API extern bool bDisableCollisionParallelFor;
