Chaos: Fixing filter's ToString functionality and adding tests.

#Synced-CL 51759702

diff --git a/Engine/Source/Runtime/Experimental/Chaos/Private/Chaos/CollisionFilterData.cpp b/Engine/Source/Runtime/Experimental/Chaos/Private/Chaos/CollisionFilterData.cpp
--- a/Engine/Source/Runtime/Experimental/Chaos/Private/Chaos/CollisionFilterData.cpp
+++ b/Engine/Source/Runtime/Experimental/Chaos/Private/Chaos/CollisionFilterData.cpp
@@ -177,7 +177,7 @@ namespace Chaos::Filter
 
 	FString FInstanceData::ToString() const
 	{
-		return FString::Format(TEXT("OwnerId(%d) ComponentId(%d)"), { GetOwnerId(), GetComponentId() });
+		return FString::Format(TEXT("OwnerId({0}) ComponentId({1})"), { GetOwnerId(), GetComponentId() });
 	}
 
 	bool FShapeFilterData::IsValid() const
@@ -297,7 +297,7 @@ namespace Chaos::Filter
 
 	FString FShapeFilterData::ToString() const
 	{
-		return FString::Format(TEXT("Block(%x) Overlap(%x) CollisionChannel(%x) MaskFilter(%x) Flags(%x)"), { GetBlockChannels(), GetOverlapChannels(), GetCollisionChannelIndex(), GetMaskFilter(), (uint8)GetFlags() });
+		return FString::Printf(TEXT("Block(0x%llX) Overlap(0x%llX) CollisionChannel(0x%X) MaskFilter(0x%X) Flags(0x%X)"), GetBlockChannels(), GetOverlapChannels(), GetCollisionChannelIndex(), GetMaskFilter(), (uint8)GetFlags());
 	}
 
 	uint32 FShapeFilterData::GetLegacyWord3() const
@@ -456,11 +456,11 @@ namespace Chaos::Filter
 	{
 		if (GetQueryType() == EQueryType::Channel)
 		{
-			return FString::Format(TEXT("Type(Channel) Block(%x) Overlap(%x) ChannelIndex(%x) IgnoreMask(%d) Flags(%x)"), { GetBlockChannels(), GetOverlapChannels(), GetCollisionChannelIndex(), GetIgnoreMask(), (uint8)GetFlags() });
+			return FString::Printf(TEXT("Type(Channel) Block(0x%llX) Overlap(0x%llX) ChannelIndex(0x%X) IgnoreMask(0x%X) Flags(0x%X)"), GetBlockChannels(), GetOverlapChannels(), GetCollisionChannelIndex(), GetIgnoreMask(), (uint8)GetFlags());
 		}
 		else
 		{
-			return FString::Format(TEXT("Type(Object) IsMulti(%d) ObjectTypes(%x) IgnoreMask(%d) Flags(%x)"), { IsMultiQuery(), GetObjectTypesToQueryMask(), GetIgnoreMask(), (uint8)GetFlags() });
+			return FString::Printf(TEXT("Type(Object) IsMulti(%d) ObjectTypes(0x%llX) IgnoreMask(0x%x) Flags(0x%X)"), IsMultiQuery(), GetObjectTypesToQueryMask(), GetIgnoreMask(), (uint8)GetFlags());
 		}
 	}
 
