Chaos: Add spatial hash point-triangle query, graph coloring support for two particle arrays, and geometry utilities to support the unified deformable solver

#Synced-CL 51662528

diff --git a/Engine/Source/Runtime/Experimental/Chaos/Private/Chaos/Deformable/ChaosDeformableSolver.cpp b/Engine/Source/Runtime/Experimental/Chaos/Private/Chaos/Deformable/ChaosDeformableSolver.cpp
--- a/Engine/Source/Runtime/Experimental/Chaos/Private/Chaos/Deformable/ChaosDeformableSolver.cpp
+++ b/Engine/Source/Runtime/Experimental/Chaos/Private/Chaos/Deformable/ChaosDeformableSolver.cpp
@@ -1853,7 +1853,7 @@ namespace Chaos::Softs
 			int32 TransientIndex = GSMainConstraint->AddTransientConstraintResidualAndHessianRange(1);
 			GSMainConstraint->TransientConstraintResidualAndHessian()[TransientIndex] = [this](const FSolverParticles& Particles, const int32 ConstraintIndex, const int32 ConstraintIndexLocal, const FSolverReal Dt, TVec3<FSolverReal>& ParticleResidual, Chaos::PMatrix<FSolverReal, 3, 3>& ParticleHessian)
 				{
-					this->GSWeakConstraints->AddWCResidual(Particles, ConstraintIndex + this->GSWeakConstraints->InitialWCSize, ConstraintIndexLocal, Dt, ParticleResidual, ParticleHessian);
+					this->GSWeakConstraints->AddWCResidual(Particles, ConstraintIndex + this->GSWeakConstraints->GetInitialWCSize(), ConstraintIndexLocal, Dt, ParticleResidual, ParticleHessian);
 				};
 
 			if (!Property.bEnablePositionTargets)
diff --git a/Engine/Source/Runtime/Experimental/Chaos/Private/Chaos/Deformable/GaussSeidelMainConstraint.cpp b/Engine/Source/Runtime/Experimental/Chaos/Private/Chaos/Deformable/GaussSeidelMainConstraint.cpp
--- a/Engine/Source/Runtime/Experimental/Chaos/Private/Chaos/Deformable/GaussSeidelMainConstraint.cpp
+++ b/Engine/Source/Runtime/Experimental/Chaos/Private/Chaos/Deformable/GaussSeidelMainConstraint.cpp
@@ -13,7 +13,7 @@ namespace Chaos::Softs
 
 	FSolverReal CGTol = 1e-4f;
 
-	FAutoConsoleVariableRef CVarClothCGTol(TEXT("p.Chaos.Cloth.CGTol"), MaxItCG, TEXT("CG Tolerance [def: 1e-4]"));
+	FAutoConsoleVariableRef CVarClothCGTol(TEXT("p.Chaos.Cloth.CGTol"), CGTol, TEXT("CG Tolerance [def: 1e-4]"));
 
 	void FGaussSeidelMainConstraints::AddStaticConstraints(const TArray<TArray<int32>>& ExtraConstraints, TArray<TArray<int32>>& ExtraIncidentElements, TArray<TArray<int32>>& ExtraIncidentElementsLocal)
 	{	
@@ -182,16 +182,18 @@ namespace Chaos::Softs
 		PassedIters++;
 
 
+#if WITH_EDITOR
 		if (Write2File)
 		{
-			FString file = FPaths::ProjectDir();
-			file.Append(TEXT("/DebugOutput/NewtonResidual.txt"));
+			FString File = FPaths::ProjectDir();
+			File.Append(TEXT("/DebugOutput/NewtonResidual.txt"));
 			if (PassedIters == 0)
 			{
-				FFileHelper::SaveStringToFile(FString("Newton Norm\r\n"), *file);
+				FFileHelper::SaveStringToFile(FString("Newton Norm\r\n"), *File);
 			}
-			FFileHelper::SaveStringToFile((FString::SanitizeFloat(NewtonNorm) + FString(",\r\n")), *file, FFileHelper::EEncodingOptions::AutoDetect, &IFileManager::Get(), EFileWrite::FILEWRITE_Append);
+			FFileHelper::SaveStringToFile((FString::SanitizeFloat(NewtonNorm) + FString(",\r\n")), *File, FFileHelper::EEncodingOptions::AutoDetect, &IFileManager::Get(), EFileWrite::FILEWRITE_Append);
 		}
+#endif
 
 		return NewtonResidual;
 
@@ -211,15 +213,13 @@ namespace Chaos::Softs
 			}
 		};
 
-		int32 max_it_cg = MaxItCG;
-		FSolverReal cg_tol = CGTol;
 		int32 MaxNewtonIt = 1;
 
 		for (int32 It = 0; It < MaxNewtonIt; It++)
 		{
 			const TArray<FSolverVec3> Residual = ComputeNewtonResiduals(Particles, Dt, false, nullptr);
 
-			auto multiply = [this, &ProjectBCs, &Dt, &Particles](TArray<FSolverVec3>& y, const TArray<FSolverVec3>& x) 
+			auto multiply = [this, &ProjectBCs, &Dt, &Particles](TArray<FSolverVec3>& y, const TArray<FSolverVec3>& x)
 			{
 				TArray<FSolverVec3> XProj = x;
 				ProjectBCs(XProj);
@@ -231,8 +231,8 @@ namespace Chaos::Softs
 					this->AddInternalForceDifferentials[i](Particles, XProj, y);
 				}
 
-				PhysicsParallelFor(y.Num(), 
-					[&y, &Particles, &XProj, Dt](const int32 i) 
+				PhysicsParallelFor(y.Num(),
+					[&y, &Particles, &XProj, Dt](const int32 i)
 					{
 						y[i] = Particles.M(i) * XProj[i] + Dt * Dt * y[i];
 					});
@@ -243,12 +243,431 @@ namespace Chaos::Softs
 			TArray<FSolverVec3> Deltax;
 			Deltax.Init(FSolverVec3(0.f), Particles.Size());
 
-			Chaos::LanczosCG<FSolverReal>(multiply, Deltax, Residual, MaxItCG, CGTol, use_list.Get()); 
+			Chaos::LanczosCG<FSolverReal>(multiply, Deltax, Residual, MaxItCG, CGTol, use_list.Get());
 
-			PhysicsParallelFor(Particles.Size(), [&Particles, &Deltax](const int32 i) 
+			PhysicsParallelFor(Particles.Size(), [&Particles, &Deltax](const int32 i)
 				{
 					Particles.P(i) -= Deltax[i];
 				});
 		}
 	}
+
+	// -------------------------------------------------------------------
+	// Methods moved from header
+	// -------------------------------------------------------------------
+
+	FGaussSeidelMainConstraints::FGaussSeidelMainConstraints(
+		const FSolverParticles& InParticles,
+		const bool bDoQuasistaticsIn,
+		const bool bDoSORIn,
+		const FSolverReal InOmegaSOR,
+		const int32 ParallelMaxIn,
+		const FSolverReal MaxDxRatioIn,
+		const FDeformableXPBDCorotatedParams& InParams)
+		: bDoQuasistatics(bDoQuasistaticsIn)
+		, bDoAcceleration(bDoSORIn)
+		, OmegaSOR(InOmegaSOR)
+		, ParallelMax(ParallelMaxIn)
+		, CorotatedParams(InParams)
+	{
+		Resize((int32)InParticles.Size());
+
+		InitializeLambdas();
+
+		TVec3<FSolverReal> MaxCoord((FSolverReal)100.), MinCoord((FSolverReal)-100.);
+		for (int32 i = 0; i < (int32)InParticles.Size(); i++)
+		{
+			for (int32 j = 0; j < 3; j++)
+			{
+				if (InParticles.X(i)[j] < MinCoord[j])
+				{
+					MinCoord[j] = InParticles.X(i)[j];
+				}
+				if (InParticles.X(i)[j] > MaxCoord[j])
+				{
+					MaxCoord[j] = InParticles.X(i)[j];
+				}
+			}
+		}
+		MaxDxSize = (MaxCoord - MinCoord).Size() * MaxDxRatioIn;
+	}
+
+	FGaussSeidelMainConstraints::FGaussSeidelMainConstraints(const FSolverParticles& InParticles, const FCollectionPropertyConstFacade& Property, const FSolverReal MaxDxRatioIn)
+		: bDoQuasistatics(GetDoQuasistatics(Property, false))
+		, bDoAcceleration(GetAccelerateSolverUsingSOR(Property, true))
+		, OmegaSOR(GetOmegaSOR(Property, (FSolverReal)1.6))
+		, ParallelMax(1000)
+		, CorotatedParams(FDeformableXPBDCorotatedParams())
+	{
+		Resize(InParticles.Size());
+		InitializeLambdas();
+		TVec3<FSolverReal> MaxCoord((FSolverReal)100.), MinCoord((FSolverReal)-100.);
+		for (int32 Vdx = 0; Vdx < (int32)InParticles.Size(); ++Vdx)
+		{
+			MaxCoord.ComponentwiseMax(InParticles.X(Vdx));
+			MinCoord.ComponentwiseMin(InParticles.X(Vdx));
+		}
+		MaxDxSize = (MaxCoord - MinCoord).Size() * MaxDxRatioIn;
+	}
+
+	void FGaussSeidelMainConstraints::Resize(const int32 NewSize)
+	{
+		xtilde.SetNum(NewSize);
+		StaticIncidentElements.SetNum(NewSize);
+		StaticIncidentElementsLocal.SetNum(NewSize);
+		DynamicIncidentElements.SetNum(NewSize);
+		DynamicIncidentElementsLocal.SetNum(NewSize);
+		TransientIncidentElements.SetNum(NewSize);
+		TransientIncidentElementsLocal.SetNum(NewSize);
+		X_k_1.Init(TVector<FSolverReal, 3>(FSolverReal(0.)), NewSize);
+		X_k.Init(TVector<FSolverReal, 3>(FSolverReal(0.)), NewSize);
+	}
+
+	int32 FGaussSeidelMainConstraints::AddStaticConstraintResidualAndHessianRange(int32 NumConstraints)
+	{
+		int32 CurrentSize = AddStaticConstraintResidualAndHessian.Num();
+		AddStaticConstraintResidualAndHessian.AddDefaulted(NumConstraints);
+		return CurrentSize;
+	}
+
+	int32 FGaussSeidelMainConstraints::AddTransientConstraintResidualAndHessianRange(int32 NumConstraints)
+	{
+		int32 CurrentSize = AddTransientConstraintResidualAndHessian.Num();
+		AddTransientConstraintResidualAndHessian.AddDefaulted(NumConstraints);
+		return CurrentSize;
+	}
+
+	int32 FGaussSeidelMainConstraints::AddDynamicConstraintResidualAndHessianRange(int32 NumConstraints)
+	{
+		int32 CurrentSize = AddDynamicConstraintResidualAndHessian.Num();
+		AddDynamicConstraintResidualAndHessian.AddDefaulted(NumConstraints);
+		return CurrentSize;
+	}
+
+	int32 FGaussSeidelMainConstraints::AddCrossBodyStaticConstraintResidualAndHessianRange(int32 NumConstraints)
+	{
+		int32 CurrentSize = AddCrossBodyStaticConstraintResidualAndHessian.Num();
+		AddCrossBodyStaticConstraintResidualAndHessian.AddDefaulted(NumConstraints);
+		return CurrentSize;
+	}
+
+	int32 FGaussSeidelMainConstraints::AddCrossBodyTransientConstraintResidualAndHessianRange(int32 NumConstraints)
+	{
+		int32 CurrentSize = AddCrossBodyTransientConstraintResidualAndHessian.Num();
+		AddCrossBodyTransientConstraintResidualAndHessian.AddDefaulted(NumConstraints);
+		return CurrentSize;
+	}
+
+	int32 FGaussSeidelMainConstraints::AddPerNodeHessianRange(int32 NumConstraints)
+	{
+		int32 CurrentSize = AddPerNodeHessian.Num();
+		AddPerNodeHessian.AddDefaulted(NumConstraints);
+		return CurrentSize;
+	}
+
+	int32 FGaussSeidelMainConstraints::AddTransientPerNodeHessianRange(int32 NumConstraints)
+	{
+		int32 CurrentSize = AddTransientPerNodeHessian.Num();
+		AddTransientPerNodeHessian.AddDefaulted(NumConstraints);
+		return CurrentSize;
+	}
+
+	int32 FGaussSeidelMainConstraints::AddAddInternalForceDifferentialsRange(int32 NumConstraints)
+	{
+		int32 CurrentSize = AddInternalForceDifferentials.Num();
+		AddInternalForceDifferentials.AddDefaulted(NumConstraints);
+		return CurrentSize;
+	}
+
+	void FGaussSeidelMainConstraints::ResetDynamicConstraints()
+	{
+		DynamicConstraints = {};
+		for (int32 ParticleIdx = 0; ParticleIdx < DynamicIncidentElements.Num(); ParticleIdx++)
+		{
+			DynamicIncidentElements[ParticleIdx].SetNum(0);
+			DynamicIncidentElementsLocal[ParticleIdx].SetNum(0);
+		}
+		DynamicIncidentElementsOffsets = {};
+	}
+
+	void FGaussSeidelMainConstraints::ResetTransientConstraints()
+	{
+		TransientConstraints = {};
+		for (int32 ParticleIdx = 0; ParticleIdx < TransientIncidentElements.Num(); ParticleIdx++)
+		{
+			TransientIncidentElements[ParticleIdx].SetNum(0);
+			TransientIncidentElementsLocal[ParticleIdx].SetNum(0);
+		}
+		TransientIncidentElementsOffsets = {};
+	}
+
+	void FGaussSeidelMainConstraints::ResetTransientConstraintsAndRules()
+	{
+		ResetTransientConstraints();
+		AddTransientConstraintResidualAndHessian.Empty();
+		AddTransientPerNodeHessian.Empty();
+	}
+
+	void FGaussSeidelMainConstraints::Apply(FSolverParticles& Particles, const FSolverReal Dt, const int32 MaxWriteIters, const bool Write2File, const TPBDActiveView<FSolverParticles>* InParticleActiveView)
+	{
+		PERF_SCOPE(STAT_ChaosGSMainConstraint_Apply);
+
+		if (DebugResidual && PassedIters < MaxWriteIters)
+		{
+			ComputeNewtonResiduals(Particles, Dt, Write2File);
+		}
+
+		std::atomic<int32> ParticleFailureCounter(0);
+
+		for (int32 i = 0; i < ParticlesPerColor.Num(); i++)
+		{
+			int32 NumBatch = ParticlesPerColor[i].Num() / CorotatedParams.XPBDCorotatedBatchSize;
+			if (ParticlesPerColor[i].Num() % CorotatedParams.XPBDCorotatedBatchSize != 0)
+			{
+				NumBatch ++;
+			}
+
+			PhysicsParallelFor(NumBatch, [&](const int32 BatchIndex)
+				{
+					for (int32 BatchSubIndex = 0; BatchSubIndex < CorotatedParams.XPBDCorotatedBatchSize; BatchSubIndex++) {
+						int32 TaskIndex = CorotatedParams.XPBDCorotatedBatchSize * BatchIndex + BatchSubIndex;
+						if (TaskIndex < ParticlesPerColor[i].Num())
+						{
+							if (Particles.InvM(ParticlesPerColor[i][TaskIndex]) != FSolverReal(0))
+							{
+								int32 ParticleIndex = ParticlesPerColor[i][TaskIndex];
+								if (!ApplySingleParticle(ParticleIndex, Dt, Particles))
+								{
+									ParticleFailureCounter.fetch_add(1, std::memory_order_relaxed);
+								}
+							}
+						}
+					}
+				}, NumBatch < CorotatedParams.XPBDCorotatedBatchThreshold);
+
+		}
+
+
+		if (bDoAcceleration)
+		{
+			PERF_SCOPE(STAT_ChaosGSMainConstraint_Acceleration);
+			if (InParticleActiveView)
+			{
+				InParticleActiveView->ParallelFor(AccelerationTechniquePerParticle, CorotatedParams.XPBDCorotatedBatchSize);
+			}
+			else
+			{
+				PhysicsParallelFor(Particles.Size(), [this, &Particles](const int32 ParticleIndex)
+					{
+						this->AccelerationTechniquePerParticle(Particles, ParticleIndex);
+					}, Particles.Size() < 1000);
+			}
+		}
+
+		CurrentIt ++;
+
+		if (ParticleFailureCounter.load(std::memory_order_relaxed) > 0)
+		{
+			UE_LOG(LogDeformableGaussSeidelMainConstraint, Warning, TEXT("%d Particle(s) are skipped because of too large dx size"), ParticleFailureCounter.load(std::memory_order_relaxed));
+		}
+	}
+
+	void FGaussSeidelMainConstraints::InitStaticColor(const FSolverParticles& Particles, const TPBDActiveView<FSolverParticles>* InParticleActiveView)
+	{
+		StaticParticlesPerColor = ComputeNodalColoring(StaticConstraints, Particles, 0, Particles.Size(), StaticIncidentElements, StaticIncidentElementsLocal, InParticleActiveView,  &StaticParticleColors);
+		ParticleColors = StaticParticleColors;
+		ParticlesPerColor = StaticParticlesPerColor;
+	}
+
+	void FGaussSeidelMainConstraints::InitTransientColor(const FSolverParticles& Particles)
+	{
+		PERF_SCOPE(STAT_ChaosGSMainConstraint_InitTransientColor);
+		ParticleColors = StaticParticleColors;
+		ParticlesPerColor = StaticParticlesPerColor;
+		Chaos::ComputeExtraNodalColoring(StaticConstraints, DynamicConstraints, TransientConstraints, Particles, StaticIncidentElements, DynamicIncidentElements, TransientIncidentElements, ParticleColors, ParticlesPerColor);
+	}
+
+	void FGaussSeidelMainConstraints::InitDynamicColor(const FSolverParticles& Particles)
+	{
+		PERF_SCOPE(STAT_ChaosGSMainConstraint_InitDynamicColor);
+		ParticleColors = StaticParticleColors;
+		ParticlesPerColor = StaticParticlesPerColor;
+		Chaos::ComputeExtraNodalColoring(StaticConstraints, DynamicConstraints, Particles, StaticIncidentElements, TransientIncidentElements, ParticleColors, ParticlesPerColor);
+	}
+
+	void FGaussSeidelMainConstraints::Init(const FSolverReal Dt, const FSolverParticles& Particles)
+	{
+		Resize((int32)Particles.Size());
+
+		PERF_SCOPE(STAT_ChaosGSMainConstraint_Init);
+		if (!bDoQuasistatics)
+		{
+			for (int32 i = 0; i < xtilde.Num(); i++)
+			{
+				xtilde[i] = Particles.P(i); // Here P = xn + dt*(vn + dt*an) after evolution initial guess
+			}
+		}
+		CurrentIt = 0;
+	}
+
+	bool FGaussSeidelMainConstraints::IsClean(const TArray<TArray<int32>>& ConstraintsIn, const TArray<TArray<int32>>& IncidentElementsIn, const TArray<TArray<int32>>& IncidentElementsLocalIn)
+	{
+		if (IncidentElementsIn.Num() == IncidentElementsLocalIn.Num())
+		{
+			int32 TotalEntries = 0;
+			for (int32 i = 0; i < IncidentElementsIn.Num(); i++)
+			{
+				TotalEntries += IncidentElementsIn[i].Num();
+				if (IncidentElementsIn[i].Num() != IncidentElementsLocalIn[i].Num())
+				{
+					return false;
+				}
+				for (int32 j = 0; j < IncidentElementsIn[i].Num(); j++)
+				{
+					if (IncidentElementsIn[i][j] >= ConstraintsIn.Num() || IncidentElementsLocalIn[i][j] >= ConstraintsIn[IncidentElementsIn[i][j]].Num())
+					{
+						return false;
+					}
+				}
+			}
+			if (TotalEntries > 0)
+			{
+				return true;
+			}
+		}
+		return false;
+	}
+
+	bool FGaussSeidelMainConstraints::ApplySingleParticle(const int32 ParticleIdx, const FSolverReal Dt, FSolverParticles& Particles)
+	{
+		int32 ConstraintIndex = 0;
+		Chaos::TVector<FSolverReal, 3> ParticleResidual((FSolverReal)0.);
+		Chaos::PMatrix<FSolverReal, 3, 3> ParticleHessian((FSolverReal)0., (FSolverReal)0., (FSolverReal)0.);
+
+		ComputeInitialResidualAndHessian(Particles, ParticleIdx, Dt, ParticleResidual, ParticleHessian);
+
+		for (int32 i = 0; i < StaticIncidentElements[ParticleIdx].Num(); i++)
+		{
+			while (StaticIncidentElements[ParticleIdx][i] >= StaticIncidentElementsOffsets[ConstraintIndex + 1] && ConstraintIndex < StaticIncidentElementsOffsets.Num() - 1)
+			{
+				ConstraintIndex ++;
+			}
+
+			AddStaticConstraintResidualAndHessian[ConstraintIndex](Particles, StaticIncidentElements[ParticleIdx][i] - StaticIncidentElementsOffsets[ConstraintIndex], StaticIncidentElementsLocal[ParticleIdx][i], Dt, ParticleResidual, ParticleHessian);
+		}
+
+		ConstraintIndex = 0;
+
+		for (int32 i = 0; i < DynamicIncidentElements[ParticleIdx].Num(); i++)
+		{
+			while (DynamicIncidentElements[ParticleIdx][i] >= DynamicIncidentElementsOffsets[ConstraintIndex + 1] && ConstraintIndex < DynamicIncidentElementsOffsets.Num() - 1)
+			{
+				ConstraintIndex ++;
+			}
+
+			AddDynamicConstraintResidualAndHessian[ConstraintIndex](Particles, DynamicIncidentElements[ParticleIdx][i] - DynamicIncidentElementsOffsets[ConstraintIndex], DynamicIncidentElementsLocal[ParticleIdx][i], Dt, ParticleResidual, ParticleHessian);
+		}
+
+		ConstraintIndex = 0;
+
+		for (int32 i = 0; i < TransientIncidentElements[ParticleIdx].Num(); i++)
+		{
+			while (TransientIncidentElements[ParticleIdx][i] >= TransientIncidentElementsOffsets[ConstraintIndex + 1] && ConstraintIndex < TransientIncidentElementsOffsets.Num() - 1)
+			{
+				ConstraintIndex ++;
+			}
+
+			AddTransientConstraintResidualAndHessian[ConstraintIndex](Particles, TransientIncidentElements[ParticleIdx][i] - TransientIncidentElementsOffsets[ConstraintIndex], TransientIncidentElementsLocal[ParticleIdx][i], Dt, ParticleResidual, ParticleHessian);
+		}
+
+		for (int32 i = 0; i < AddPerNodeHessian.Num(); i++)
+		{
+			AddPerNodeHessian[i](ParticleIdx, Dt, ParticleHessian);
+		}
+		for (int32 i = 0; i < AddTransientPerNodeHessian.Num(); i++)
+		{
+			AddTransientPerNodeHessian[i](ParticleIdx, Dt, ParticleHessian);
+		}
+		FSolverReal HessianScale = (FSolverReal)1.;
+		FSolverReal HessianDet = ParticleHessian.Determinant();
+		auto IsIllConditioned = [](FSolverReal HessianDet)
+		{
+			return FMath::Abs(HessianDet) < TMathUtilConstants<FSolverReal>::Epsilon ||
+				FMath::Abs(HessianDet) > TMathUtilConstants<FSolverReal>::MaxReal;
+		};
+		if (IsIllConditioned(HessianDet))
+		{
+			// scale the hessian so that the determinant (~HessianScale^3) falls into the normal range.
+			HessianScale = (FSolverReal)0.;
+			for (int32 RowIdx = 0; RowIdx < 3; ++RowIdx)
+			{
+				for (int32 ColIdx = 0; ColIdx < 3; ++ColIdx)
+				{
+					HessianScale = FMath::Max(FMath::Abs(ParticleHessian.GetAt(RowIdx, ColIdx)), HessianScale);
+				}
+			}
+			if (IsIllConditioned(HessianScale))
+			{
+				return false;
+			}
+			else
+			{
+				ParticleHessian *= 1 / HessianScale;
+				HessianDet = ParticleHessian.Determinant();
+				if (IsIllConditioned(HessianDet))
+				{
+					return false;
+				}
+			}
+		}
+
+		Chaos::PMatrix<FSolverReal, 3, 3> HessianInv = ParticleHessian.SymmetricCofactorMatrix();
+		HessianInv *= FSolverReal(1) / HessianDet;
+		Chaos::TVector<FSolverReal, 3> Dx = HessianInv.GetTransposed() * (-ParticleResidual) / HessianScale; // add back HessianScale
+
+		if (Dx.Size() < MaxDxSize)
+		{
+			Particles.P(ParticleIdx) += Dx;
+		}
+		else
+		{
+			return false;
+		}
+		return true;
+	}
+
+	void FGaussSeidelMainConstraints::InitializeLambdas()
+	{
+		ComputeInitialResidualAndHessian = [this](const FSolverParticles& Particles, const int32 ParticleIdx, const FSolverReal Dt, TVec3<FSolverReal>& ParticleResidual, Chaos::PMatrix<FSolverReal, 3, 3>& ParticleHessian)
+		{
+			if (!this->bDoQuasistatics) {
+				for (int32 Alpha = 0; Alpha < 3; Alpha++) {
+					ParticleResidual[Alpha] = Particles.M(ParticleIdx) * (Particles.P(ParticleIdx)[Alpha] - this->xtilde[ParticleIdx][Alpha]);
+				}
+				for (int32 Alpha = 0; Alpha < 3; Alpha++) {
+					ParticleHessian.SetAt(Alpha, Alpha, Particles.M(ParticleIdx));
+				}
+			}
+			else
+			{
+				for (int32 Alpha = 0; Alpha < 3; Alpha++) {
+					ParticleResidual[Alpha] = -Dt * Dt * ExternalAcceleration[Alpha] * Particles.M(ParticleIdx);
+				}
+			}
+		};
+
+		AccelerationTechniquePerParticle = [this](FSolverParticles& Particles, int32 ParticleIndex)
+		{
+			if (Particles.InvM(ParticleIndex) != FSolverReal(0) && CurrentIt > SORStart)
+			{
+				Particles.P(ParticleIndex) = OmegaSOR * (Particles.P(ParticleIndex) - this->X_k_1[ParticleIndex]) + this->X_k_1[ParticleIndex];
+			}
+			if (Particles.InvM(ParticleIndex) != FSolverReal(0))
+			{
+				this->X_k_1[ParticleIndex] = this->X_k[ParticleIndex];
+				this->X_k[ParticleIndex] = Particles.P(ParticleIndex);
+			}
+		};
+	}
+
 }
\ No newline at end of file
diff --git a/Engine/Source/Runtime/Experimental/Chaos/Private/Chaos/Deformable/GaussSeidelWeakConstraints.cpp b/Engine/Source/Runtime/Experimental/Chaos/Private/Chaos/Deformable/GaussSeidelWeakConstraints.cpp
new file mode 100644
--- /dev/null
+++ b/Engine/Source/Runtime/Experimental/Chaos/Private/Chaos/Deformable/GaussSeidelWeakConstraints.cpp
@@ -0,0 +1,959 @@
+// Copyright Epic Games, Inc. All Rights Reserved.
+
+#include "Chaos/Deformable/GaussSeidelWeakConstraints.h"
+#include "Chaos/HierarchicalSpatialHash.h"
+
+namespace Chaos::Softs
+{
+	using Chaos::TVec3;
+
+	// -------------------------------------------------------------------
+	// FGaussSeidelSpringConstraintData
+	// -------------------------------------------------------------------
+
+	FGaussSeidelSpringConstraintData::FGaussSeidelSpringConstraintData()
+	{
+		AddArray(&MIndices);
+		AddArray(&MSecondIndices);
+		AddArray(&MWeights);
+		AddArray(&MSecondWeights);
+		AddArray(&MStiffness);
+		AddArray(&MIsAnisotropic);
+		AddArray(&MNormals);
+		AddArray(&MIsZeroRestLength);
+		AddArray(&MRestLength);
+	}
+
+	FGaussSeidelSpringConstraintData::FGaussSeidelSpringConstraintData(const FGaussSeidelSpringConstraintData& Other)
+		: FGaussSeidelSpringConstraintData()
+	{
+		AddConstraints(Other.Size());
+		for (int32 ConstraintIdx = 0; ConstraintIdx < Other.Size(); ++ConstraintIdx)
+		{
+			SetSingleConstraint(Other.GetSingleConstraintData(ConstraintIdx), ConstraintIdx);
+		}
+	}
+
+	FGaussSeidelSpringConstraintData::FGaussSeidelSpringConstraintData(FGaussSeidelSpringConstraintData&& Other)
+		: TArrayCollection(), MIndices(MoveTemp(Other.MIndices))
+		, MSecondIndices(MoveTemp(Other.MSecondIndices))
+		, MWeights(MoveTemp(Other.MWeights))
+		, MSecondWeights(MoveTemp(Other.MSecondWeights))
+		, MStiffness(MoveTemp(Other.MStiffness))
+		, MIsAnisotropic(MoveTemp(Other.MIsAnisotropic))
+		, MNormals(MoveTemp(Other.MNormals))
+		, MIsZeroRestLength(MoveTemp(Other.MIsZeroRestLength))
+		, MRestLength(MoveTemp(Other.MRestLength))
+	{
+		AddArray(&MIndices);
+		AddArray(&MSecondIndices);
+		AddArray(&MWeights);
+		AddArray(&MSecondWeights);
+		AddArray(&MStiffness);
+		AddArray(&MIsAnisotropic);
+		AddArray(&MNormals);
+		AddArray(&MIsZeroRestLength);
+		AddArray(&MRestLength);
+		ResizeHelper(Other.MSize);
+		Other.MSize = 0;
+	}
+
+	FGaussSeidelSpringConstraintData& FGaussSeidelSpringConstraintData::operator=(FGaussSeidelSpringConstraintData&& Other)
+	{
+		MIndices = MoveTemp(Other.MIndices);
+		MSecondIndices = MoveTemp(Other.MSecondIndices);
+		MWeights = MoveTemp(Other.MWeights);
+		MSecondWeights = MoveTemp(Other.MSecondWeights);
+		MStiffness = MoveTemp(Other.MStiffness);
+		MIsAnisotropic = MoveTemp(Other.MIsAnisotropic);
+		MNormals = MoveTemp(Other.MNormals);
+		MIsZeroRestLength = MoveTemp(Other.MIsZeroRestLength);
+		MRestLength = MoveTemp(Other.MRestLength);
+		ResizeHelper(Other.Size());
+		Other.MSize = 0;
+		return *this;
+	}
+
+	void FGaussSeidelSpringConstraintData::SetSingleConstraint(const FGaussSeidelSpringConstraintSingleData& SingleData, const int32 ConstraintIndex)
+	{
+		MIndices[ConstraintIndex] = SingleData.SingleIndices;
+		MSecondIndices[ConstraintIndex] = SingleData.SingleSecondIndices;
+		MStiffness[ConstraintIndex] = SingleData.SingleStiffness;
+		MWeights[ConstraintIndex] = SingleData.SingleWeights;
+		MSecondWeights[ConstraintIndex] = SingleData.SingleSecondWeights;
+		MNormals[ConstraintIndex] = SingleData.SingleNormal;
+		MIsAnisotropic[ConstraintIndex] = SingleData.bIsAnisotropic;
+		MIsZeroRestLength[ConstraintIndex] = SingleData.bIsZeroRestLength;
+		MRestLength[ConstraintIndex] = SingleData.RestLength;
+	}
+
+	void FGaussSeidelSpringConstraintData::AddSingleConstraint(const FGaussSeidelSpringConstraintSingleData& SingleData)
+	{
+		AddConstraints(1);
+		SetSingleConstraint(SingleData, MSize - 1);
+	}
+
+	const FGaussSeidelSpringConstraintSingleData FGaussSeidelSpringConstraintData::GetSingleConstraintData(const int32 ConstraintIndex) const
+	{
+		FGaussSeidelSpringConstraintSingleData SingleConstraintData;
+		if (ConstraintIndex > INDEX_NONE && static_cast<uint32>(ConstraintIndex) < MSize)
+		{
+			SingleConstraintData.SingleIndices = MIndices[ConstraintIndex];
+			SingleConstraintData.SingleSecondIndices = MSecondIndices[ConstraintIndex];
+			SingleConstraintData.SingleStiffness = MStiffness[ConstraintIndex];
+			SingleConstraintData.SingleWeights = MWeights[ConstraintIndex];
+			SingleConstraintData.SingleSecondWeights = MSecondWeights[ConstraintIndex];
+			SingleConstraintData.bIsAnisotropic = MIsAnisotropic[ConstraintIndex];
+			SingleConstraintData.SingleNormal = MNormals[ConstraintIndex];
+			SingleConstraintData.bIsZeroRestLength = MIsZeroRestLength[ConstraintIndex];
+			SingleConstraintData.RestLength = MRestLength[ConstraintIndex];
+		}
+		return SingleConstraintData;
+	}
+
+	// -------------------------------------------------------------------
+	// FGaussSeidelSpringConstraints
+	// -------------------------------------------------------------------
+
+	FGaussSeidelSpringConstraints::FGaussSeidelSpringConstraints(
+		const TArray<TArray<int32>>& InIndices,
+		const TArray<TArray<FSolverReal>>& InWeights,
+		const TArray<FSolverReal>& InStiffness,
+		const TArray<TArray<int32>>& InSecondIndices,
+		const TArray<TArray<FSolverReal>>& InSecondWeights,
+		const TArray<bool>& InIsAnisotropic,
+		const TArray<bool>& InZeroRestLength,
+		const FDeformableXPBDWeightedSpringConstraintParams& InParams
+	) : DebugDrawParams(InParams)
+	{
+		ensureMsgf(InIndices.Num() == InSecondIndices.Num(), TEXT("Input Double Bindings have wrong size"));
+
+		ConstraintsData.AddConstraints(InIndices.Num());
+
+		for (int32 SpringIdx = 0; SpringIdx < InIndices.Num(); ++SpringIdx)
+		{
+			FGaussSeidelSpringConstraintSingleData SingleConstraintData;
+			SingleConstraintData.SingleIndices = InIndices[SpringIdx];
+			SingleConstraintData.SingleSecondIndices = InSecondIndices[SpringIdx];
+			SingleConstraintData.SingleWeights = InWeights[SpringIdx];
+			SingleConstraintData.SingleSecondWeights = InSecondWeights[SpringIdx];
+			SingleConstraintData.SingleStiffness = InStiffness[SpringIdx];
+			SingleConstraintData.bIsAnisotropic = InIsAnisotropic[SpringIdx];
+			SingleConstraintData.bIsZeroRestLength = InZeroRestLength[SpringIdx];
+			ConstraintsData.SetSingleConstraint(SingleConstraintData, SpringIdx);
+		}
+	}
+
+PRAGMA_DISABLE_DEPRECATION_WARNINGS
+	FGaussSeidelSpringConstraints::FGaussSeidelSpringConstraints(
+		const TArray<TArray<int32>>& InIndices,
+		const TArray<TArray<FSolverReal>>& InWeights,
+		const TArray<FSolverReal>& InStiffness,
+		const TArray<TArray<int32>>& InSecondIndices,
+		const TArray<TArray<FSolverReal>>& InSecondWeights,
+		const FDeformableXPBDWeightedSpringConstraintParams& InParams
+	) : DebugDrawParams(InParams)
+	{
+		ensureMsgf(InIndices.Num() == InSecondIndices.Num(), TEXT("Input Double Bindings have wrong size"));
+
+		ConstraintsData.AddConstraints(InIndices.Num());
+
+		for (int32 i = 0; i < InIndices.Num(); i++)
+		{
+			FGaussSeidelSpringConstraintSingleData SingleConstraintData;
+			SingleConstraintData.SingleIndices = InIndices[i];
+			SingleConstraintData.SingleSecondIndices = InSecondIndices[i];
+			SingleConstraintData.SingleWeights = InWeights[i];
+			SingleConstraintData.SingleSecondWeights = InSecondWeights[i];
+			SingleConstraintData.SingleStiffness = InStiffness[i];
+			ConstraintsData.SetSingleConstraint(SingleConstraintData, i);
+		}
+
+		for (int32 i = 0; i < ConstraintsData.Size(); i++)
+		{
+			const TArray<int32>& SingleIndices = ConstraintsData.GetIndices(i);
+			const TArray<int32>& SingleSecondIndices = ConstraintsData.GetSecondIndices(i);
+			for (int32 j = 0; j < SingleSecondIndices.Num(); j++)
+			{
+				ensureMsgf(!SingleIndices.Contains(SingleSecondIndices[j]), TEXT("Indices and Second Indices overlaps. Currently not supported"));
+			}
+		}
+	}
+PRAGMA_ENABLE_DEPRECATION_WARNINGS
+
+	void FGaussSeidelSpringConstraints::ComputeInitialWCData(const FSolverParticlesRange& InParticles)
+	{
+		ComputeIncidentElements();
+		//Update rest state normal and nodal weights
+		UpdateTriangleNormalAndNodalWeight(InParticles, /*bUseParticleX = */true);
+		SetNoCollisionData();
+		SetConstraintsRestLength(InParticles, InParticles);
+	}
+
+	void FGaussSeidelSpringConstraints::ComputeInitialWCData(const FSolverParticles& InParticles)
+	{
+		ComputeInitialWCData(FSolverParticlesRange(const_cast<FSolverParticles*>(&InParticles), /*InOffset = */ 0, /*InRange = */ InParticles.Size()));
+	}
+
+	void FGaussSeidelSpringConstraints::AddWCHessian(const int32 ParticleIdx, const FSolverReal Dt, Chaos::PMatrix<FSolverReal, 3, 3>& ParticleHessian) const
+	{
+		if (NodalWeights[ParticleIdx].Num() > 0)
+		{
+			for (int32 Alpha = 0; Alpha < 3; Alpha++)
+			{
+				ParticleHessian.SetAt(Alpha, Alpha, ParticleHessian.GetAt(Alpha, Alpha) + Dt * Dt * NodalWeights[ParticleIdx][Alpha]);
+			}
+
+			ParticleHessian.SetAt(0, 1, ParticleHessian.GetAt(0, 1) + Dt * Dt * NodalWeights[ParticleIdx][3]);
+			ParticleHessian.SetAt(0, 2, ParticleHessian.GetAt(0, 2) + Dt * Dt * NodalWeights[ParticleIdx][4]);
+			ParticleHessian.SetAt(1, 2, ParticleHessian.GetAt(1, 2) + Dt * Dt * NodalWeights[ParticleIdx][5]);
+			ParticleHessian.SetAt(1, 0, ParticleHessian.GetAt(1, 0) + Dt * Dt * NodalWeights[ParticleIdx][3]);
+			ParticleHessian.SetAt(2, 0, ParticleHessian.GetAt(2, 0) + Dt * Dt * NodalWeights[ParticleIdx][4]);
+			ParticleHessian.SetAt(2, 1, ParticleHessian.GetAt(2, 1) + Dt * Dt * NodalWeights[ParticleIdx][5]);
+		}
+	}
+
+	void FGaussSeidelSpringConstraints::AddExtraConstraints(const TArray<TArray<int32>>& InIndices,
+								const TArray<TArray<FSolverReal>>& InWeights,
+								const TArray<FSolverReal>& InStiffness,
+								const TArray<TArray<int32>>& InSecondIndices,
+								const TArray<TArray<FSolverReal>>& InSecondWeights,
+								const TArray<bool>& InIsAnisotropic,
+								const TArray<bool>& InIsZeroRestLength)
+	{
+		const int32 Offset = ConstraintsData.Size();
+
+		ConstraintsData.AddConstraints(InIndices.Num());
+
+		for (int32 i = 0; i < InIndices.Num(); i++)
+		{
+			FGaussSeidelSpringConstraintSingleData SingleConstraintData;
+			SingleConstraintData.SingleIndices = InIndices[i];
+			SingleConstraintData.SingleSecondIndices = InSecondIndices[i];
+			SingleConstraintData.SingleWeights = InWeights[i];
+			SingleConstraintData.SingleSecondWeights = InSecondWeights[i];
+			SingleConstraintData.SingleStiffness = InStiffness[i];
+			SingleConstraintData.bIsAnisotropic = InIsAnisotropic[i];
+			SingleConstraintData.bIsZeroRestLength = InIsZeroRestLength[i];
+			ConstraintsData.SetSingleConstraint(SingleConstraintData, i + Offset);
+		}
+	}
+
+	void FGaussSeidelSpringConstraints::Init(const FSolverParticlesRange& InParticles, const FSolverReal Dt)
+	{
+		UpdateTriangleNormalAndNodalWeight(InParticles, /*bUseParticleX =*/ false);
+		if (DebugDrawParams.bVisualizeBindings)
+		{
+			VisualizeAllBindings(InParticles, Dt);
+		}
+	}
+
+	void FGaussSeidelSpringConstraints::Init(const FSolverParticles& InParticles, const FSolverReal Dt)
+	{
+		Init(FSolverParticlesRange(const_cast<FSolverParticles*>(&InParticles), /*InOffset = */ 0, /*InRange = */ InParticles.Size()), Dt);
+	}
+
+	void FGaussSeidelSpringConstraints::CollisionDetectionBVH(const FSolverParticles& Particles, const TArray<TVec3<int32>>& SurfaceElements, const TArray<int32>& ComponentIndex, float DetectRadius, float PositionTargetStiffness, bool UseAnisotropicSpring)
+	{
+		TRACE_CPUPROFILER_EVENT_SCOPE(STAT_ChaosGaussSeidelWeakConstraintsCollisionDetection);
+		Resize(InitialWCSize);
+
+		TArray<Chaos::TVector<int32, 3>> SurfaceElementsArray;
+		for (int32 i = 0; i < SurfaceElements.Num(); i++)
+		{
+			Chaos::TVector<int32, 3> CurrentSurfaceElements(0);
+			for (int32 j = 0; j < 3; j++)
+			{
+				CurrentSurfaceElements[j] = SurfaceElements[i][j];
+			}
+			if (CurrentSurfaceElements[0] != INDEX_NONE
+				&& CurrentSurfaceElements[1] != INDEX_NONE
+				&& CurrentSurfaceElements[2] != INDEX_NONE)
+			{
+				SurfaceElementsArray.Emplace(CurrentSurfaceElements);
+			}
+		}
+		TArray<TArray<int32>> LocalIndex;
+		TArray<TArray<int32>>* LocalIndexPtr = &LocalIndex;
+		TArray<TArray<int>> GlobalIndex = Chaos::Utilities::ComputeIncidentElements(SurfaceElementsArray, LocalIndexPtr);
+		int32 ActualParticleCount = 0;
+		for (int32 l = 0; l < GlobalIndex.Num(); l++)
+		{
+			if (GlobalIndex[l].Num() > 0)
+			{
+				ActualParticleCount += 1;
+			}
+		}
+		TArray<Chaos::TVector<float, 3>> SurfaceElementsPositions;
+		SurfaceElementsPositions.SetNum(ActualParticleCount);
+		TArray<int32> SurfaceElementsMap;
+		SurfaceElementsMap.SetNum(ActualParticleCount);
+		int32 CurrentParticleIndex = 0;
+		for (int32 i = 0; i < GlobalIndex.Num(); i++)
+		{
+			if (GlobalIndex[i].Num() > 0)
+			{
+				SurfaceElementsPositions[CurrentParticleIndex] = Particles.P(SurfaceElements[GlobalIndex[i][0]][LocalIndex[i][0]]);
+				SurfaceElementsMap[CurrentParticleIndex] = SurfaceElements[GlobalIndex[i][0]][LocalIndex[i][0]];
+				CurrentParticleIndex += 1;
+			}
+		}
+
+		TArray<Chaos::FSphere*> VertexSpherePtrs;
+		TArray<Chaos::FSphere> VertexSpheres;
+
+		VertexSpheres.Init(Chaos::FSphere(Chaos::TVec3<Chaos::FReal>(0), DetectRadius), SurfaceElementsPositions.Num());
+		VertexSpherePtrs.SetNum(SurfaceElementsPositions.Num());
+
+		for (int32 i = 0; i < SurfaceElementsPositions.Num(); i++)
+		{
+			Chaos::TVec3<Chaos::FReal> SphereCenter(SurfaceElementsPositions[i]);
+			Chaos::FSphere VertexSphere(SphereCenter, DetectRadius);
+			VertexSpheres[i] = Chaos::FSphere(SphereCenter, DetectRadius);
+			VertexSpherePtrs[i] = &VertexSpheres[i];
+		}
+		Chaos::TBoundingVolumeHierarchy<
+			TArray<Chaos::FSphere*>,
+			TArray<int32>,
+			Chaos::FReal,
+			3> VertexBVH(VertexSpherePtrs);
+
+		for (int32 i = 0; i < SurfaceElements.Num(); i++)
+		{
+			TArray<int32> TriangleIntersections0 = VertexBVH.FindAllIntersections(Particles.P(SurfaceElements[i][0]));
+			TArray<int32> TriangleIntersections1 = VertexBVH.FindAllIntersections(Particles.P(SurfaceElements[i][1]));
+			TArray<int32> TriangleIntersections2 = VertexBVH.FindAllIntersections(Particles.P(SurfaceElements[i][2]));
+			TriangleIntersections0.Sort();
+			TriangleIntersections1.Sort();
+			TriangleIntersections2.Sort();
+
+			TArray<int32> TriangleIntersections({});
+			for (int32 k = 0; k < TriangleIntersections0.Num(); k++)
+			{
+				if (TriangleIntersections1.Contains(TriangleIntersections0[k])
+					&& TriangleIntersections2.Contains(TriangleIntersections0[k]))
+				{
+					TriangleIntersections.Emplace(TriangleIntersections0[k]);
+				}
+			}
+
+			int32 TriangleIndex = ComponentIndex[SurfaceElements[i][0]];
+			int32 MinIndex = INDEX_NONE;
+			float MinDis = DetectRadius;
+			Chaos::TVector<float, 3> ClosestBary(0.f);
+			Chaos::TVector<float, 3> FaceNormal(0.f);
+			for (int32 j = 0; j < TriangleIntersections.Num(); j++)
+			{
+				if (ComponentIndex[SurfaceElementsMap[TriangleIntersections[j]]] >= 0 && TriangleIndex >= 0 && ComponentIndex[SurfaceElementsMap[TriangleIntersections[j]]] != TriangleIndex)
+				{
+					Chaos::TVector<float, 3> Bary, TriPos0(Particles.P(SurfaceElements[i][0])), TriPos1(Particles.P(SurfaceElements[i][1])), TriPos2(Particles.P(SurfaceElements[i][2])), ParticlePos(Particles.P(SurfaceElementsMap[TriangleIntersections[j]]));
+					Chaos::TVector<Chaos::FRealSingle, 3> ClosestPoint = Chaos::FindClosestPointAndBaryOnTriangle(TriPos0, TriPos1, TriPos2, ParticlePos, Bary);
+					Chaos::FRealSingle CurrentDistance = (Particles.P(SurfaceElementsMap[TriangleIntersections[j]]) - ClosestPoint).Size();
+					if (CurrentDistance < MinDis)
+					{
+						Chaos::TVector<FSolverReal, 3> Normal = FVector3f::CrossProduct(TriPos2 - TriPos0, TriPos1 - TriPos0); //The normal needs to point outwards of the geometry
+						if (FVector3f::DotProduct(ParticlePos - TriPos0, Normal) < 0.f)
+						{
+							Normal.SafeNormalize(1e-8f);
+							MinDis = CurrentDistance;
+							MinIndex = SurfaceElementsMap[TriangleIntersections[j]];
+							ClosestBary = Bary;
+							FaceNormal = Normal;
+						}
+					}
+
+				}
+			}
+			if (MinIndex != INDEX_NONE
+				&& MinIndex != SurfaceElements[i][0]
+				&& MinIndex != SurfaceElements[i][1]
+				&& MinIndex != SurfaceElements[i][2])
+			{
+				FGaussSeidelSpringConstraintSingleData SingleConstraintData;
+				SingleConstraintData.SingleIndices = {SurfaceElements[i][0], SurfaceElements[i][1], SurfaceElements[i][2]};
+				SingleConstraintData.SingleSecondIndices = {MinIndex};
+				SingleConstraintData.SingleWeights = {ClosestBary[0], ClosestBary[1], ClosestBary[2]};
+				SingleConstraintData.SingleSecondWeights = {FSolverReal(1.f)};
+				SingleConstraintData.bIsAnisotropic = UseAnisotropicSpring;
+				SingleConstraintData.SingleNormal = FaceNormal;
+				SingleConstraintData.bIsZeroRestLength = true; //Push-out type collision springs should be zero rest length
+
+				float SpringStiffness = 0.f;
+				for (int32 k = 0; k < 3; k++)
+				{
+					SpringStiffness += ClosestBary[k] * PositionTargetStiffness * Particles.M(SurfaceElements[i][k]);
+				}
+				SpringStiffness += PositionTargetStiffness * Particles.M(MinIndex);
+				SingleConstraintData.SingleStiffness = (FSolverReal)SpringStiffness;
+				ConstraintsData.AddSingleConstraint(SingleConstraintData);
+			}
+		}
+	}
+
+	void FGaussSeidelSpringConstraints::ComputeCollisionWCDataSimplified(TArray<TArray<int32>>& ExtraConstraints, TArray<TArray<int32>>& ExtraWCIncidentElements, TArray<TArray<int32>>& ExtraWCIncidentElementsLocal)
+	{
+		ensureMsgf(ConstraintsData.Size() >= InitialWCSize, TEXT("The size of Indices is smaller than InitialWCSize"));
+
+		ExtraConstraints.Init(TArray<int32>(), ConstraintsData.Size() - InitialWCSize);
+		for (int32 i = InitialWCSize; i < static_cast<int32>(ConstraintsData.Size()); i++)
+		{
+			const TArray<int32>& LocalIndices = ConstraintsData.GetIndices(i);
+			const TArray<int32>& LocalSecondIndices = ConstraintsData.GetSecondIndices(i);
+
+			ExtraConstraints[i - InitialWCSize].SetNum(LocalIndices.Num() + LocalSecondIndices.Num());
+			for (int32 j = 0; j < LocalIndices.Num(); j++)
+			{
+				ExtraConstraints[i - InitialWCSize][j] = LocalIndices[j];
+			}
+			for (int32 j = 0; j < LocalSecondIndices.Num(); j++)
+			{
+				ExtraConstraints[i - InitialWCSize][j + LocalIndices.Num()] = LocalSecondIndices[j];
+			}
+		}
+
+		ExtraWCIncidentElements = Chaos::Utilities::ComputeIncidentElements(ExtraConstraints, &ExtraWCIncidentElementsLocal);
+
+		NodalWeights = NoCollisionNodalWeights;
+		for (int32 i = 0; i < ExtraWCIncidentElements.Num(); i++)
+		{
+			if (ExtraWCIncidentElements[i].Num() > 0)
+			{
+				int32 ParticleIdx = ExtraConstraints[ExtraWCIncidentElements[i][0]][ExtraWCIncidentElementsLocal[i][0]];
+				if (NodalWeights[ParticleIdx].Num() == 0)
+				{
+					NodalWeights[ParticleIdx].Init(FSolverReal(0), 6);
+				}
+				for (int32 j = 0; j < ExtraWCIncidentElements[i].Num(); j++)
+				{
+					int32 LocalIndex = ExtraWCIncidentElementsLocal[i][j];
+					int32 ConstraintIndex = ExtraWCIncidentElements[i][j] + InitialWCSize;
+
+					const FGaussSeidelSpringConstraintSingleData& SingleData = ConstraintsData.GetSingleConstraintData(ConstraintIndex);
+
+					FSolverReal Weight = FSolverReal(0);
+					if (LocalIndex >= SingleData.SingleIndices.Num())
+					{
+						Weight = SingleData.SingleWeights[LocalIndex - SingleData.SingleIndices.Num()];
+					}
+					else
+					{
+						Weight = SingleData.SingleWeights[LocalIndex];
+					}
+					const FSolverReal WeightedStiffness = Weight * Weight * SingleData.SingleStiffness;
+					if (SingleData.bIsAnisotropic)
+					{
+						const TVec3<FSolverReal>& Normal = SingleData.SingleNormal;
+						for (int32 Alpha = 0; Alpha < 3; Alpha++)
+						{
+							NodalWeights[ParticleIdx][Alpha] += Normal[Alpha] * Normal[Alpha] * WeightedStiffness;
+						}
+
+						NodalWeights[ParticleIdx][3] += Normal[0] * Normal[1] * WeightedStiffness;
+						NodalWeights[ParticleIdx][4] += Normal[0] * Normal[2] * WeightedStiffness;
+						NodalWeights[ParticleIdx][5] += Normal[1] * Normal[2] * WeightedStiffness;
+					}
+					else
+					{
+						for (int32 Alpha = 0; Alpha < 3; Alpha++)
+						{
+							NodalWeights[ParticleIdx][Alpha] += WeightedStiffness;
+						}
+					}
+				}
+			}
+		}
+	}
+
+	const TArray<TArray<int32>>& FGaussSeidelSpringConstraints::GetStaticConstraintArrays(TArray<TArray<int32>>& IncidentElements, TArray<TArray<int32>>& IncidentElementsLocal) const
+	{
+		IncidentElements = NoCollisionWCIncidentElements;
+		IncidentElementsLocal = NoCollisionWCIncidentElementsLocal;
+		return NoCollisionConstraints;
+	}
+
+	TArray<TArray<int32>> FGaussSeidelSpringConstraints::GetDynamicConstraintArrays(TArray<TArray<int32>>& IncidentElements, TArray<TArray<int32>>& IncidentElementsLocal) const
+	{
+		TArray<TArray<int32>> ExtraConstraints;
+		ExtraConstraints.Init(TArray<int32>(), ConstraintsData.Size());
+		for (int32 i = InitialWCSize; i < ConstraintsData.Size(); i++)
+		{
+			const TArray<int32>& LocalIndices = ConstraintsData.GetIndices(i);
+			const TArray<int32>& LocalSecondIndices = ConstraintsData.GetSecondIndices(i);
+			ExtraConstraints[i - InitialWCSize].SetNum(LocalIndices.Num() + LocalSecondIndices.Num());
+			for (int32 j = 0; j < LocalIndices.Num(); j++)
+			{
+				ExtraConstraints[i - InitialWCSize][j] = LocalIndices[j];
+			}
+			for (int32 j = 0; j < LocalSecondIndices.Num(); j++)
+			{
+				ExtraConstraints[i - InitialWCSize][j + LocalIndices.Num()] = LocalSecondIndices[j];
+			}
+		}
+
+		IncidentElements = Chaos::Utilities::ComputeIncidentElements(ExtraConstraints, &IncidentElementsLocal);
+
+		return ExtraConstraints;
+	}
+
+PRAGMA_DISABLE_DEPRECATION_WARNINGS
+	void FGaussSeidelSpringConstraints::AddZeroRestLengthWCResidualAndHessian(const FSolverParticles& InParticles, const int32 ConstraintIndex, const int32 LocalIndex, const FSolverReal Dt, TVec3<FSolverReal>& ParticleResidual, Chaos::PMatrix<FSolverReal, 3, 3>& ParticleHessian) const
+	{
+		const FGaussSeidelSpringConstraintSingleData& SingleData = ConstraintsData.GetSingleConstraintData(ConstraintIndex);
+
+		const TVec3<FSolverReal> SpringEdge = ComputeSpringEdge(InParticles, SingleData.SingleIndices, SingleData.SingleSecondIndices,
+			SingleData.SingleWeights, SingleData.SingleSecondWeights, /*bUseParticleX =*/false);
+		FSolverReal Weight = FSolverReal(0);
+		if (LocalIndex >= SingleData.SingleIndices.Num())
+		{
+			Weight = -SingleData.SingleSecondWeights[LocalIndex - SingleData.SingleIndices.Num()];
+		}
+		else
+		{
+			Weight = SingleData.SingleWeights[LocalIndex];
+		}
+		if (SingleData.bIsAnisotropic)
+		{
+			FSolverReal Comp = TVec3<FSolverReal>::DotProduct(SpringEdge, SingleData.SingleNormal);
+			TVec3<FSolverReal> Proj = SingleData.SingleNormal * Comp;
+			for (int32 Alpha = 0; Alpha < 3; Alpha++)
+			{
+				ParticleResidual[Alpha] += Dt * Dt * SingleData.SingleStiffness * Proj[Alpha] * Weight;
+			}
+		}
+		else
+		{
+			for (int32 Alpha = 0; Alpha < 3; Alpha++)
+			{
+				ParticleResidual[Alpha] += Dt * Dt * SingleData.SingleStiffness * SpringEdge[Alpha] * Weight;
+			}
+		}
+	}
+PRAGMA_ENABLE_DEPRECATION_WARNINGS
+
+	void FGaussSeidelSpringConstraints::AddWCResidual(const FSolverParticlesRange& InParticles, const int32 ConstraintIndex, const int32 LocalIndex, const FSolverReal Dt, TVec3<FSolverReal>& ParticleResidual, Chaos::PMatrix<FSolverReal, 3, 3>& ParticleHessian) const
+	{
+		const FGaussSeidelSpringConstraintSingleData& SingleData = ConstraintsData.GetSingleConstraintData(ConstraintIndex);
+
+		const TVec3<FSolverReal> SpringEdge = ComputeSpringEdge(InParticles, SingleData.SingleIndices, SingleData.SingleSecondIndices,
+			SingleData.SingleWeights, SingleData.SingleSecondWeights, /*bUseParticleX =*/false);
+		FSolverReal Weight = FSolverReal(0);
+		if (LocalIndex >= SingleData.SingleIndices.Num())
+		{
+			Weight = -SingleData.SingleSecondWeights[LocalIndex - SingleData.SingleIndices.Num()];
+		}
+		else
+		{
+			Weight = SingleData.SingleWeights[LocalIndex];
+		}
+		TVec3<FSolverReal> Projection;
+		if (SingleData.bIsAnisotropic)
+		{
+			FSolverReal LengthDiff = TVec3<FSolverReal>::DotProduct(SpringEdge, SingleData.SingleNormal) - SingleData.RestLength;
+			Projection = SingleData.SingleNormal * LengthDiff;
+		}
+		else
+		{
+			Projection = SpringEdge;
+			if (!SingleData.bIsZeroRestLength) // if not zero rest-length, apply repulsion force
+			{
+				Projection -= SingleData.RestLength * SpringEdge.GetSafeNormal();
+			}
+		}
+		for (int32 Alpha = 0; Alpha < 3; ++Alpha)
+		{
+			ParticleResidual[Alpha] += Dt * Dt * SingleData.SingleStiffness * Projection[Alpha] * Weight;
+		}
+	}
+
+	void FGaussSeidelSpringConstraints::AddWCResidual(const FSolverParticles& InParticles, const int32 ConstraintIndex, const int32 LocalIndex, const FSolverReal Dt, TVec3<FSolverReal>& ParticleResidual, Chaos::PMatrix<FSolverReal, 3, 3>& ParticleHessian) const
+	{
+		AddWCResidual(FSolverParticlesRange(const_cast<FSolverParticles*>(&InParticles), /*InOffset = */ 0, /*InRange = */ InParticles.Size()), ConstraintIndex, LocalIndex, Dt, ParticleResidual, ParticleHessian);
+	}
+
+	void FGaussSeidelSpringConstraints::SetConstraintsRestLength(const FSolverParticlesRange& InParticles, const FSolverParticlesRange& InSecondParticles)
+	{
+		for (int32 ConstraintIdx = 0; ConstraintIdx < ConstraintsData.Size(); ++ConstraintIdx)
+		{
+			if (ConstraintsData.GetIsZeroRestLength(ConstraintIdx))
+			{
+				ConstraintsData.SetRestLength(ConstraintIdx, (FSolverReal)0.);
+			}
+			else
+			{
+				const TArray<int32>& Indices = ConstraintsData.GetIndices(ConstraintIdx);
+				const TArray<int32>& SecondIndices = ConstraintsData.GetSecondIndices(ConstraintIdx);
+				const TArray<FSolverReal>& Weights = ConstraintsData.GetWeights(ConstraintIdx);
+				const TArray<FSolverReal>& SecondWeights = ConstraintsData.GetSecondWeights(ConstraintIdx);
+				const TVec3<FSolverReal> RestSpringEdge = ComputeSpringEdge(InParticles, InSecondParticles, Indices, SecondIndices, Weights, SecondWeights, /*bUseParticleX =*/true);
+				if (ConstraintsData.GetIsAnisotropic(ConstraintIdx))
+				{
+					//if the spring is anisotropic, rest length could be negative depending on the normal direction
+					ConstraintsData.SetRestLength(ConstraintIdx, TVec3<FSolverReal>::DotProduct(ConstraintsData.GetNormal(ConstraintIdx), RestSpringEdge));
+				}
+				else
+				{
+					ConstraintsData.SetRestLength(ConstraintIdx, RestSpringEdge.Size());
+				}
+			}
+		}
+	}
+
+	void FGaussSeidelSpringConstraints::UpdateTriangleNormalAndNodalWeight(const FSolverParticlesRange& InParticles, const FSolverParticlesRange& InSecondParticles, bool bUseParticleX)
+	{
+		UpdateTriangleNormal(InParticles, InSecondParticles, bUseParticleX);
+		UpdateNodalWeight(InParticles, InSecondParticles, bUseParticleX);
+	}
+
+	TVec3<FSolverReal> FGaussSeidelSpringConstraints::ComputeSpringEdge(const FSolverParticlesRange& InParticles, const FSolverParticlesRange& InSecondParticles,
+		const TArray<int32>& LocalIndices, const TArray<int32>& LocalSecondIndices,
+		const TArray<FSolverReal>& Weight, const TArray<FSolverReal>& SecondWeight, bool bUseParticleX) const
+	{
+		auto GetPosition = [bUseParticleX](const FSolverParticlesRange& InParticles, int32 Idx) -> const TVec3<FSolverReal>&
+			{
+				return bUseParticleX ? InParticles.X(Idx) : InParticles.P(Idx);
+			};
+		TVec3<FSolverReal> SpringEdge((FSolverReal)0.);
+		if (ensure(LocalIndices.Num() == Weight.Num() && LocalSecondIndices.Num() == SecondWeight.Num()))
+		{
+			for (int32 Idx = 0; Idx < LocalIndices.Num(); ++Idx)
+			{
+				SpringEdge += Weight[Idx] * GetPosition(InParticles, LocalIndices[Idx]);
+			}
+			for (int32 SecondIdx = 0; SecondIdx < LocalSecondIndices.Num(); ++SecondIdx)
+			{
+				SpringEdge -= SecondWeight[SecondIdx] * GetPosition(InSecondParticles, LocalSecondIndices[SecondIdx]);
+			}
+		}
+		return SpringEdge;
+	}
+
+	TVec3<FSolverReal> FGaussSeidelSpringConstraints::ComputeSpringEdge(const FSolverParticles& InParticles, const FSolverParticles& InSecondParticles,
+		const TArray<int32>& LocalIndices, const TArray<int32>& LocalSecondIndices,
+		const TArray<FSolverReal>& Weight, const TArray<FSolverReal>& SecondWeight, bool bUseParticleX) const
+	{
+		return ComputeSpringEdge(FSolverParticlesRange(const_cast<FSolverParticles*>(&InParticles), /*InOffset = */ 0, /*InRange = */ InParticles.Size()),
+			FSolverParticlesRange(const_cast<FSolverParticles*>(&InSecondParticles), /*InOffset = */ 0, /*InRange = */ InSecondParticles.Size()),
+			LocalIndices, LocalSecondIndices, Weight, SecondWeight, bUseParticleX);
+	}
+
+	TVec3<FSolverReal> FGaussSeidelSpringConstraints::ComputeSpringEdge(const FSolverParticlesRange& InParticles,
+		const TArray<int32>& LocalIndices, const TArray<int32>& LocalSecondIndices,
+		const TArray<FSolverReal>& Weight, const TArray<FSolverReal>& SecondWeight, bool bUseParticleX) const
+	{
+		return ComputeSpringEdge(InParticles, InParticles, LocalIndices, LocalSecondIndices, Weight, SecondWeight, bUseParticleX);
+	}
+
+	TVec3<FSolverReal> FGaussSeidelSpringConstraints::ComputeSpringEdge(const FSolverParticles& InParticles, const TArray<int32>& LocalIndices,
+		const TArray<int32>& LocalSecondIndices, const TArray<FSolverReal>& Weight, const TArray<FSolverReal>& SecondWeight, bool bUseParticleX) const
+	{
+		return ComputeSpringEdge(FSolverParticlesRange(const_cast<FSolverParticles*>(&InParticles), /*InOffset = */ 0, /*InRange = */ InParticles.Size()),
+			LocalIndices, LocalSecondIndices, Weight, SecondWeight, bUseParticleX);
+	}
+
+	void FGaussSeidelSpringConstraints::SetNoCollisionData()
+	{
+		InitialWCSize = ConstraintsData.Size();
+		NoCollisionNodalWeights = NodalWeights;
+		NoCollisionWCIncidentElements = WCIncidentElements;
+		NoCollisionWCIncidentElementsLocal = WCIncidentElementsLocal;
+	}
+
+	void FGaussSeidelSpringConstraints::ComputeIncidentElements()
+	{
+		TArray<TArray<int32>> ExtraConstraints;
+		ExtraConstraints.Init(TArray<int32>(), ConstraintsData.Size());
+		for (int32 ConstraintIdx = 0; ConstraintIdx < ExtraConstraints.Num(); ++ConstraintIdx)
+		{
+			ExtraConstraints[ConstraintIdx].SetNum(ConstraintsData.GetIndices(ConstraintIdx).Num() + ConstraintsData.GetSecondIndices(ConstraintIdx).Num());
+			for (int32 LocalIdx = 0; LocalIdx < ConstraintsData.GetIndices(ConstraintIdx).Num(); ++LocalIdx)
+			{
+				ExtraConstraints[ConstraintIdx][LocalIdx] = ConstraintsData.GetIndices(ConstraintIdx)[LocalIdx];
+			}
+			for (int32 LocalSecondIdx = 0; LocalSecondIdx < ConstraintsData.GetSecondIndices(ConstraintIdx).Num(); ++LocalSecondIdx)
+			{
+				ExtraConstraints[ConstraintIdx][LocalSecondIdx + ConstraintsData.GetIndices(ConstraintIdx).Num()] = ConstraintsData.GetSecondIndices(ConstraintIdx)[LocalSecondIdx];
+			}
+		}
+		WCIncidentElements = Chaos::Utilities::ComputeIncidentElements(ExtraConstraints, &WCIncidentElementsLocal);
+		NoCollisionConstraints = ExtraConstraints;
+	}
+
+	void FGaussSeidelSpringConstraints::UpdateTriangleNormal(const FSolverParticlesRange& InParticles, const FSolverParticlesRange& InSecondParticles, bool bUseParticleX)
+	{
+		auto GetPosition = [bUseParticleX](const FSolverParticlesRange& InParticles, int32 Idx) -> const TVec3<FSolverReal>&
+			{
+				return bUseParticleX ? InParticles.X(Idx) : InParticles.P(Idx);
+			};
+		for (int32 i = 0; i < ConstraintsData.Size(); i++)
+		{
+			if (ConstraintsData.GetIsAnisotropic(i))
+			{
+				const TArray<int32>& IndicesTemp = ConstraintsData.GetIndices(i);
+				const TArray<int32>& SecondIndicesTemp = ConstraintsData.GetSecondIndices(i);
+				Chaos::TVector<FSolverReal, 3> Normal((FSolverReal)0.);
+				if (IndicesTemp.Num() == 3)
+				{
+					Chaos::TVector<float, 3> TriPos0(GetPosition(InParticles, IndicesTemp[0])),
+						TriPos1(GetPosition(InParticles, IndicesTemp[1])),
+						TriPos2(GetPosition(InParticles, IndicesTemp[2]));
+					Normal = FVector3f::CrossProduct(TriPos2 - TriPos0, TriPos1 - TriPos0).GetSafeNormal(); //triangle normal convention (see FTriangleMesh::GetFaceNormals())
+				}
+				else if (SecondIndicesTemp.Num() == 3)
+				{
+					Chaos::TVector<float, 3> TriPos0(GetPosition(InSecondParticles, SecondIndicesTemp[0])),
+						TriPos1(GetPosition(InSecondParticles, SecondIndicesTemp[1])),
+						TriPos2(GetPosition(InSecondParticles, SecondIndicesTemp[2]));
+					Normal = FVector3f::CrossProduct(TriPos2 - TriPos0, TriPos1 - TriPos0).GetSafeNormal(); //triangle normal convention (see FTriangleMesh::GetFaceNormals())
+				}
+				ConstraintsData.SetNormal(i, Normal);
+			}
+		}
+	}
+
+	void FGaussSeidelSpringConstraints::UpdateNodalWeight(const FSolverParticlesRange& InParticles, const FSolverParticlesRange& InSecondParticles, bool bUseParticleX)
+	{
+		NodalWeights.Init({}, InParticles.Size());
+		for (int32 ParticleIdx = 0; ParticleIdx < WCIncidentElements.Num(); ++ParticleIdx)
+		{
+			if (WCIncidentElements[ParticleIdx].Num() > 0)
+			{
+				NodalWeights[ParticleIdx].Init(FSolverReal(0), 6);
+				for (int32 j = 0; j < WCIncidentElements[ParticleIdx].Num(); j++)
+				{
+					int32 ConstraintIndex = WCIncidentElements[ParticleIdx][j];
+					int32 LocalIndex = WCIncidentElementsLocal[ParticleIdx][j];
+
+					FSolverReal Weight = FSolverReal(0);
+					if (LocalIndex >= ConstraintsData.GetIndices(ConstraintIndex).Num())
+					{
+						Weight = ConstraintsData.GetSecondWeights(ConstraintIndex)[LocalIndex - ConstraintsData.GetIndices(ConstraintIndex).Num()];
+					}
+					else
+					{
+						Weight = ConstraintsData.GetWeights(ConstraintIndex)[LocalIndex];
+					}
+
+					const FSolverReal WeightedStiffness = Weight * Weight * ConstraintsData.GetStiffness(ConstraintIndex);
+					if (ConstraintsData.GetIsAnisotropic(ConstraintIndex))
+					{
+						const TVec3<FSolverReal>& Normal = ConstraintsData.GetNormal(ConstraintIndex);
+						for (int32 Alpha = 0; Alpha < 3; Alpha++)
+						{
+							NodalWeights[ParticleIdx][Alpha] += Normal[Alpha] * Normal[Alpha] * WeightedStiffness;
+						}
+
+						NodalWeights[ParticleIdx][3] += Normal[0] * Normal[1] * WeightedStiffness;
+						NodalWeights[ParticleIdx][4] += Normal[0] * Normal[2] * WeightedStiffness;
+						NodalWeights[ParticleIdx][5] += Normal[1] * Normal[2] * WeightedStiffness;
+					}
+					else
+					{
+						for (int32 Alpha = 0; Alpha < 3; Alpha++)
+						{
+							NodalWeights[ParticleIdx][Alpha] += WeightedStiffness;
+						}
+					}
+				}
+			}
+		}
+	}
+
+	void FGaussSeidelSpringConstraints::UpdateTriangleNormalAndNodalWeight(const FSolverParticlesRange& InParticles, bool bUseParticleX)
+	{
+		UpdateTriangleNormalAndNodalWeight(InParticles, InParticles, bUseParticleX);
+	}
+
+	void FGaussSeidelSpringConstraints::Resize(int32 Size)
+	{
+		ConstraintsData.Resize(Size);
+	}
+
+	void FGaussSeidelSpringConstraints::UpdatePointTriangleCollisionWCData(const FSolverParticles& Particles)
+	{
+		FGaussSeidelSpringConstraintData OriginalConstraintsData = ConstraintsData;
+
+		ConstraintsData.Resize(InitialWCSize);
+
+		for (int32 i = InitialWCSize; i < OriginalConstraintsData.Size(); i++)
+		{
+			const TArray<int32>& IndicesTemp = OriginalConstraintsData.GetIndices(i);
+			const TArray<int32>& SecondIndicesTemp = OriginalConstraintsData.GetSecondIndices(i);
+			ensureMsgf(OriginalConstraintsData.GetIndices(i).Num() == 3, TEXT("Collision format is not point-triangle"));
+			ensureMsgf(OriginalConstraintsData.GetSecondIndices(i).Num() == 1, TEXT("Collision format is not point-triangle"));
+			Chaos::TVector<float, 3> TriPos0(Particles.P(IndicesTemp[0])), TriPos1(Particles.P(IndicesTemp[1])), TriPos2(Particles.P(IndicesTemp[2])), ParticlePos(Particles.P(SecondIndicesTemp[0]));
+			Chaos::TVector<FSolverReal, 3> Normal = FVector3f::CrossProduct(TriPos2 - TriPos0, TriPos1 - TriPos0); //triangle normal convention (see FTriangleMesh::GetFaceNormals())
+			if (FVector3f::DotProduct(ParticlePos - TriPos0, Normal) < 0.f) //not resolved, keep the spring
+			{
+				ConstraintsData.AddSingleConstraint(OriginalConstraintsData.GetSingleConstraintData(i));
+			}
+		}
+	}
+
+	void FGaussSeidelSpringConstraints::VisualizeAllBindings(const FSolverParticlesRange& InParticles, const FSolverReal Dt) const
+	{
+#if WITH_EDITOR
+		auto DoubleVert = [](Chaos::TVec3<FSolverReal> V) { return FVector3d(V.X, V.Y, V.Z); };
+		for (int32 i = 0; i < ConstraintsData.Size(); i++)
+		{
+			const FGaussSeidelSpringConstraintSingleData& SingleConstraintData = ConstraintsData.GetSingleConstraintData(i);
+			Chaos::TVec3<FSolverReal> SourcePos((FSolverReal)0.), TargetPos((FSolverReal)0.);
+			for (int32 j = 0; j < SingleConstraintData.SingleIndices.Num(); j++)
+			{
+				SourcePos += SingleConstraintData.SingleWeights[j] * InParticles.P(SingleConstraintData.SingleIndices[j]);
+			}
+			for (int32 j = 0; j < SingleConstraintData.SingleSecondIndices.Num(); j++)
+			{
+				TargetPos += SingleConstraintData.SingleSecondWeights[j] * InParticles.P(SingleConstraintData.SingleSecondIndices[j]);
+			}
+
+			float ParticleThickness = DebugDrawParams.DebugParticleWidth;
+			float LineThickness = DebugDrawParams.DebugLineWidth;
+
+			if (SingleConstraintData.SingleIndices.Num() == 1)
+			{
+				Chaos::FDebugDrawQueue::GetInstance().DrawDebugPoint(DoubleVert(SourcePos), FColor::Red, false, Dt, 0, ParticleThickness);
+				for (int32 j = 0; j < SingleConstraintData.SingleSecondIndices.Num(); j++)
+				{
+					Chaos::FDebugDrawQueue::GetInstance().DrawDebugPoint(DoubleVert(InParticles.P(SingleConstraintData.SingleSecondIndices[j])), FColor::Green, false, Dt, 0, ParticleThickness);
+					Chaos::FDebugDrawQueue::GetInstance().DrawDebugLine(DoubleVert(InParticles.P(SingleConstraintData.SingleSecondIndices[j])), DoubleVert(InParticles.P(SingleConstraintData.SingleSecondIndices[(j + 1) % SingleConstraintData.SingleSecondIndices.Num()])), FColor::Green, false, Dt, 0, LineThickness);
+				}
+
+			}
+
+			if (SingleConstraintData.SingleSecondIndices.Num() == 1)
+			{
+				Chaos::FDebugDrawQueue::GetInstance().DrawDebugPoint(DoubleVert(TargetPos), FColor::Red, false, Dt, 0, ParticleThickness);
+				for (int32 j = 0; j < SingleConstraintData.SingleIndices.Num(); j++)
+				{
+					Chaos::FDebugDrawQueue::GetInstance().DrawDebugPoint(DoubleVert(InParticles.P(SingleConstraintData.SingleIndices[j])), FColor::Green, false, Dt, 0, ParticleThickness);
+					Chaos::FDebugDrawQueue::GetInstance().DrawDebugLine(DoubleVert(InParticles.P(SingleConstraintData.SingleIndices[j])), DoubleVert(InParticles.P(SingleConstraintData.SingleIndices[(j + 1) % SingleConstraintData.SingleIndices.Num()])), FColor::Green, false, Dt, 0, LineThickness);
+				}
+			}
+
+			Chaos::FDebugDrawQueue::GetInstance().DrawDebugLine(DoubleVert(SourcePos), DoubleVert(TargetPos), FColor::Yellow, false, Dt, 0, LineThickness);
+		}
+#endif
+	}
+
+	// -------------------------------------------------------------------
+	// CollisionDetectionSpatialHash / CollisionDetectionSpatialHashInComponent
+	// -------------------------------------------------------------------
+
+	template<typename SpatialAccelerator>
+	void FGaussSeidelSpringConstraints::CollisionDetectionSpatialHash(const FSolverParticles& Particles, const TArray<int32>& SurfaceVertices, const FTriangleMesh& TriangleMesh, const TArray<int32>& ComponentIndex, const SpatialAccelerator& Spatial, float DetectRadius, float PositionTargetStiffness, bool UseAnisotropicSpring)
+	{
+		TRACE_CPUPROFILER_EVENT_SCOPE(STAT_ChaosGaussSeidelWeakConstraintsCollisionDetectionSpatialHash);
+		Resize(InitialWCSize + Particles.Size());
+		std::atomic<int32> ConstraintIndex(InitialWCSize);
+		const TArray<TVec3<int32>>& Elements = TriangleMesh.GetSurfaceElements();
+		const float HalfRadius = DetectRadius / 2;
+		PhysicsParallelFor(SurfaceVertices.Num(),
+			[this, &Spatial, &Particles, &SurfaceVertices, &ConstraintIndex, &TriangleMesh, &Elements, &HalfRadius, &ComponentIndex, &PositionTargetStiffness, &UseAnisotropicSpring](int32 i)
+			{
+				const int32 Index = SurfaceVertices[i];
+				TArray< TTriangleCollisionPoint<FSolverReal>> Result;
+				//PointProximityQuery
+				if (TriangleMesh.PointClosestTriangleQuery(Spatial, static_cast<const TArrayView<const FSolverVec3>&>(Particles.XArray()), Index, Particles.GetX(Index), HalfRadius, HalfRadius,
+					[this, &ComponentIndex, &Elements](const int32 PointIndex, const int32 TriangleIndex)->bool
+					{
+						//Skip particles that are bound in initial springs
+						return ComponentIndex[PointIndex] != ComponentIndex[Elements[TriangleIndex][0]] && (!NoCollisionWCIncidentElements.IsValidIndex(PointIndex) || NoCollisionWCIncidentElements[PointIndex].Num() == 0);
+					},
+					Result))
+				{
+					for (const TTriangleCollisionPoint<FSolverReal>& CollisionPoint : Result)
+					{
+						if (CollisionPoint.Phi < 0)
+						{
+							const TVector<int32, 3>& Elem = Elements[CollisionPoint.Indices[1]];
+							const int32 IndexToWrite = ConstraintIndex.fetch_add(1);
+
+							FGaussSeidelSpringConstraintSingleData SingleConstraintData;
+							SingleConstraintData.SingleIndices = { Elem[0], Elem[1] ,Elem[2] };
+							SingleConstraintData.SingleSecondIndices =  { Index };
+							SingleConstraintData.SingleWeights = { CollisionPoint.Bary[1], CollisionPoint.Bary[2], CollisionPoint.Bary[3] };
+							SingleConstraintData.SingleSecondWeights = {FSolverReal(1.f)};
+							SingleConstraintData.bIsAnisotropic = UseAnisotropicSpring;
+							SingleConstraintData.SingleNormal = CollisionPoint.Normal;
+							SingleConstraintData.bIsZeroRestLength = true; //Push-out type collision springs should be zero rest length
+
+							float SpringStiffness = 0.f;
+							for (int32 k = 0; k < 3; k++)
+							{
+								SpringStiffness += SingleConstraintData.SingleWeights[k] * PositionTargetStiffness * Particles.M(Elem[k]);
+							}
+							SpringStiffness += PositionTargetStiffness * Particles.M(Index);
+							SingleConstraintData.SingleStiffness = (FSolverReal)SpringStiffness;
+							ConstraintsData.SetSingleConstraint(SingleConstraintData, IndexToWrite);
+						}
+					}
+				}
+			}
+		);
+
+		// Shrink the arrays to the actual number of found constraints.
+		const int32 ConstraintNum = ConstraintIndex.load();
+		Resize(ConstraintNum);
+	}
+
+	template<typename SpatialAccelerator>
+	void FGaussSeidelSpringConstraints::CollisionDetectionSpatialHashInComponent(const FSolverParticles& Particles, const TArray<int32>& SurfaceVertices, const FTriangleMesh& TriangleMesh, const TMap<int32, TSet<int32>>& ExcludeMap, const SpatialAccelerator& Spatial, float DetectRadius, float PositionTargetStiffness, bool UseAnisotropicSpring)
+	{
+		TRACE_CPUPROFILER_EVENT_SCOPE(STAT_ChaosGaussSeidelWeakConstraintsCollisionDetectionSpatialHashInComponent);
+		Resize(InitialWCSize + Particles.Size());
+		std::atomic<int32> ConstraintIndex(InitialWCSize);
+		const TArray<TVec3<int32>>& Elements = TriangleMesh.GetSurfaceElements();
+		const float HalfRadius = DetectRadius/2;
+		PhysicsParallelFor(SurfaceVertices.Num(),
+			[this, &Spatial, &Particles, &SurfaceVertices, &ConstraintIndex, &TriangleMesh, &ExcludeMap, &Elements, &HalfRadius, &PositionTargetStiffness, &UseAnisotropicSpring](int32 i)
+			{
+				const int32 Index = SurfaceVertices[i];
+				TArray< TTriangleCollisionPoint<FSolverReal>> Result;
+				//PointProximityQuery
+				if (TriangleMesh.PointClosestTriangleQuery(Spatial, static_cast<const TArrayView<const FSolverVec3>&>(Particles.XArray()), Index, Particles.GetX(Index), HalfRadius, HalfRadius,
+					[this, &Elements, &ExcludeMap](const int32 PointIndex, const int32 TriangleIndex)->bool
+					{
+						return  !(ExcludeMap.Find(PointIndex) && ExcludeMap[PointIndex].Contains(TriangleIndex));
+					},
+					Result))
+				{
+					for (const TTriangleCollisionPoint<FSolverReal>& CollisionPoint : Result)
+					{
+						if (CollisionPoint.Phi < 0)
+						{
+							const TVector<int32, 3>& Elem = Elements[CollisionPoint.Indices[1]];
+							const int32 IndexToWrite = ConstraintIndex.fetch_add(1);
+
+							FGaussSeidelSpringConstraintSingleData SingleConstraintData;
+							SingleConstraintData.SingleIndices = { Elem[0], Elem[1] ,Elem[2] };
+							SingleConstraintData.SingleSecondIndices =  { Index };
+							SingleConstraintData.SingleWeights = { CollisionPoint.Bary[1], CollisionPoint.Bary[2], CollisionPoint.Bary[3] };
+							SingleConstraintData.SingleSecondWeights = {FSolverReal(1.f)};
+							SingleConstraintData.bIsAnisotropic = UseAnisotropicSpring;
+							SingleConstraintData.SingleNormal = CollisionPoint.Normal;
+							SingleConstraintData.bIsZeroRestLength = true; //Push-out type collision springs should be zero rest length
+
+							float SpringStiffness = 0.f;
+							for (int32 k = 0; k < 3; k++)
+							{
+								SpringStiffness += SingleConstraintData.SingleWeights[k] * PositionTargetStiffness * Particles.M(Elem[k]);
+							}
+							SpringStiffness += PositionTargetStiffness * Particles.M(Index);
+							SingleConstraintData.SingleStiffness = (FSolverReal)SpringStiffness;
+							ConstraintsData.SetSingleConstraint(SingleConstraintData, IndexToWrite);
+						}
+					}
+				}
+			}
+		);
+
+		// Shrink the arrays to the actual number of found constraints.
+		const int32 ConstraintNum = ConstraintIndex.load();
+		Resize(ConstraintNum);
+	}
+
+	// Explicit template instantiation for the only used type
+	template void FGaussSeidelSpringConstraints::CollisionDetectionSpatialHash<FTriangleMesh::TSpatialHashType<FSolverReal>>(const FSolverParticles&, const TArray<int32>&, const FTriangleMesh&, const TArray<int32>&, const FTriangleMesh::TSpatialHashType<FSolverReal>&, float, float, bool);
+	template void FGaussSeidelSpringConstraints::CollisionDetectionSpatialHashInComponent<FTriangleMesh::TSpatialHashType<FSolverReal>>(const FSolverParticles&, const TArray<int32>&, const FTriangleMesh&, const TMap<int32, TSet<int32>>&, const FTriangleMesh::TSpatialHashType<FSolverReal>&, float, float, bool);
+
+} // namespace Chaos::Softs
diff --git a/Engine/Source/Runtime/Experimental/Chaos/Private/Chaos/Deformable/XPBDWeakConstraints.cpp b/Engine/Source/Runtime/Experimental/Chaos/Private/Chaos/Deformable/XPBDWeakConstraints.cpp
new file mode 100644
--- /dev/null
+++ b/Engine/Source/Runtime/Experimental/Chaos/Private/Chaos/Deformable/XPBDWeakConstraints.cpp
@@ -0,0 +1,525 @@
+// Copyright Epic Games, Inc. All Rights Reserved.
+
+#include "Chaos/XPBDWeakConstraints.h"
+#include "Chaos/DebugDrawQueue.h"
+
+namespace Chaos::Softs
+{
+	using Chaos::TVec3;
+
+	// -------------------------------------------------------------------
+	// Constructors
+	// -------------------------------------------------------------------
+
+	FXPBDWeightedSpringConstraints::FXPBDWeightedSpringConstraints(
+		const FSolverParticles& InParticles,
+		const TArray<TArray<int32>>& InIndices,
+		const TArray<TArray<FSolverReal>>& InWeights,
+		const TArray<FSolverReal>& InStiffness
+	)
+		: Indices(InIndices), Weights(InWeights), Stiffness(InStiffness)
+	{
+		InitColor(InParticles);
+		LambdaArray.Init((FSolverReal)0., Indices.Num() * 3);
+	}
+
+	FXPBDWeightedSpringConstraints::FXPBDWeightedSpringConstraints(
+		const FSolverParticles& InParticles,
+		const TArray<TArray<int32>>& InIndices,
+		const TArray<TArray<FSolverReal>>& InWeights,
+		const TArray<FSolverReal>& InStiffness,
+		const TArray<TArray<int32>>& InSecondIndices,
+		const TArray<TArray<FSolverReal>>& InSecondWeights,
+		const FDeformableXPBDWeightedSpringConstraintParams& InParams
+	)
+		: Indices(InIndices), Weights(InWeights), SecondIndices(InSecondIndices), SecondWeights(InSecondWeights), Stiffness(InStiffness), DebugDrawParams(InParams)
+	{
+		ensureMsgf(Indices.Num() == SecondIndices.Num(), TEXT("Input Double Bindings have wrong size"));
+
+		for (int32 i = 0; i < Indices.Num(); i++)
+		{
+			TSet<int32> IndicesSet = TSet<int32>(Indices[i]);
+			for (int32 j = 0; j < SecondIndices[i].Num(); j++)
+			{
+				ensureMsgf(!IndicesSet.Contains(SecondIndices[i][j]), TEXT("Indices and Second Indices overlaps. Currently not supported"));
+			}
+		}
+		InitColor(InParticles);
+		LambdaArray.Init((FSolverReal)0., Indices.Num() * 3);
+	}
+
+	FXPBDWeightedSpringConstraints::FXPBDWeightedSpringConstraints(
+		const FSolverParticles& InParticles,
+		const FSolverParticles& InSecondParticles,
+		const TArray<TArray<int32>>& InIndices,
+		const TArray<TArray<FSolverReal>>& InWeights,
+		const TArray<FSolverReal>& InStiffness,
+		const TArray<TArray<int32>>& InSecondIndices,
+		const TArray<TArray<FSolverReal>>& InSecondWeights,
+		const FDeformableXPBDWeightedSpringConstraintParams& InParams
+	)
+		: Indices(InIndices), Weights(InWeights), SecondIndices(InSecondIndices), SecondWeights(InSecondWeights), Stiffness(InStiffness), DebugDrawParams(InParams)
+	{
+		ensureMsgf(Indices.Num() == SecondIndices.Num(), TEXT("Input Double Bindings have wrong size"));
+
+		InitColor(InParticles, InSecondParticles);
+		LambdaArray.Init((FSolverReal)0., Indices.Num() * 3);
+	}
+
+	// -------------------------------------------------------------------
+	// Apply methods
+	// -------------------------------------------------------------------
+
+	void FXPBDWeightedSpringConstraints::ApplyInParallel(FSolverParticlesRange& Particles, const FSolverReal Dt) const
+	{
+		TRACE_CPUPROFILER_EVENT_SCOPE(STAT_ChaosXPBDWeakConstraintApply);
+
+		if ((ConstraintsPerColor.Num() > 0))
+		{
+			if (SecondIndices.Num() == 0)
+			{
+				for (int32 Color = 0; Color < ConstraintsPerColor.Num(); Color++)
+				{
+					// TODO: replace with UE::Tasks::FTask with dependencies and minimum batch size control
+					PhysicsParallelFor(ConstraintsPerColor[Color].Num(), [&](const int32 Index)
+						{
+							const int32 ConstraintIndex = ConstraintsPerColor[Color][Index];
+							ApplySingleConstraintWithoutSelfTarget(Particles, Dt, ConstraintIndex);
+						});
+				}
+			}
+			else
+			{
+				for (int32 Color = 0; Color < ConstraintsPerColor.Num(); Color++)
+				{
+					PhysicsParallelFor(ConstraintsPerColor[Color].Num(), [&](const int32 Index)
+						{
+							const int32 ConstraintIndex = ConstraintsPerColor[Color][Index];
+							ApplySingleConstraintWithSelfTarget(Particles, Dt, ConstraintIndex);
+						});
+				}
+			}
+		}
+	}
+
+	void FXPBDWeightedSpringConstraints::ApplyInParallel(FSolverParticles& Particles, const FSolverReal Dt) const
+	{
+		FSolverParticlesRange Range(&Particles, /*InOffset = */ 0, /*InRange = */ Particles.Size());
+		ApplyInParallel(Range, Dt);
+	}
+
+	void FXPBDWeightedSpringConstraints::Apply(FSolverParticlesRange& Particles, const FSolverReal Dt) const
+	{
+		ApplyInParallel(Particles, Dt);
+	}
+
+	void FXPBDWeightedSpringConstraints::ApplyInSerial(FSolverParticlesRange& Particles, FSolverParticlesRange& SecondParticles, const FSolverReal Dt) const
+	{
+		TRACE_CPUPROFILER_EVENT_SCOPE(STAT_ChaosXPBDWeakConstraintApply);
+
+		for (int32 ConstraintIdx = 0; ConstraintIdx < Indices.Num(); ++ConstraintIdx)
+		{
+			ApplySingleConstraintWithSelfTarget(Particles, SecondParticles, Dt, ConstraintIdx);
+		}
+	}
+
+	void FXPBDWeightedSpringConstraints::ApplyInParallel(FSolverParticlesRange& Particles, FSolverParticlesRange& SecondParticles, const FSolverReal Dt) const
+	{
+		TRACE_CPUPROFILER_EVENT_SCOPE(STAT_ChaosXPBDWeakConstraintTwoParticlesApply);
+		if ((ConstraintsPerColor.Num() > 0))
+		{
+			if (ensureMsgf(Indices.Num() == SecondIndices.Num(), TEXT("Constraint source and target have nonequal sizes")))
+			{
+				for (int32 Color = 0; Color < ConstraintsPerColor.Num(); Color++)
+				{
+					PhysicsParallelFor(ConstraintsPerColor[Color].Num(), [&](const int32 Index)
+						{
+							const int32 ConstraintIndex = ConstraintsPerColor[Color][Index];
+							ApplySingleConstraintBetweenParticles(Particles, SecondParticles, Dt, ConstraintIndex);
+						});
+				}
+			}
+		}
+	}
+
+	// -------------------------------------------------------------------
+	// Init methods
+	// -------------------------------------------------------------------
+
+	void FXPBDWeightedSpringConstraints::Init(const FSolverParticlesRange& InParticles, const FSolverReal Dt) const
+	{
+		for (FSolverReal& Lambdas : LambdaArray) { Lambdas = (FSolverReal)0.; }
+		if (DebugDrawParams.bVisualizeBindings)
+		{
+			VisualizeAllBindings(InParticles, Dt);
+		}
+	}
+
+	void FXPBDWeightedSpringConstraints::Init(const FSolverParticles& InParticles, const FSolverReal Dt) const
+	{
+		Init(FSolverParticlesRange(const_cast<FSolverParticles*>(&InParticles), /*InOffset = */ 0, /*InRange = */ InParticles.Size()), Dt);
+	}
+
+	void FXPBDWeightedSpringConstraints::Init(const FSolverParticlesRange& InParticles, const FSolverParticlesRange& InSecondParticles, const FSolverReal Dt) const
+	{
+		for (FSolverReal& Lambdas : LambdaArray) { Lambdas = (FSolverReal)0.; }
+	}
+
+	// -------------------------------------------------------------------
+	// ComputeSpringEdge (virtual)
+	// -------------------------------------------------------------------
+
+	bool FXPBDWeightedSpringConstraints::ComputeSpringEdge(const FSolverParticlesRange& InParticles, TVec3<FSolverReal>& OutSpringEdge, int32 ConstraintIndex, bool bUseParticleX, bool bUseConstraintTargetPosition) const
+	{
+		return ComputeSpringEdge(InParticles, InParticles, OutSpringEdge, ConstraintIndex, bUseParticleX, bUseConstraintTargetPosition);
+	}
+
+	bool FXPBDWeightedSpringConstraints::ComputeSpringEdge(const FSolverParticlesRange& InParticles, const FSolverParticlesRange& InSecondParticles,
+		TVec3<FSolverReal>& OutSpringEdge, int32 ConstraintIndex, bool bUseParticleX, bool bUseConstraintTargetPosition) const
+	{
+		OutSpringEdge = TVec3<FSolverReal>((FSolverReal)0.);
+		if (ensure(Indices[ConstraintIndex].Num() == Weights[ConstraintIndex].Num() && 
+			(bUseConstraintTargetPosition || SecondIndices[ConstraintIndex].Num() == SecondWeights[ConstraintIndex].Num())))
+		{
+			auto GetPosition = [bUseParticleX](const FSolverParticlesRange& InParticles, int32 Idx) -> const TVec3<FSolverReal>&
+				{
+					return bUseParticleX ? InParticles.X(Idx) : InParticles.P(Idx);
+				};
+			for (int32 Idx = 0; Idx < Indices[ConstraintIndex].Num(); ++Idx)
+			{
+				OutSpringEdge += Weights[ConstraintIndex][Idx] * GetPosition(InParticles, Indices[ConstraintIndex][Idx]);
+			}
+			if (bUseConstraintTargetPosition)
+			{
+				OutSpringEdge -= Constraints[ConstraintIndex];
+			}
+			else
+			{
+				for (int32 SecondIdx = 0; SecondIdx < SecondIndices[ConstraintIndex].Num(); ++SecondIdx)
+				{
+					OutSpringEdge -= SecondWeights[ConstraintIndex][SecondIdx] * GetPosition(InSecondParticles, SecondIndices[ConstraintIndex][SecondIdx]);
+				}
+			}
+			return true;
+		}
+		else
+		{
+			return false;
+		}
+	}
+
+	// -------------------------------------------------------------------
+	// InitColor (private)
+	// -------------------------------------------------------------------
+
+	void FXPBDWeightedSpringConstraints::InitColor(const FSolverParticles& Particles)
+	{
+		Chaos::ComputeWeakConstraintsColoring(Indices, SecondIndices, Particles, ConstraintsPerColor);
+	}
+
+	void FXPBDWeightedSpringConstraints::InitColor(const FSolverParticles& Particles, const FSolverParticles& SecondParticles)
+	{
+		Chaos::ComputeWeakConstraintsColoring(Indices, SecondIndices, Particles, SecondParticles, ConstraintsPerColor);
+	}
+
+	// -------------------------------------------------------------------
+	// ApplySingleConstraint variants (private)
+	// -------------------------------------------------------------------
+
+	void FXPBDWeightedSpringConstraints::ApplySingleConstraintWithoutSelfTarget(FSolverParticlesRange& Particles, const FSolverReal Dt, const int32 ConstraintIndex, const FSolverReal Tol) const
+	{
+		if (Stiffness[ConstraintIndex] < Tol || Dt == (FSolverReal)0.)
+		{
+			return;
+		}
+
+		const FSolverReal AlphaTilde = Stiffness[ConstraintIndex] > SoftMaxStiffness ? FSolverReal(2) / (Stiffness[ConstraintIndex] * Dt * Dt) : FSolverReal(0);
+
+		TVec3<FSolverReal> SpringEdge;
+		if (!ComputeSpringEdge(Particles, SpringEdge, ConstraintIndex, /*bUseParticleX*/ false, /*bUseConstraintTargetPosition*/ true))
+		{
+			return;
+		}
+
+		FSolverReal Denom = AlphaTilde;
+		for (int32 j = 0; j < Indices[ConstraintIndex].Num(); j++) {
+			Denom += Weights[ConstraintIndex][j] * Weights[ConstraintIndex][j] * Particles.InvM(Indices[ConstraintIndex][j]);
+
+		}
+
+		if (Denom == (FSolverReal)0.)
+		{
+			return;
+		}
+
+		for (int32 Beta = 0; Beta < 3; Beta++)
+		{
+			FSolverReal Cj = SpringEdge[Beta];
+			FSolverReal DLambda = -Cj - AlphaTilde * LambdaArray[ConstraintIndex * 3 + Beta];
+
+			DLambda /= Denom;
+			LambdaArray[ConstraintIndex * 3 + Beta] += DLambda;
+			for (int32 j = 0; j < Indices[ConstraintIndex].Num(); j++) {
+				Particles.P(Indices[ConstraintIndex][j])[Beta] += DLambda * Weights[ConstraintIndex][j] * Particles.InvM(Indices[ConstraintIndex][j]);
+			}
+		}
+	}
+
+	void FXPBDWeightedSpringConstraints::ApplySingleConstraintWithSelfTarget(FSolverParticlesRange& Particles, const FSolverReal Dt, const int32 ConstraintIndex, const FSolverReal Tol) const
+	{
+		if (Stiffness[ConstraintIndex] < Tol || Dt == (FSolverReal)0. || 
+			!ensure(SecondIndices.IsValidIndex(ConstraintIndex) && SecondWeights.IsValidIndex(ConstraintIndex)))
+		{
+			return;
+		}
+
+		const FSolverReal AlphaTilde = Stiffness[ConstraintIndex] > SoftMaxStiffness ? FSolverReal(2) / (Stiffness[ConstraintIndex] * Dt * Dt) : FSolverReal(0);
+
+		TVec3<FSolverReal> SpringEdge;
+		if (!ComputeSpringEdge(Particles, SpringEdge, ConstraintIndex, /*bUseParticleX*/ false, /*bUseConstraintTargetPosition*/ false))
+		{
+			return;
+		}
+
+		FSolverReal Denom = AlphaTilde;
+		for (int32 j = 0; j < Indices[ConstraintIndex].Num(); j++) {
+			Denom += Weights[ConstraintIndex][j] * Weights[ConstraintIndex][j] * Particles.InvM(Indices[ConstraintIndex][j]);
+		}
+
+		for (int32 j = 0; j < SecondIndices[ConstraintIndex].Num(); j++) {
+			Denom += SecondWeights[ConstraintIndex][j] * SecondWeights[ConstraintIndex][j] * Particles.InvM(SecondIndices[ConstraintIndex][j]);
+		}
+
+		if (Denom == (FSolverReal)0.)
+		{
+			return;
+		}
+
+		for (int32 Beta = 0; Beta < 3; Beta++)
+		{
+			FSolverReal Cj = SpringEdge[Beta];
+			FSolverReal DLambda = -Cj - AlphaTilde * LambdaArray[ConstraintIndex * 3 + Beta];
+
+			DLambda /= Denom;
+			LambdaArray[ConstraintIndex * 3 + Beta] += DLambda;
+			for (int32 j = 0; j < Indices[ConstraintIndex].Num(); j++) {
+				Particles.P(Indices[ConstraintIndex][j])[Beta] += DLambda * Weights[ConstraintIndex][j] * Particles.InvM(Indices[ConstraintIndex][j]);
+			}
+
+			for (int32 j = 0; j < SecondIndices[ConstraintIndex].Num(); j++) {
+				Particles.P(SecondIndices[ConstraintIndex][j])[Beta] -= DLambda * SecondWeights[ConstraintIndex][j] * Particles.InvM(SecondIndices[ConstraintIndex][j]);
+			}
+		}
+	}
+
+	//Deprecating in the next CL, replaced with ApplySingleConstraintBetweenParticles
+	void FXPBDWeightedSpringConstraints::ApplySingleConstraintWithSelfTarget(FSolverParticlesRange& Particles, FSolverParticlesRange& SecondParticles, const FSolverReal Dt, const int32 ConstraintIndex, const FSolverReal Tol) const
+	{
+		if (Stiffness[ConstraintIndex] < Tol)
+		{
+			return;
+		}
+
+		ensure(SecondIndices.Num() > 0);
+		FSolverReal AlphaTilde = FSolverReal(2) / (Stiffness[ConstraintIndex] * Dt * Dt);
+		if (Stiffness[ConstraintIndex] > SoftMaxStiffness)
+		{
+			AlphaTilde = FSolverReal(0);
+		}
+
+		TVec3<FSolverReal> SpringEdge((FSolverReal)0.);
+		for (int32 j = 0; j < Indices[ConstraintIndex].Num(); j++)
+		{
+			for (int32 Alpha = 0; Alpha < 3; Alpha++)
+			{
+				SpringEdge[Alpha] += Weights[ConstraintIndex][j] * Particles.P(Indices[ConstraintIndex][j])[Alpha];
+			}
+		}
+		for (int32 j = 0; j < SecondIndices[ConstraintIndex].Num(); j++)
+		{
+			SpringEdge -= SecondWeights[ConstraintIndex][j] * SecondParticles.P(SecondIndices[ConstraintIndex][j]);
+		}
+
+		FSolverReal Denom = AlphaTilde;
+		for (int32 j = 0; j < Indices[ConstraintIndex].Num(); j++)
+		{
+			Denom += Weights[ConstraintIndex][j] * Weights[ConstraintIndex][j] * Particles.InvM(Indices[ConstraintIndex][j]);
+		}
+		for (int32 j = 0; j < SecondIndices[ConstraintIndex].Num(); j++)
+		{
+			Denom += SecondWeights[ConstraintIndex][j] * SecondWeights[ConstraintIndex][j] * SecondParticles.InvM(SecondIndices[ConstraintIndex][j]);
+		}
+		if (Denom == (FSolverReal)0.)
+		{
+			return;
+		}
+		for (int32 Beta = 0; Beta < 3; Beta++)
+		{
+			FSolverReal Cj = SpringEdge[Beta];
+			FSolverReal DLambda = -Cj - AlphaTilde * LambdaArray[ConstraintIndex * 3 + Beta];
+			DLambda /= Denom;
+			LambdaArray[ConstraintIndex * 3 + Beta] += DLambda;
+			for (int32 j = 0; j < Indices[ConstraintIndex].Num(); j++)
+			{
+				Particles.P(Indices[ConstraintIndex][j])[Beta] += DLambda * Weights[ConstraintIndex][j] * Particles.InvM(Indices[ConstraintIndex][j]);
+			}
+			for (int32 j = 0; j < SecondIndices[ConstraintIndex].Num(); j++)
+			{
+				SecondParticles.P(SecondIndices[ConstraintIndex][j])[Beta] -= DLambda * SecondWeights[ConstraintIndex][j] * SecondParticles.InvM(SecondIndices[ConstraintIndex][j]);
+			}
+		}
+	}
+
+	void FXPBDWeightedSpringConstraints::ApplySingleConstraintBetweenParticles(FSolverParticlesRange& Particles, FSolverParticlesRange& SecondParticles, const FSolverReal Dt, const int32 ConstraintIndex, const FSolverReal Tol) const
+	{
+		if (Stiffness[ConstraintIndex] < Tol || Dt == (FSolverReal)0. ||
+			!ensure(SecondIndices.IsValidIndex(ConstraintIndex) && SecondWeights.IsValidIndex(ConstraintIndex)))
+		{
+			return;
+		}
+
+		FSolverReal AlphaTilde = Stiffness[ConstraintIndex] > SoftMaxStiffness? FSolverReal(2) / (Stiffness[ConstraintIndex] * Dt * Dt) : FSolverReal(0);
+
+		TVec3<FSolverReal> SpringEdge;
+		if (!ComputeSpringEdge(Particles, SecondParticles, SpringEdge, ConstraintIndex,
+			/*bUseParticleX*/ false, /*bUseConstraintTargetPosition*/ false))
+		{
+			return;
+		}
+
+		FSolverReal Denom = AlphaTilde;
+		for (int32 j = 0; j < Indices[ConstraintIndex].Num(); j++)
+		{
+			Denom += Weights[ConstraintIndex][j] * Weights[ConstraintIndex][j] * Particles.InvM(Indices[ConstraintIndex][j]);
+		}
+		for (int32 j = 0; j < SecondIndices[ConstraintIndex].Num(); j++)
+		{
+			Denom += SecondWeights[ConstraintIndex][j] * SecondWeights[ConstraintIndex][j] * SecondParticles.InvM(SecondIndices[ConstraintIndex][j]);
+		}
+		if (Denom == (FSolverReal)0.)
+		{
+			return;
+		}
+		for (int32 Beta = 0; Beta < 3; Beta++)
+		{
+			FSolverReal Cj = SpringEdge[Beta];
+			FSolverReal DLambda = -Cj - AlphaTilde * LambdaArray[ConstraintIndex * 3 + Beta];
+			DLambda /= Denom;
+			LambdaArray[ConstraintIndex * 3 + Beta] += DLambda;
+			for (int32 j = 0; j < Indices[ConstraintIndex].Num(); j++)
+			{
+				Particles.P(Indices[ConstraintIndex][j])[Beta] += DLambda * Weights[ConstraintIndex][j] * Particles.InvM(Indices[ConstraintIndex][j]);
+			}
+			for (int32 j = 0; j < SecondIndices[ConstraintIndex].Num(); j++)
+			{
+				SecondParticles.P(SecondIndices[ConstraintIndex][j])[Beta] -= DLambda * SecondWeights[ConstraintIndex][j] * SecondParticles.InvM(SecondIndices[ConstraintIndex][j]);
+			}
+		}
+	}
+
+	// -------------------------------------------------------------------
+	// VisualizeAllBindings
+	// -------------------------------------------------------------------
+
+	void FXPBDWeightedSpringConstraints::VisualizeAllBindings(const FSolverParticlesRange& InParticles, const FSolverParticlesRange& InSecondParticles,
+		const TFunctionRef<void(const FVector& Pos0, const FVector& Pos1, const FLinearColor& Color)>& DrawLine,
+		const TFunctionRef<void(const FVector& Pos, const FLinearColor& Color)>& DrawPoint) const
+	{
+#if WITH_EDITOR
+		auto DoubleVert = [](Chaos::TVec3<FSolverReal> V) { return FVector3d(V.X, V.Y, V.Z); };
+		for (int32 ConstraintIndex = 0; ConstraintIndex < Indices.Num(); ++ConstraintIndex)
+		{
+			Chaos::TVec3<FSolverReal> SourcePos((FSolverReal)0.), TargetPos((FSolverReal)0.);
+			for (int32 IndicesLocalIndex = 0; IndicesLocalIndex < Indices[ConstraintIndex].Num(); ++IndicesLocalIndex)
+			{
+				SourcePos += Weights[ConstraintIndex][IndicesLocalIndex] * InParticles.P(Indices[ConstraintIndex][IndicesLocalIndex]);
+			}
+			for (int32 SecondIndicesLocalIndex = 0; SecondIndicesLocalIndex < SecondIndices[ConstraintIndex].Num(); ++SecondIndicesLocalIndex)
+			{
+				TargetPos += SecondWeights[ConstraintIndex][SecondIndicesLocalIndex] * InSecondParticles.P(SecondIndices[ConstraintIndex][SecondIndicesLocalIndex]);
+			}
+
+			if (Indices[ConstraintIndex].Num() > 1)
+			{
+				for (int32 IndicesLocalIndex = 0; IndicesLocalIndex < Indices[ConstraintIndex].Num(); ++IndicesLocalIndex)
+				{
+					DrawLine(DoubleVert(InParticles.P(Indices[ConstraintIndex][IndicesLocalIndex])), DoubleVert(InParticles.P(Indices[ConstraintIndex][(IndicesLocalIndex + 1) % Indices[ConstraintIndex].Num()])), FColor::Green);
+				}
+			}
+
+			if (SecondIndices[ConstraintIndex].Num() > 1)
+			{
+				for (int32 SecondIndicesLocalIndex = 0; SecondIndicesLocalIndex < SecondIndices[ConstraintIndex].Num(); ++SecondIndicesLocalIndex)
+				{
+					DrawLine(DoubleVert(InSecondParticles.P(SecondIndices[ConstraintIndex][SecondIndicesLocalIndex])), DoubleVert(InSecondParticles.P(SecondIndices[ConstraintIndex][(SecondIndicesLocalIndex + 1) % SecondIndices[ConstraintIndex].Num()])), FColor::Green);
+				}
+			}
+			DrawPoint(DoubleVert(SourcePos), FColor::Red);
+			DrawPoint(DoubleVert(TargetPos), FColor::Red);
+			DrawLine(DoubleVert(SourcePos), DoubleVert(TargetPos), FColor::Yellow);
+		}
+#endif
+	}
+
+	void FXPBDWeightedSpringConstraints::VisualizeAllBindings(const FSolverParticlesRange& InParticles, const FSolverReal Dt) const
+	{
+#if WITH_EDITOR
+		auto DoubleVert = [](Chaos::TVec3<FSolverReal> V) { return FVector3d(V.X, V.Y, V.Z); };
+		for (int32 ConstraintIndex = 0; ConstraintIndex < Indices.Num(); ++ConstraintIndex)
+		{
+			Chaos::TVec3<FSolverReal> SourcePos((FSolverReal)0.), TargetPos((FSolverReal)0.);
+			for (int32 IndicesLocalIndex = 0; IndicesLocalIndex < Indices[ConstraintIndex].Num(); ++IndicesLocalIndex)
+			{
+				SourcePos += Weights[ConstraintIndex][IndicesLocalIndex] * InParticles.P(Indices[ConstraintIndex][IndicesLocalIndex]);
+			}
+			if (SecondIndices.IsValidIndex(ConstraintIndex) && SecondIndices[ConstraintIndex].Num() > 0)
+			{
+				for (int32 SecondIndicesLocalIndex = 0; SecondIndicesLocalIndex < SecondIndices[ConstraintIndex].Num(); ++SecondIndicesLocalIndex)
+				{
+					TargetPos += SecondWeights[ConstraintIndex][SecondIndicesLocalIndex] * InParticles.P(SecondIndices[ConstraintIndex][SecondIndicesLocalIndex]);
+				}
+			}
+			else
+			{
+				if (ensure(Constraints.IsValidIndex(ConstraintIndex)))
+				{
+					TargetPos = Constraints[ConstraintIndex];
+				}
+			}
+
+			float ParticleThickness = DebugDrawParams.DebugParticleWidth;
+			float LineThickness = DebugDrawParams.DebugLineWidth;
+
+			if (Indices[ConstraintIndex].Num() > 1)
+			{
+				for (int32 IndicesLocalIndex = 0; IndicesLocalIndex < Indices[ConstraintIndex].Num(); ++IndicesLocalIndex)
+				{
+					Chaos::FDebugDrawQueue::GetInstance().DrawDebugLine(DoubleVert(InParticles.P(Indices[ConstraintIndex][IndicesLocalIndex])), 
+						DoubleVert(InParticles.P(Indices[ConstraintIndex][(IndicesLocalIndex + 1) % Indices[ConstraintIndex].Num()])), FColor::Green, false, Dt, 0, LineThickness);
+				}
+			}
+
+			if (SecondIndices[ConstraintIndex].Num() > 1)
+			{
+				for (int32 SecondIndicesLocalIndex = 0; SecondIndicesLocalIndex < SecondIndices[ConstraintIndex].Num(); ++SecondIndicesLocalIndex)
+				{
+					Chaos::FDebugDrawQueue::GetInstance().DrawDebugLine(DoubleVert(InParticles.P(SecondIndices[ConstraintIndex][SecondIndicesLocalIndex])), 
+						DoubleVert(InParticles.P(SecondIndices[ConstraintIndex][(SecondIndicesLocalIndex + 1) % SecondIndices[ConstraintIndex].Num()])), FColor::Green, false, Dt, 0, LineThickness);
+				}
+			}
+
+			Chaos::FDebugDrawQueue::GetInstance().DrawDebugPoint(DoubleVert(SourcePos), FColor::Red, false, Dt, 0, ParticleThickness);
+			Chaos::FDebugDrawQueue::GetInstance().DrawDebugPoint(DoubleVert(TargetPos), FColor::Red, false, Dt, 0, ParticleThickness);
+			Chaos::FDebugDrawQueue::GetInstance().DrawDebugLine(DoubleVert(SourcePos), DoubleVert(TargetPos), FColor::Yellow, false, Dt, 0, LineThickness);
+		}
+#endif
+	}
+
+	void FXPBDWeightedSpringConstraints::VisualizeAllBindings(const FSolverParticlesRange& InParticles,
+		const TFunctionRef<void(const FVector& Pos0, const FVector& Pos1, const FLinearColor& Color)>& DrawLine,
+		const TFunctionRef<void(const FVector& Pos, const FLinearColor& Color)>& DrawPoint) const
+	{
+		VisualizeAllBindings(InParticles, InParticles, DrawLine, DrawPoint);
+	}
+
+} // End namespace Chaos::Softs
diff --git a/Engine/Source/Runtime/Experimental/Chaos/Private/Chaos/GraphColoring.cpp b/Engine/Source/Runtime/Experimental/Chaos/Private/Chaos/GraphColoring.cpp
--- a/Engine/Source/Runtime/Experimental/Chaos/Private/Chaos/GraphColoring.cpp
+++ b/Engine/Source/Runtime/Experimental/Chaos/Private/Chaos/GraphColoring.cpp
@@ -116,6 +116,59 @@ static bool VerifyWeakConstraintsColoring(const TArray<TArray<int32>>& Indices,
 
 }
 
+template <typename T>
+static bool VerifyWeakConstraintsColoring(const TArray<TArray<int32>>& Indices, const TArray<TArray<int32>>& SecondIndices, const Chaos::TDynamicParticles<T, 3>& InParticles, const Chaos::TDynamicParticles<T, 3>& InSecondParticles, const TArray<TArray<int32>>& ConstraintsPerColor)
+{
+	TArray<bool> ConstraintIsIncluded;
+	ConstraintIsIncluded.Init(false, Indices.Num());
+	for (int32 ColorIndex = 0; ColorIndex < ConstraintsPerColor.Num(); ++ColorIndex)
+	{
+		for (int32 ConstraintIndex = 0; ConstraintIndex < ConstraintsPerColor[ColorIndex].Num(); ++ConstraintIndex)
+		{
+			ConstraintIsIncluded[ConstraintsPerColor[ColorIndex][ConstraintIndex]] = true;
+		}
+	}
+
+	for (int32 ConstraintIndex = 0; ConstraintIndex < Indices.Num(); ++ConstraintIndex)
+	{
+		if (!ConstraintIsIncluded[ConstraintIndex])
+		{
+			return false;
+		}
+	}
+
+	for (int32 ColorIndex = 0; ColorIndex < ConstraintsPerColor.Num(); ++ColorIndex)
+	{
+		TSet<int32> CoveredParticles;
+		TSet<int32> CoveredSecondParticles;
+		for (int32 ConstraintIndex = 0; ConstraintIndex < ConstraintsPerColor[ColorIndex].Num(); ++ConstraintIndex)
+		{
+			for (int32 ParticleIndex : Indices[ConstraintsPerColor[ColorIndex][ConstraintIndex]])
+			{
+				if (CoveredParticles.Contains(ParticleIndex))
+				{
+					return false;
+				}
+				CoveredParticles.Emplace(ParticleIndex);
+			}
+
+			if (SecondIndices.Num() > 0)
+			{
+				for (int32 SecondParticleIndex : SecondIndices[ConstraintsPerColor[ColorIndex][ConstraintIndex]])
+				{
+					if (CoveredSecondParticles.Contains(SecondParticleIndex))
+					{
+						return false;
+					}
+					CoveredSecondParticles.Emplace(SecondParticleIndex);
+				}
+
+			}
+		}
+	}
+	return true;
+}
+
 template <typename T>
 static bool VerifyNodalColoring(const TArray<Chaos::TVec4<int32>>& Graph, const Chaos::TDynamicParticles<T, 3>& InParticles, const int32 GraphParticlesStart, const int32 GraphParticlesEnd, const TArray<TArray<int32>>& IncidentElements, const TArray<TArray<int32>>& IncidentElementsLocalIndex, const TArray<TArray<int32>> ParticlesPerColor)
 {
@@ -562,8 +615,8 @@ void Chaos::ComputeGridBasedGraphSubColoringPointer(const TArray<TArray<int32>>&
 template<typename T>
 void Chaos::ComputeWeakConstraintsColoring(const TArray<TArray<int32>>& Indices, const TArray<TArray<int32>>& SecondIndices, const Chaos::TDynamicParticles<T, 3>& InParticles, TArray<TArray<int32>>& ConstraintsPerColor)
 {
-	TArray<TSet<int32>*> UsedColors;
-	UsedColors.Init(nullptr, InParticles.Size());
+	TArray<TSet<int32>> UsedColors;
+	UsedColors.SetNum(InParticles.Size());
 
 	ensure(Indices.Num() == SecondIndices.Num() || SecondIndices.Num() == 0);
 
@@ -574,22 +627,44 @@ void Chaos::ComputeWeakConstraintsColoring(const TArray<TArray<int32>>& Indices,
 	{
 		for (int32 ConstraintIndex = 0; ConstraintIndex < Indices.Num(); ConstraintIndex++)
 		{
-			for (int32 i = 0; i < Indices[ConstraintIndex].Num(); i++)
+			int32 ColorToUse = 0;
+			while (true)
 			{
-				if (!UsedColors[Indices[ConstraintIndex][i]]) {
-					UsedColors[Indices[ConstraintIndex][i]] = new TSet<int32>();
+				bool ColorFound = false;
+				for (int32 node : Indices[ConstraintIndex]) {
+					if (UsedColors[node].Contains(ColorToUse)) {
+						ColorFound = true;
+						break;
+					}
+				}
+				if (!ColorFound) {
+					break;
 				}
-				
+				ColorToUse++;
+			}
+			ConstraintColors[ConstraintIndex] = ColorToUse;
+			for (int32 node : Indices[ConstraintIndex])
+			{
+				UsedColors[node].Emplace(ColorToUse);
 			}
+		}
+	}
+	else
+	{
+		for (int32 ConstraintIndex = 0; ConstraintIndex < Indices.Num(); ConstraintIndex++)
+		{
 			int32 ColorToUse = 0;
 			while (true)
 			{
 				bool ColorFound = false;
-				for (auto node : Indices[ConstraintIndex]) {
-					if (!UsedColors[node]) {
-						UsedColors[node] = new TSet<int32>();
+				for (int32 node : Indices[ConstraintIndex]) {
+					if (UsedColors[node].Contains(ColorToUse)) {
+						ColorFound = true;
+						break;
 					}
-					if (UsedColors[node]->Contains(ColorToUse)) {
+				}
+				for (int32 node : SecondIndices[ConstraintIndex]) {
+					if (UsedColors[node].Contains(ColorToUse)) {
 						ColorFound = true;
 						break;
 					}
@@ -600,42 +675,88 @@ void Chaos::ComputeWeakConstraintsColoring(const TArray<TArray<int32>>& Indices,
 				ColorToUse++;
 			}
 			ConstraintColors[ConstraintIndex] = ColorToUse;
+			for (int32 node : Indices[ConstraintIndex]) {
+				UsedColors[node].Emplace(ColorToUse);
+			}
+			for (int32 node : SecondIndices[ConstraintIndex]) {
+				UsedColors[node].Emplace(ColorToUse);
+			}
 		}
 	}
-	else 
+
+	int32 NumColors = FMath::Max<int32>(ConstraintColors);
+	ConstraintsPerColor.Empty();
+	ConstraintsPerColor.SetNum(NumColors + 1);
+
+	for (int32 j = 0; j < Indices.Num(); j++) 
+	{
+		ConstraintsPerColor[ConstraintColors[j]].Emplace(j);
+	}
+
+	checkSlow(VerifyWeakConstraintsColoring<T>(Indices, SecondIndices, InParticles, ConstraintsPerColor));
+}
+
+template<typename T>
+void Chaos::ComputeWeakConstraintsColoring(const TArray<TArray<int32>>& Indices, const TArray<TArray<int32>>& SecondIndices, const Chaos::TDynamicParticles<T, 3>& InParticles, const Chaos::TDynamicParticles<T, 3>& InSecondParticles, TArray<TArray<int32>>& ConstraintsPerColor)
+{
+	TArray<TSet<int32>> ParticleUsedColors, SecondParticleUsedColors;
+	ParticleUsedColors.SetNum(InParticles.Size());
+	SecondParticleUsedColors.SetNum(InSecondParticles.Size());
+
+	ensure(Indices.Num() == SecondIndices.Num() || SecondIndices.Num() == 0);
+
+	TArray<int32> ConstraintColors;
+	ConstraintColors.Init(-1, Indices.Num());
+
+	if (SecondIndices.Num() == 0)
 	{
 		for (int32 ConstraintIndex = 0; ConstraintIndex < Indices.Num(); ConstraintIndex++)
 		{
-			for (int32 i = 0; i < Indices[ConstraintIndex].Num(); i++)
+			int32 ColorToUse = 0;
+			while (true)
 			{
-				if (!UsedColors[Indices[ConstraintIndex][i]]) {
-					UsedColors[Indices[ConstraintIndex][i]] = new TSet<int32>();
+				bool ColorFound = false;
+				for (int32 ParticleIndex : Indices[ConstraintIndex])
+				{
+					if (ParticleUsedColors[ParticleIndex].Contains(ColorToUse))
+					{
+						ColorFound = true;
+						break;
+					}
 				}
+				if (!ColorFound)
+				{
+					break;
+				}
+				ColorToUse++;
 			}
-			for (int32 j = 0; j < SecondIndices[ConstraintIndex].Num(); j++)
+			ConstraintColors[ConstraintIndex] = ColorToUse;
+			for (int32 ParticleIndex : Indices[ConstraintIndex])
 			{
-				if (!UsedColors[SecondIndices[ConstraintIndex][j]]) {
-					UsedColors[SecondIndices[ConstraintIndex][j]] = new TSet<int32>();
-				}
+				ParticleUsedColors[ParticleIndex].Emplace(ColorToUse);
 			}
+		}
+	}
+	else
+	{
+		for (int32 ConstraintIndex = 0; ConstraintIndex < Indices.Num(); ConstraintIndex++)
+		{
 			int32 ColorToUse = 0;
 			while (true)
 			{
 				bool ColorFound = false;
-				for (auto node : Indices[ConstraintIndex]) {
-					if (!UsedColors[node]) {
-						UsedColors[node] = new TSet<int32>();
-					}
-					if (UsedColors[node]->Contains(ColorToUse)) {
+				for (int32 ParticleIndex : Indices[ConstraintIndex])
+				{
+					if (ParticleUsedColors[ParticleIndex].Contains(ColorToUse))
+					{
 						ColorFound = true;
 						break;
 					}
 				}
-				for (auto node : SecondIndices[ConstraintIndex]) {
-					if (!UsedColors[node]) {
-						UsedColors[node] = new TSet<int32>();
-					}
-					if (UsedColors[node]->Contains(ColorToUse)) {
+				for (int32 SecondParticleIndex : SecondIndices[ConstraintIndex])
+				{
+					if (SecondParticleUsedColors[SecondParticleIndex].Contains(ColorToUse))
+					{
 						ColorFound = true;
 						break;
 					}
@@ -646,34 +767,30 @@ void Chaos::ComputeWeakConstraintsColoring(const TArray<TArray<int32>>& Indices,
 				ColorToUse++;
 			}
 			ConstraintColors[ConstraintIndex] = ColorToUse;
-			for (auto node : Indices[ConstraintIndex]) {
-				UsedColors[node]->Emplace(ColorToUse);
+			for (int32 ParticleIndex : Indices[ConstraintIndex])
+			{
+				ParticleUsedColors[ParticleIndex].Emplace(ColorToUse);
 			}
-			for (auto node : SecondIndices[ConstraintIndex]) {
-				UsedColors[node]->Emplace(ColorToUse);
+			for (int32 SecondParticleIndex : SecondIndices[ConstraintIndex])
+			{
+				SecondParticleUsedColors[SecondParticleIndex].Emplace(ColorToUse);
 			}
 		}
 	}
 
-	for (int ii = 0; ii < UsedColors.Num(); ii++) {
-		delete (UsedColors[ii]);
-	}
-
 	int32 NumColors = FMath::Max<int32>(ConstraintColors);
 
-	//int32 num_colors = *std::max_element(ElementSubColors.begin(), ElementSubColors.end());
 	ConstraintsPerColor.Empty();
 	ConstraintsPerColor.SetNum(NumColors + 1);
 
-	for (int32 j = 0; j < Indices.Num(); j++) 
+	for (int32 ConstraintIndex = 0; ConstraintIndex < Indices.Num(); ConstraintIndex++)
 	{
-		ConstraintsPerColor[ConstraintColors[j]].Emplace(j);
+		ConstraintsPerColor[ConstraintColors[ConstraintIndex]].Emplace(ConstraintIndex);
 	}
 
-	checkSlow(VerifyWeakConstraintsColoring<T>(Indices, SecondIndices, InParticles, ConstraintsPerColor));
+	checkSlow(VerifyWeakConstraintsColoring<T>(Indices, SecondIndices, InParticles, InSecondParticles, ConstraintsPerColor));
 }
 
-
 template<typename T>
 TArray<TArray<int32>> Chaos::ComputeNodalColoring(const TArray<TVec4<int32>>& Graph, const Chaos::TDynamicParticles<T, 3>& InParticles, const int32 GraphParticlesStart, const int32 GraphParticlesEnd, const TArray<TArray<int32>>& IncidentElements, const TArray<TArray<int32>>& IncidentElementsLocalIndex)
 {
@@ -694,49 +811,39 @@ TArray<TArray<int32>> Chaos::ComputeNodalColoring(const TArray<TVec4<int32>>& Gr
 		}
 	}
 
-	TArray<TSet<int32>*> ElementColorsSet;
-	ElementColorsSet.Init(nullptr, Graph.Num());
+	TArray<TSet<int32>> ElementColorsSet;
+	ElementColorsSet.SetNum(Graph.Num());
 	TArray<int32> ParticleColors;
 	ParticleColors.Init(-1, GraphParticlesEnd - GraphParticlesStart);
 
-	for (int32 p = 0; p < GraphParticlesEnd - GraphParticlesStart; p++) 
+	for (int32 p = 0; p < GraphParticlesEnd - GraphParticlesStart; p++)
 	{
 		int32 ParticleIndex = p + GraphParticlesStart;
-		if (InParticles.InvM(ParticleIndex) != (T)0. && Particle2Incident[ParticleIndex] != -1) 
+		if (InParticles.InvM(ParticleIndex) != (T)0. && Particle2Incident[ParticleIndex] != -1)
 		{
-			for (int32 j = 0; j < IncidentElements[Particle2Incident[ParticleIndex]].Num(); j++) 
-			{
-				int32 LocalIndex = IncidentElementsLocalIndex[Particle2Incident[ParticleIndex]][j];
-				int32 e = IncidentElements[Particle2Incident[ParticleIndex]][j];
-				if (!ElementColorsSet[e]) 
-				{
-					ElementColorsSet[e] = new TSet<int32>();
-				}
-			}
 			int32 color_to_use = 0;
-			while (true) 
+			while (true)
 			{
 				bool ColorFound = false;
-				for (int32 j = 0; j < IncidentElements[Particle2Incident[ParticleIndex]].Num(); j++) 
+				for (int32 j = 0; j < IncidentElements[Particle2Incident[ParticleIndex]].Num(); j++)
 				{
 					int32 e = IncidentElements[Particle2Incident[ParticleIndex]][j];
-					if (ElementColorsSet[e]->Contains(color_to_use)) 
+					if (ElementColorsSet[e].Contains(color_to_use))
 					{
 						ColorFound = true;
 						break;
 					}
 				}
-				if (!ColorFound) 
+				if (!ColorFound)
 				{
 					ParticleColors[p] = color_to_use;
-					for (int32 j = 0; j < IncidentElements[Particle2Incident[ParticleIndex]].Num(); j++) 
+					for (int32 j = 0; j < IncidentElements[Particle2Incident[ParticleIndex]].Num(); j++)
 					{
 						int32 e = IncidentElements[Particle2Incident[ParticleIndex]][j];
-						ElementColorsSet[e]->Emplace(color_to_use);
-						if (ElementColorsSet[e]->Num() == 4) 
+						ElementColorsSet[e].Emplace(color_to_use);
+						if (ElementColorsSet[e].Num() == 4)
 						{
-							delete ElementColorsSet[e];
-							ElementColorsSet[e] = nullptr;
+							ElementColorsSet[e].Empty();
 						}
 					}
 					break;
@@ -746,14 +853,6 @@ TArray<TArray<int32>> Chaos::ComputeNodalColoring(const TArray<TVec4<int32>>& Gr
 		}
 	}
 
-	for (int32 ii = 0; ii < ElementColorsSet.Num(); ii++) 
-	{
-		if (ElementColorsSet[ii]) 
-		{
-			delete (ElementColorsSet[ii]);
-		}
-	}
-
 	int32 SizeColors = FMath::Max<int32>(ParticleColors);
 
 	ParticlesPerColor.SetNum(0);
@@ -791,52 +890,42 @@ TArray<TArray<int32>> Chaos::ComputeNodalColoring(const TArray<TArray<int32>>& G
 		Particle2Incident[i] = i+GraphParticlesStart;
 	}
 
-	TArray<TSet<int32>*> ElementColorsSet;
-	ElementColorsSet.Init(nullptr, Graph.Num());
+	TArray<TSet<int32>> ElementColorsSet;
+	ElementColorsSet.SetNum(Graph.Num());
 	TArray<int32> ParticleColors;
 	ParticleColors.Init(INDEX_NONE, GraphParticlesEnd - GraphParticlesStart);
 
 	if (InParticleActiveView)
 	{
-		InParticleActiveView->SequentialFor([&ElementColorsSet, &Graph, &GraphParticlesStart, &IncidentElements, &IncidentElementsLocalIndex, &Particle2Incident, &ParticleColors]
+		InParticleActiveView->SequentialFor([&ElementColorsSet, &Graph, &GraphParticlesStart, &IncidentElements, &Particle2Incident, &ParticleColors]
 		(const Chaos::TDynamicParticles<T, 3>& InParticles, int32 ParticleIndex)
 			{
 				const int32 p = ParticleIndex - GraphParticlesStart;
-				if (InParticles.InvM(ParticleIndex) != (T)0. && Particle2Incident[ParticleIndex] != INDEX_NONE) 
+				if (InParticles.InvM(ParticleIndex) != (T)0. && Particle2Incident[ParticleIndex] != INDEX_NONE)
 				{
-					for (int32 j = 0; j < IncidentElements[Particle2Incident[ParticleIndex]].Num(); j++) 
-					{
-						int32 LocalIndex = IncidentElementsLocalIndex[Particle2Incident[ParticleIndex]][j];
-						int32 e = IncidentElements[Particle2Incident[ParticleIndex]][j];
-						if (!ElementColorsSet[e]) 
-						{
-							ElementColorsSet[e] = new TSet<int32>();
-						}
-					}
 					int32 color_to_use = 0;
-					while (true) 
+					while (true)
 					{
 						bool ColorFound = false;
-						for (int32 j = 0; j < IncidentElements[Particle2Incident[ParticleIndex]].Num(); j++) 
+						for (int32 j = 0; j < IncidentElements[Particle2Incident[ParticleIndex]].Num(); j++)
 						{
 							int32 e = IncidentElements[Particle2Incident[ParticleIndex]][j];
-							if (ElementColorsSet[e]->Contains(color_to_use)) 
+							if (ElementColorsSet[e].Contains(color_to_use))
 							{
 								ColorFound = true;
 								break;
 							}
 						}
-						if (!ColorFound) 
+						if (!ColorFound)
 						{
 							ParticleColors[p] = color_to_use;
-							for (int32 j = 0; j < IncidentElements[Particle2Incident[ParticleIndex]].Num(); j++) 
+							for (int32 j = 0; j < IncidentElements[Particle2Incident[ParticleIndex]].Num(); j++)
 							{
 								int32 e = IncidentElements[Particle2Incident[ParticleIndex]][j];
-								ElementColorsSet[e]->Emplace(color_to_use);
-								if (ElementColorsSet[e]->Num() == Graph[e].Num())
+								ElementColorsSet[e].Emplace(color_to_use);
+								if (ElementColorsSet[e].Num() == Graph[e].Num())
 								{
-									delete ElementColorsSet[e];
-									ElementColorsSet[e] = nullptr;
+									ElementColorsSet[e].Empty();
 								}
 							}
 							break;
@@ -848,44 +937,34 @@ TArray<TArray<int32>> Chaos::ComputeNodalColoring(const TArray<TArray<int32>>& G
 	}
 	else
 	{
-		for (int32 p = 0; p < GraphParticlesEnd - GraphParticlesStart; p++) 
+		for (int32 p = 0; p < GraphParticlesEnd - GraphParticlesStart; p++)
 		{
 			int32 ParticleIndex = p + GraphParticlesStart;
-			if (InParticles.InvM(ParticleIndex) != (T)0. && Particle2Incident[ParticleIndex] != INDEX_NONE) 
+			if (InParticles.InvM(ParticleIndex) != (T)0. && Particle2Incident[ParticleIndex] != INDEX_NONE)
 			{
-				for (int32 j = 0; j < IncidentElements[Particle2Incident[ParticleIndex]].Num(); j++) 
-				{
-					int32 LocalIndex = IncidentElementsLocalIndex[Particle2Incident[ParticleIndex]][j];
-					int32 e = IncidentElements[Particle2Incident[ParticleIndex]][j];
-					if (!ElementColorsSet[e]) 
-					{
-						ElementColorsSet[e] = new TSet<int32>();
-					}
-				}
 				int32 color_to_use = 0;
-				while (true) 
+				while (true)
 				{
 					bool ColorFound = false;
-					for (int32 j = 0; j < IncidentElements[Particle2Incident[ParticleIndex]].Num(); j++) 
+					for (int32 j = 0; j < IncidentElements[Particle2Incident[ParticleIndex]].Num(); j++)
 					{
 						int32 e = IncidentElements[Particle2Incident[ParticleIndex]][j];
-						if (ElementColorsSet[e]->Contains(color_to_use)) 
+						if (ElementColorsSet[e].Contains(color_to_use))
 						{
 							ColorFound = true;
 							break;
 						}
 					}
-					if (!ColorFound) 
+					if (!ColorFound)
 					{
 						ParticleColors[p] = color_to_use;
-						for (int32 j = 0; j < IncidentElements[Particle2Incident[ParticleIndex]].Num(); j++) 
+						for (int32 j = 0; j < IncidentElements[Particle2Incident[ParticleIndex]].Num(); j++)
 						{
 							int32 e = IncidentElements[Particle2Incident[ParticleIndex]][j];
-							ElementColorsSet[e]->Emplace(color_to_use);
-							if (ElementColorsSet[e]->Num() == Graph[e].Num())
+							ElementColorsSet[e].Emplace(color_to_use);
+							if (ElementColorsSet[e].Num() == Graph[e].Num())
 							{
-								delete ElementColorsSet[e];
-								ElementColorsSet[e] = nullptr;
+								ElementColorsSet[e].Empty();
 							}
 						}
 						break;
@@ -896,14 +975,6 @@ TArray<TArray<int32>> Chaos::ComputeNodalColoring(const TArray<TArray<int32>>& G
 		}
 	}
 
-	for (int32 ii = 0; ii < ElementColorsSet.Num(); ii++) 
-	{
-		if (ElementColorsSet[ii]) 
-		{
-			delete (ElementColorsSet[ii]);
-		}
-	}
-
 	int32 SizeColors = FMath::Max<int32>(ParticleColors);
 
 	ParticlesPerColor.SetNum(0);
@@ -1128,6 +1199,9 @@ template CHAOS_API void Chaos::ComputeGridBasedGraphSubColoringPointer(const TAr
 template CHAOS_API void Chaos::ComputeWeakConstraintsColoring<Chaos::FRealSingle>(const TArray<TArray<int32>>& Indices, const TArray<TArray<int32>>& SecondIndices, const Chaos::TDynamicParticles<Chaos::FRealSingle, 3>& InParticles, TArray<TArray<int32>>& ConstraintsPerColor);
 template CHAOS_API void Chaos::ComputeWeakConstraintsColoring<Chaos::FRealDouble>(const TArray<TArray<int32>>& Indices, const TArray<TArray<int32>>& SecondIndices, const Chaos::TDynamicParticles<Chaos::FRealDouble, 3>& InParticles, TArray<TArray<int32>>& ConstraintsPerColor);
 
+template CHAOS_API void Chaos::ComputeWeakConstraintsColoring<Chaos::FRealSingle>(const TArray<TArray<int32>>& Indices, const TArray<TArray<int32>>& SecondIndices, const Chaos::TDynamicParticles<Chaos::FRealSingle, 3>& InParticles, const Chaos::TDynamicParticles<FRealSingle, 3>& InSecondParticles, TArray<TArray<int32>>& ConstraintsPerColor);
+template CHAOS_API void Chaos::ComputeWeakConstraintsColoring<Chaos::FRealDouble>(const TArray<TArray<int32>>& Indices, const TArray<TArray<int32>>& SecondIndices, const Chaos::TDynamicParticles<Chaos::FRealDouble, 3>& InParticles, const Chaos::TDynamicParticles<FRealDouble, 3>& InSecondParticles, TArray<TArray<int32>>& ConstraintsPerColor);
+
 template CHAOS_API TArray<TArray<int32>> Chaos::ComputeNodalColoring<Chaos::FRealSingle>(const TArray<TVec4<int32>>& Graph, const Chaos::TDynamicParticles<Chaos::FRealSingle, 3>& InParticles, const int32 GraphParticlesStart, const int32 GraphParticlesEnd, const TArray<TArray<int32>>& IncidentElements, const TArray<TArray<int32>>& IncidentElementsLocalIndex);
 template CHAOS_API TArray<TArray<int32>> Chaos::ComputeNodalColoring<Chaos::FRealDouble>(const TArray<TVec4<int32>>& Graph, const Chaos::TDynamicParticles<Chaos::FRealDouble, 3>& InParticles, const int32 GraphParticlesStart, const int32 GraphParticlesEnd, const TArray<TArray<int32>>& IncidentElements, const TArray<TArray<int32>>& IncidentElementsLocalIndex);
 
diff --git a/Engine/Source/Runtime/Experimental/Chaos/Private/Chaos/TriangleMesh.cpp b/Engine/Source/Runtime/Experimental/Chaos/Private/Chaos/TriangleMesh.cpp
--- a/Engine/Source/Runtime/Experimental/Chaos/Private/Chaos/TriangleMesh.cpp
+++ b/Engine/Source/Runtime/Experimental/Chaos/Private/Chaos/TriangleMesh.cpp
@@ -1975,6 +1975,75 @@ bool FTriangleMesh::PointClosestTriangleQuery(const TSpatialHashType<T>& Spatial
 template CHAOS_API bool FTriangleMesh::PointClosestTriangleQuery<FRealSingle>(const TSpatialHashType<FRealSingle>& SpatialHash, const TConstArrayView<TVector<FRealSingle, 3>>& Points, const int32 PointIndex, const TVector<FRealSingle, 3>& PointPosition, const FRealSingle PointThickness, const FRealSingle ThisThickness, TFunctionRef<bool(const int32 PointIndex, const int32 TriangleIndex)> BroadphaseTest, TArray<TTriangleCollisionPoint<FRealSingle>>& Result) const;
 template CHAOS_API bool FTriangleMesh::PointClosestTriangleQuery<FRealDouble>(const TSpatialHashType<FRealDouble>& SpatialHash, const TConstArrayView<TVector<FRealDouble, 3>>& Points, const int32 PointIndex, const TVector<FRealDouble, 3>& PointPosition, const FRealDouble PointThickness, const FRealDouble ThisThickness, TFunctionRef<bool(const int32 PointIndex, const int32 TriangleIndex)> BroadphaseTest, TArray<TTriangleCollisionPoint<FRealDouble>>& Result) const;
 
+template<typename T>
+bool FTriangleMesh::PointClosestTriangleQuery(const TSpatialHashType<T>& SpatialHash, TFunctionRef<const TVec3<T>& (int32 PointIndex)> GetPosition, const int32 PointIndex, const TVec3<T>& PointPosition, const T PointThickness, const T ThisThickness,
+	TFunctionRef<bool(const int32 PointIndex, const int32 TriangleIndex)> BroadphaseTest, TArray<TTriangleCollisionPoint<T>>& Result) const
+{
+	TRACE_CPUPROFILER_EVENT_SCOPE(FTriangleMesh_PointClosestTriangleQuery);
+	const T TotalThickness = ThisThickness + PointThickness;
+	const T TotalThicknessSq = TotalThickness * TotalThickness;
+	typename TSpatialHashType<T>::FVectorAABB QueryBounds(PointPosition);
+	QueryBounds.Thicken(TotalThickness);
+
+	const TArray<int32> PotentialIntersections = SpatialHash.FindAllIntersections(QueryBounds,
+		[PointIndex, &BroadphaseTest](int32 Payload)
+		{
+			return BroadphaseTest(PointIndex, Payload);
+		});
+	TRACE_CPUPROFILER_EVENT_SCOPE(FTriangleMesh_PointClosestTriangleQueryAddTriangle);
+	Result.Reset(0);
+	T MinDistanceSq = TotalThicknessSq;
+	int32 MinTriId = INDEX_NONE;
+	for (int32 TriIdx : PotentialIntersections)
+	{
+		const TVec3<T>& A = GetPosition(MElements[TriIdx][0]);
+		const TVec3<T>& B = GetPosition(MElements[TriIdx][1]);
+		const TVec3<T>& C = GetPosition(MElements[TriIdx][2]);
+		TVec3<T> Bary;
+		const TVec3<T> ClosestPoint = FindClosestPointAndBaryOnTriangle(A, B, C, PointPosition, Bary);
+
+		const T DistSq = (PointPosition - ClosestPoint).SizeSquared();
+		if (DistSq >= TotalThicknessSq || DistSq >= MinDistanceSq) // also filters 0 vs 0 case
+		{
+			// Failed narrow test.
+			continue;
+		}
+		else
+		{
+			MinDistanceSq = DistSq;
+			MinTriId = TriIdx;
+		}
+	}
+	if (MinTriId != INDEX_NONE)
+	{
+		const TVec3<T>& A = GetPosition(MElements[MinTriId][0]);
+		const TVec3<T>& B = GetPosition(MElements[MinTriId][1]);
+		const TVec3<T>& C = GetPosition(MElements[MinTriId][2]);
+		const TTriangle<T> Triangle(A, B, C);
+		TVec3<T> Bary;
+		const TVec3<T> ClosestPoint = FindClosestPointAndBaryOnTriangle(A, B, C, PointPosition, Bary);
+		const T DistSq = (PointPosition - ClosestPoint).SizeSquared();
+		const TVec3<T> Normal = Triangle.GetNormal();
+
+		TTriangleCollisionPoint<T> CollisionPoint;
+		CollisionPoint.ContactType = TTriangleCollisionPoint<T>::EContactType::PointFace;
+		CollisionPoint.Indices[0] = PointIndex;
+		CollisionPoint.Indices[1] = MinTriId;
+		CollisionPoint.Bary = TVec4<T>((T)1., Bary.X, Bary.Y, Bary.Z);
+		CollisionPoint.Location = ClosestPoint;
+		CollisionPoint.Normal = Normal;
+		CollisionPoint.Phi = FMath::Sqrt(DistSq);
+		if (TVec3<T>::DotProduct(Normal, PointPosition - A) > 0) //Point is inside of triangle (normal is pointing inward)
+		{
+			CollisionPoint.Phi = -CollisionPoint.Phi;
+		}
+		Result.Add(CollisionPoint);
+	}
+	return Result.Num() > 0;
+}
+template CHAOS_API bool FTriangleMesh::PointClosestTriangleQuery<FRealSingle>(const TSpatialHashType<FRealSingle>& SpatialHash, TFunctionRef<const TVec3<FRealSingle>& (int32 PointIndex)> GetPosition, const int32 PointIndex, const TVector<FRealSingle, 3>& PointPosition, const FRealSingle PointThickness, const FRealSingle ThisThickness, TFunctionRef<bool(const int32 PointIndex, const int32 TriangleIndex)> BroadphaseTest, TArray<TTriangleCollisionPoint<FRealSingle>>& Result) const;
+template CHAOS_API bool FTriangleMesh::PointClosestTriangleQuery<FRealDouble>(const TSpatialHashType<FRealDouble>& SpatialHash, TFunctionRef<const TVec3<FRealDouble>& (int32 PointIndex)> GetPosition, const int32 PointIndex, const TVector<FRealDouble, 3>& PointPosition, const FRealDouble PointThickness, const FRealDouble ThisThickness, TFunctionRef<bool(const int32 PointIndex, const int32 TriangleIndex)> BroadphaseTest, TArray<TTriangleCollisionPoint<FRealDouble>>& Result) const;
+
 
 template<typename T>
 bool FTriangleMesh::EdgeIntersectionQuery(const TSpatialHashType<T>& SpatialHash, const TConstArrayView<TVec3<T>>& Points, const int32 EdgeIndex, const TVec3<T>& EdgePosition1, const TVec3<T>& EdgePosition2,
diff --git a/Engine/Source/Runtime/Experimental/Chaos/Public/Chaos/Deformable/GaussSeidelMainConstraint.h b/Engine/Source/Runtime/Experimental/Chaos/Public/Chaos/Deformable/GaussSeidelMainConstraint.h
--- a/Engine/Source/Runtime/Experimental/Chaos/Public/Chaos/Deformable/GaussSeidelMainConstraint.h
+++ b/Engine/Source/Runtime/Experimental/Chaos/Public/Chaos/Deformable/GaussSeidelMainConstraint.h
@@ -1,4 +1,4 @@
-﻿// Copyright Epic Games, Inc. All Rights Reserved.
+// Copyright Epic Games, Inc. All Rights Reserved.
 #pragma once
 
 #include "Chaos/PBDSpringConstraintsBase.h"
@@ -10,6 +10,7 @@
 #include "Chaos/Framework/Parallel.h"
 #include "Chaos/Deformable/GaussSeidelWeakConstraints.h"
 #include "Chaos/Deformable/GaussSeidelCorotatedConstraints.h"
+#include "Chaos/CollectionPropertyFacade.h"
 
 
 #define PERF_SCOPE(X) SCOPE_CYCLE_COUNTER(X); TRACE_CPUPROFILER_EVENT_SCOPE(X);
@@ -23,106 +24,71 @@ CHAOS_API DECLARE_LOG_CATEGORY_EXTERN(LogDeformableGaussSeidelMainConstraint, Lo
 
 namespace Chaos::Softs
 {
+	// TFunction type aliases for Gauss-Seidel constraint lambdas
+	// (Particles, ParticleIdx, Dt, OutResidual, OutHessian)
+	using FGSInitialResidualHessianFunc = TFunction<void(const FSolverParticles&, const int32, const FSolverReal, TVec3<FSolverReal>&, Chaos::PMatrix<FSolverReal, 3, 3>&)>;
+	// (Particles, ConstraintIndex, LocalIndex, Dt, OutResidual, OutHessian)
+	using FGSConstraintResidualHessianFunc = TFunction<void(const FSolverParticles&, const int32, const int32, const FSolverReal, TVec3<FSolverReal>&, Chaos::PMatrix<FSolverReal, 3, 3>&)>;
+	// (SourceParticles, TargetParticles, ConstraintIndex, LocalIndex, Dt, OutResidual, OutHessian)
+	using FGSCrossBodyResidualHessianFunc = TFunction<void(const FSolverParticles&, const FSolverParticles&, const int32, const int32, const FSolverReal, TVec3<FSolverReal>&, Chaos::PMatrix<FSolverReal, 3, 3>&)>;
+	// (ParticleIdx, Dt, OutHessian)
+	using FGSPerNodeHessianFunc = TFunction<void(const int32, const FSolverReal, Chaos::PMatrix<FSolverReal, 3, 3>&)>;
+	// (Particles, Dx, OutForceDifferential)
+	using FGSForceDifferentialFunc = TFunction<void(const FSolverParticles&, const TArray<TVec3<FSolverReal>>&, TArray<TVec3<FSolverReal>>&)>;
+
 	class FGaussSeidelMainConstraints
 	{
 
 	public:
-		FGaussSeidelMainConstraints(
-			const FSolverParticles& InParticles, 
+		CHAOS_API FGaussSeidelMainConstraints(
+			const FSolverParticles& InParticles,
 			const bool bDoQuasistaticsIn = false,
 			const bool bDoSORIn = true,
 			const FSolverReal InOmegaSOR = (FSolverReal)1.6,
-			const int32 ParallelMaxIn = 1000, 
+			const int32 ParallelMaxIn = 1000,
 			const FSolverReal MaxDxRatioIn = FSolverReal(1),
-			const FDeformableXPBDCorotatedParams& InParams = FDeformableXPBDCorotatedParams())
-			: bDoQuasistatics(bDoQuasistaticsIn)
-			, bDoAcceleration(bDoSORIn)
-			, OmegaSOR(InOmegaSOR)
-			, ParallelMax(ParallelMaxIn)
-			, CorotatedParams(InParams)
-		{
-			Resize((int32)InParticles.Size());
-
-			InitializeLambdas();
-			
-			TVec3<FSolverReal> MaxCoord((FSolverReal)100.), MinCoord((FSolverReal)-100.);
-			for (int32 i = 0; i < (int32)InParticles.Size(); i++)
-			{
-				for (int32 j = 0; j < 3; j++) 
-				{
-					if (InParticles.X(i)[j] < MinCoord[j])
-					{
-						MinCoord[j] = InParticles.X(i)[j];
-					}
-					if (InParticles.X(i)[j] > MaxCoord[j])
-					{
-						MaxCoord[j] = InParticles.X(i)[j];
-					}
-				}
-			}
-			MaxDxSize = (MaxCoord - MinCoord).Size() * MaxDxRatioIn;
-		}
+			const FDeformableXPBDCorotatedParams& InParams = FDeformableXPBDCorotatedParams());
+
+		CHAOS_API FGaussSeidelMainConstraints(const FSolverParticles& InParticles, const FCollectionPropertyConstFacade& Property, const FSolverReal MaxDxRatioIn);
 
 		virtual ~FGaussSeidelMainConstraints() = default;
 
-		void Resize(const int32 NewSize)
-		{
-			xtilde.SetNum(NewSize);
-			StaticIncidentElements.SetNum(NewSize);
-			StaticIncidentElementsLocal.SetNum(NewSize);
-			DynamicIncidentElements.SetNum(NewSize);
-			DynamicIncidentElementsLocal.SetNum(NewSize);
-			TransientIncidentElements.SetNum(NewSize);
-			TransientIncidentElementsLocal.SetNum(NewSize);
-			X_k_1.Init(TVector<FSolverReal, 3>(FSolverReal(0.)), NewSize);
-			X_k.Init(TVector<FSolverReal, 3>(FSolverReal(0.)), NewSize);
-		}
-
-		const TArray<TFunction<void(const FSolverParticles&, const int32, const int32, const FSolverReal, TVec3<FSolverReal>&, Chaos::PMatrix<FSolverReal, 3, 3>&)>>& StaticConstraintResidualAndHessian() const { return AddStaticConstraintResidualAndHessian; }
-		TArray<TFunction<void(const FSolverParticles&, const int32, const int32, const FSolverReal, TVec3<FSolverReal>&, Chaos::PMatrix<FSolverReal, 3, 3>&)>>& StaticConstraintResidualAndHessian() { return AddStaticConstraintResidualAndHessian; }
-		const TArray<TFunction<void(const FSolverParticles&, const int32, const int32, const FSolverReal, TVec3<FSolverReal>&, Chaos::PMatrix<FSolverReal, 3, 3>&)>>& TransientConstraintResidualAndHessian() const { return AddTransientConstraintResidualAndHessian; }
-		TArray<TFunction<void(const FSolverParticles&, const int32, const int32, const FSolverReal, TVec3<FSolverReal>&, Chaos::PMatrix<FSolverReal, 3, 3>&)>>& TransientConstraintResidualAndHessian() { return AddTransientConstraintResidualAndHessian; }
-		const TArray<TFunction<void(const FSolverParticles&, const int32, const int32, const FSolverReal, TVec3<FSolverReal>&, Chaos::PMatrix<FSolverReal, 3, 3>&)>>& DynamicConstraintResidualAndHessian() const { return AddDynamicConstraintResidualAndHessian; }
-		TArray<TFunction<void(const FSolverParticles&, const int32, const int32, const FSolverReal, TVec3<FSolverReal>&, Chaos::PMatrix<FSolverReal, 3, 3>&)>>& DynamicConstraintResidualAndHessian() { return AddDynamicConstraintResidualAndHessian; }
-		const TArray<TFunction<void(const int32, const FSolverReal, Chaos::PMatrix<FSolverReal, 3, 3>&)>>& PerNodeHessian() const { return AddPerNodeHessian; }
-		TArray<TFunction<void(const int32, const FSolverReal, Chaos::PMatrix<FSolverReal, 3, 3>&)>>& PerNodeHessian() { return AddPerNodeHessian; }
-		const TArray<TFunction<void(const FSolverParticles&, const TArray<TVec3<FSolverReal>>&, TArray<TVec3<FSolverReal>>&)>>& InternalForceDifferentials() const { return AddInternalForceDifferentials; }
-		TArray<TFunction<void(const FSolverParticles&, const TArray<TVec3<FSolverReal>>&, TArray<TVec3<FSolverReal>>&)>>& InternalForceDifferentials() { return AddInternalForceDifferentials; }
-
-		int32 AddStaticConstraintResidualAndHessianRange(int32 NumConstraints)
-		{
-			int32 CurrentSize = AddStaticConstraintResidualAndHessian.Num();
-			AddStaticConstraintResidualAndHessian.AddDefaulted(NumConstraints);
-			return CurrentSize;
-		}
-
-		int32 AddTransientConstraintResidualAndHessianRange(int32 NumConstraints)
-		{
-			int32 CurrentSize = AddTransientConstraintResidualAndHessian.Num();
-			AddTransientConstraintResidualAndHessian.AddDefaulted(NumConstraints);
-			return CurrentSize;
-		}
-
-		int32 AddDynamicConstraintResidualAndHessianRange(int32 NumConstraints)
-		{
-			int32 CurrentSize = AddDynamicConstraintResidualAndHessian.Num();
-			AddDynamicConstraintResidualAndHessian.AddDefaulted(NumConstraints);
-			return CurrentSize;
-		}
-
-		int32 AddPerNodeHessianRange(int32 NumConstraints)
-		{
-			int32 CurrentSize = AddPerNodeHessian.Num();
-			AddPerNodeHessian.AddDefaulted(NumConstraints);
-			return CurrentSize;
-		}
-
-		int32 AddAddInternalForceDifferentialsRange(int32 NumConstraints)
-		{
-			int32 CurrentSize = AddInternalForceDifferentials.Num();
-			AddInternalForceDifferentials.AddDefaulted(NumConstraints);
-			return CurrentSize;
-		}
+		CHAOS_API void Resize(const int32 NewSize);
+
+		const TArray<FGSConstraintResidualHessianFunc>& StaticConstraintResidualAndHessian() const { return AddStaticConstraintResidualAndHessian; }
+		TArray<FGSConstraintResidualHessianFunc>& StaticConstraintResidualAndHessian() { return AddStaticConstraintResidualAndHessian; }
+		const TArray<FGSConstraintResidualHessianFunc>& TransientConstraintResidualAndHessian() const { return AddTransientConstraintResidualAndHessian; }
+		TArray<FGSConstraintResidualHessianFunc>& TransientConstraintResidualAndHessian() { return AddTransientConstraintResidualAndHessian; }
+		const TArray<FGSConstraintResidualHessianFunc>& DynamicConstraintResidualAndHessian() const { return AddDynamicConstraintResidualAndHessian; }
+		TArray<FGSConstraintResidualHessianFunc>& DynamicConstraintResidualAndHessian() { return AddDynamicConstraintResidualAndHessian; }
+
+		const TArray<FGSCrossBodyResidualHessianFunc>& CrossBodyStaticConstraintResidualAndHessian() const { return AddCrossBodyStaticConstraintResidualAndHessian; }
+		TArray<FGSCrossBodyResidualHessianFunc>& CrossBodyStaticConstraintResidualAndHessian() { return AddCrossBodyStaticConstraintResidualAndHessian; }
+		const TArray<FGSCrossBodyResidualHessianFunc>& CrossBodyTransientConstraintResidualAndHessian() const { return AddCrossBodyTransientConstraintResidualAndHessian; }
+		TArray<FGSCrossBodyResidualHessianFunc>& CrossBodyTransientConstraintResidualAndHessian() { return AddCrossBodyTransientConstraintResidualAndHessian; }
+
+		const TArray<FGSPerNodeHessianFunc>& PerNodeHessian() const { return AddPerNodeHessian; }
+		TArray<FGSPerNodeHessianFunc>& PerNodeHessian() { return AddPerNodeHessian; }
+		const TArray<FGSPerNodeHessianFunc>& TransientPerNodeHessian() const { return AddTransientPerNodeHessian; }
+		TArray<FGSPerNodeHessianFunc>& TransientPerNodeHessian() { return AddTransientPerNodeHessian; }
+		const TArray<FGSForceDifferentialFunc>& InternalForceDifferentials() const { return AddInternalForceDifferentials; }
+		TArray<FGSForceDifferentialFunc>& InternalForceDifferentials() { return AddInternalForceDifferentials; }
+
+		CHAOS_API int32 AddStaticConstraintResidualAndHessianRange(int32 NumConstraints);
+
+		CHAOS_API int32 AddTransientConstraintResidualAndHessianRange(int32 NumConstraints);
+
+		CHAOS_API int32 AddDynamicConstraintResidualAndHessianRange(int32 NumConstraints);
+
+		CHAOS_API int32 AddCrossBodyStaticConstraintResidualAndHessianRange(int32 NumConstraints);
+
+		CHAOS_API int32 AddCrossBodyTransientConstraintResidualAndHessianRange(int32 NumConstraints);
+
+		CHAOS_API int32 AddPerNodeHessianRange(int32 NumConstraints);
+
+		CHAOS_API int32 AddTransientPerNodeHessianRange(int32 NumConstraints);
+
+		CHAOS_API int32 AddAddInternalForceDifferentialsRange(int32 NumConstraints);
 
 		CHAOS_API void AddStaticConstraints(const TArray<TArray<int32>>& ExtraConstraints, TArray<TArray<int32>>& ExtraIncidentElements, TArray<TArray<int32>>& ExtraIncidentElementsLocal);
 
@@ -130,158 +96,22 @@ namespace Chaos::Softs
 
 		CHAOS_API void AddDynamicConstraints(const TArray<TArray<int32>>& ExtraConstraints, TArray<TArray<int32>>& ExtraIncidentElements, TArray<TArray<int32>>& ExtraIncidentElementsLocal, bool CheckIncidentElements = false);
 
-		inline void ResetDynamicConstraints()
-		{
-			DynamicConstraints = {};
-			for (int32 p = 0; p < DynamicIncidentElements.Num(); p++)
-			{
-				DynamicIncidentElements[p].SetNum(0);
-				DynamicIncidentElementsLocal[p].SetNum(0);
-			}
-			DynamicIncidentElementsOffsets = {};	
-		}
-
-
-		void Apply(FSolverParticles& Particles, const FSolverReal Dt, const int32 MaxWriteIters = 10, const bool Write2File = false, const TPBDActiveView<FSolverParticles>* InParticleActiveView = nullptr)
-		{
- 			PERF_SCOPE(STAT_ChaosGSMainConstraint_Apply);
-			
-			if (DebugResidual && PassedIters < MaxWriteIters)
-			{
-				ComputeNewtonResiduals(Particles, Dt, Write2File);
-			}
-			
-			std::atomic<int32> ParticleFailureCounter(0);
-
-			for (int32 i = 0; i < ParticlesPerColor.Num(); i++)
-			{
-				int32 NumBatch = ParticlesPerColor[i].Num() / CorotatedParams.XPBDCorotatedBatchSize;
-				if (ParticlesPerColor[i].Num() % CorotatedParams.XPBDCorotatedBatchSize != 0)
-				{
-					NumBatch ++;
-				}
-			
-				PhysicsParallelFor(NumBatch, [&](const int32 BatchIndex)
-					{
-						for (int32 BatchSubIndex = 0; BatchSubIndex < CorotatedParams.XPBDCorotatedBatchSize; BatchSubIndex++) {
-							int32 TaskIndex = CorotatedParams.XPBDCorotatedBatchSize * BatchIndex + BatchSubIndex;
-							if (TaskIndex < ParticlesPerColor[i].Num())
-							{
-								if (Particles.InvM(ParticlesPerColor[i][TaskIndex]) != FSolverReal(0))
-								{
-									int32 ParticleIndex = ParticlesPerColor[i][TaskIndex];
-									if (!ApplySingleParticle(ParticleIndex, Dt, Particles))
-									{
-										ParticleFailureCounter.fetch_add(1, std::memory_order_relaxed);
-									}
-								}
-							}
-						}
-					}, NumBatch < CorotatedParams.XPBDCorotatedBatchThreshold);
-			
-			}
-			
-						
-			if (bDoAcceleration)
-			{
-				PERF_SCOPE(STAT_ChaosGSMainConstraint_Acceleration);
-				if (InParticleActiveView)
-				{
-					InParticleActiveView->ParallelFor(AccelerationTechniquePerParticle, CorotatedParams.XPBDCorotatedBatchSize);
-				}
-				else
-				{
-					PhysicsParallelFor(Particles.Size(), [this, &Particles](const int32 ParticleIndex)
-						{
-							this->AccelerationTechniquePerParticle(Particles, ParticleIndex);
-						}, Particles.Size() < 1000);
-				}
-			}
-			
-			CurrentIt ++;
-
-			if (ParticleFailureCounter.load(std::memory_order_relaxed) > 0)
-			{
-				UE_LOG(LogDeformableGaussSeidelMainConstraint, Warning, TEXT("%d Particle(s) are skipped because of too large dx size"), ParticleFailureCounter.load(std::memory_order_relaxed));
-			}
-		}
-
-		void InitStaticColor(const FSolverParticles& Particles, const TPBDActiveView<FSolverParticles>* InParticleActiveView = nullptr)
-		{
-			StaticParticlesPerColor = ComputeNodalColoring(StaticConstraints, Particles, 0, Particles.Size(), StaticIncidentElements, StaticIncidentElementsLocal, InParticleActiveView,  &StaticParticleColors);
-			ParticleColors = StaticParticleColors;
-			ParticlesPerColor = StaticParticlesPerColor;
-		}
-
-		void InitTransientColor(const FSolverParticles& Particles)
-		{
-			PERF_SCOPE(STAT_ChaosGSMainConstraint_InitTransientColor);
-			ParticleColors = StaticParticleColors;
-			ParticlesPerColor = StaticParticlesPerColor;
-			Chaos::ComputeExtraNodalColoring(StaticConstraints, DynamicConstraints, TransientConstraints, Particles, StaticIncidentElements, DynamicIncidentElements, TransientIncidentElements, ParticleColors, ParticlesPerColor);
-		}
-
-		void InitDynamicColor(const FSolverParticles& Particles)
-		{
-			PERF_SCOPE(STAT_ChaosGSMainConstraint_InitDynamicColor);
-			ParticleColors = StaticParticleColors;
-			ParticlesPerColor = StaticParticlesPerColor;
-			Chaos::ComputeExtraNodalColoring(StaticConstraints, DynamicConstraints, Particles, StaticIncidentElements, TransientIncidentElements, ParticleColors, ParticlesPerColor);
-		}
-
-		void Init(const FSolverReal Dt, const FSolverParticles& Particles)
-		{
-			Resize((int32)Particles.Size());
-
-			PERF_SCOPE(STAT_ChaosGSMainConstraint_Init);
-			TransientConstraints.SetNum(0);
-			for (int32 p = 0; p < TransientIncidentElements.Num(); p++)
-			{
-				TransientIncidentElements[p].SetNum(0);
-				TransientIncidentElementsLocal[p].SetNum(0);
-			}
-			TransientIncidentElementsOffsets.SetNum(0);
-			if (!bDoQuasistatics)
-			{
-				for (int32 i = 0; i < xtilde.Num(); i++)
-				{
-					//xtilde[i] = Particles.X(i) + Dt * Particles.V(i);
-					xtilde[i] = Particles.P(i);
-				}
-			}
-			CurrentIt = 0;
-		}
-
-		static bool IsClean(const TArray<TArray<int32>>& ConstraintsIn, const TArray<TArray<int32>>& IncidentElementsIn, const TArray<TArray<int32>>& IncidentElementsLocalIn)
-		{
-			if (IncidentElementsIn.Num() == IncidentElementsLocalIn.Num())
-			{
-				int32 TotalEntries = 0;
-				for (int32 i = 0; i < IncidentElementsIn.Num(); i++)
-				{
-					TotalEntries += IncidentElementsIn[i].Num();
-					if (IncidentElementsIn[i].Num() != IncidentElementsLocalIn[i].Num())
-					{
-						return false;
-					}
-					for (int32 j = 0; j < IncidentElementsIn[i].Num(); j++)
-					{
-						if (IncidentElementsIn[i][j] > ConstraintsIn.Num() || IncidentElementsLocalIn[i][j] > ConstraintsIn[IncidentElementsIn[i][j]].Num())
-						{
-							return false;
-						}
-					}
-				}
-				if (TotalEntries > 0)
-				{
-					return true;
-				}
-			}
-			return false;
-		}
+		CHAOS_API void ResetDynamicConstraints();
+		CHAOS_API void ResetTransientConstraints();
+		CHAOS_API void ResetTransientConstraintsAndRules();
+
+		CHAOS_API void Apply(FSolverParticles& Particles, const FSolverReal Dt, const int32 MaxWriteIters = 10, const bool Write2File = false, const TPBDActiveView<FSolverParticles>* InParticleActiveView = nullptr);
+
+		CHAOS_API void InitStaticColor(const FSolverParticles& Particles, const TPBDActiveView<FSolverParticles>* InParticleActiveView = nullptr);
+
+		CHAOS_API void InitTransientColor(const FSolverParticles& Particles);
+
+		CHAOS_API void InitDynamicColor(const FSolverParticles& Particles);
+
+		CHAOS_API void Init(const FSolverReal Dt, const FSolverParticles& Particles);
 
 		CHAOS_API TArray<TVec3<FSolverReal>> ComputeNewtonResiduals(const FSolverParticles& Particles, const FSolverReal Dt, const bool Write2File = false, TArray<PMatrix<FSolverReal, 3, 3>>* AllParticleHessian = nullptr);
-		
+
 		CHAOS_API void ApplyCG(FSolverParticles& Particles, const FSolverReal Dt);
 
 		/* Adds external acceleration, e.g. gravity (0,0,-980) cm/s^2 */
@@ -289,136 +119,18 @@ namespace Chaos::Softs
 
 		void ResetExternalAcceleration() { ExternalAcceleration = TVec3<FSolverReal>((FSolverReal)0.); };
 
+		UE_CHAOS_DECLARE_INDEXLESS_PROPERTYCOLLECTION_NAME(UseGaussSeidelConstraints, bool);
+		UE_CHAOS_DECLARE_INDEXLESS_PROPERTYCOLLECTION_NAME(AccelerateSolverUsingSOR, bool);
+		UE_CHAOS_DECLARE_INDEXLESS_PROPERTYCOLLECTION_NAME(OmegaSOR, float);
+		UE_CHAOS_DECLARE_INDEXLESS_PROPERTYCOLLECTION_NAME(DoQuasistatics, bool);
+
 	private:
 
-		bool ApplySingleParticle(const int32 p, const FSolverReal Dt, FSolverParticles& Particles)
-		{
-			int32 ConstraintIndex = 0;
-			Chaos::TVector<FSolverReal, 3> ParticleResidual((FSolverReal)0.);
-			Chaos::PMatrix<FSolverReal, 3, 3> ParticleHessian((FSolverReal)0., (FSolverReal)0., (FSolverReal)0.);
-
-			ComputeInitialResidualAndHessian(Particles, p, Dt, ParticleResidual, ParticleHessian);
-
-			for (int32 i = 0; i < StaticIncidentElements[p].Num(); i++)
-			{
-				while (StaticIncidentElements[p][i] >= StaticIncidentElementsOffsets[ConstraintIndex + 1] && ConstraintIndex < StaticIncidentElementsOffsets.Num() - 1)
-				{
-					ConstraintIndex ++;
-				}
-
-				AddStaticConstraintResidualAndHessian[ConstraintIndex](Particles, StaticIncidentElements[p][i] - StaticIncidentElementsOffsets[ConstraintIndex], StaticIncidentElementsLocal[p][i], Dt, ParticleResidual, ParticleHessian);
-			}
-
-			ConstraintIndex = 0;
-
-			for (int32 i = 0; i < DynamicIncidentElements[p].Num(); i++)
-			{
-				while (DynamicIncidentElements[p][i] >= DynamicIncidentElementsOffsets[ConstraintIndex + 1] && ConstraintIndex < DynamicIncidentElementsOffsets.Num() - 1)
-				{
-					ConstraintIndex ++;
-				}
-
-				AddDynamicConstraintResidualAndHessian[ConstraintIndex](Particles, DynamicIncidentElements[p][i] - DynamicIncidentElementsOffsets[ConstraintIndex], DynamicIncidentElementsLocal[p][i], Dt, ParticleResidual, ParticleHessian);
-			}
-
-			ConstraintIndex = 0;
-
-			for (int32 i = 0; i < TransientIncidentElements[p].Num(); i++)
-			{
-				while (TransientIncidentElements[p][i] >= TransientIncidentElementsOffsets[ConstraintIndex + 1] && ConstraintIndex < TransientIncidentElementsOffsets.Num() - 1)
-				{
-					ConstraintIndex ++;
-				}
-
-				AddTransientConstraintResidualAndHessian[ConstraintIndex](Particles, TransientIncidentElements[p][i] - TransientIncidentElementsOffsets[ConstraintIndex], TransientIncidentElementsLocal[p][i], Dt, ParticleResidual, ParticleHessian);
-			}
-
-			for (int32 i = 0; i < AddPerNodeHessian.Num(); i++)
-			{
-				AddPerNodeHessian[i](p, Dt, ParticleHessian);
-			}
-
-			FSolverReal HessianScale = (FSolverReal)1.;
-			FSolverReal HessianDet = ParticleHessian.Determinant();
-			auto IsIllConditioned = [](FSolverReal HessianDet)
-			{
-				return FMath::Abs(HessianDet) < TMathUtilConstants<FSolverReal>::Epsilon ||
-					FMath::Abs(HessianDet) > TMathUtilConstants<FSolverReal>::MaxReal;
-			};
-			if (IsIllConditioned(HessianDet))
-			{
-				// scale the hessian so that the determinant (~HessianScale^3) falls into the normal range.
-				HessianScale = (FSolverReal)0.;
-				for (int32 RowIdx = 0; RowIdx < 3; ++RowIdx)
-				{
-					for (int32 ColIdx = 0; ColIdx < 3; ++ColIdx)
-					{
-						HessianScale = FMath::Max(FMath::Abs(ParticleHessian.GetAt(RowIdx, ColIdx)), HessianScale);
-					}
-				}
-				if (IsIllConditioned(HessianScale))
-				{
-					return false;
-				}
-				else
-				{
-					ParticleHessian *= 1 / HessianScale;
-					HessianDet = ParticleHessian.Determinant();
-					if (IsIllConditioned(HessianDet))
-					{
-						return false;
-					}
-				}
-			}
-
-			Chaos::PMatrix<FSolverReal, 3, 3> HessianInv = ParticleHessian.SymmetricCofactorMatrix();
-			HessianInv *= FSolverReal(1) / HessianDet;
-			Chaos::TVector<FSolverReal, 3> Dx = HessianInv.GetTransposed() * (-ParticleResidual) / HessianScale; // add back HessianScale
-
-			if (Dx.Size() < MaxDxSize)
-			{
-				Particles.P(p) += Dx;
-			}
-			else
-			{
-				return false;
-			}
-			return true;
-		}
-		
-		void InitializeLambdas()
-		{
-			ComputeInitialResidualAndHessian = [this](const FSolverParticles& Particles, const int32 p, const FSolverReal Dt, TVec3<FSolverReal>& ParticleResidual, Chaos::PMatrix<FSolverReal, 3, 3>& ParticleHessian)
-			{
-				if (!this->bDoQuasistatics) {
-					for (int32 alpha = 0; alpha < 3; alpha++) {
-						ParticleResidual[alpha] = Particles.M(p) * (Particles.P(p)[alpha] - this->xtilde[p][alpha]);
-					}
-					for (int32 alpha = 0; alpha < 3; alpha++) {
-						ParticleHessian.SetAt(alpha, alpha, Particles.M(p));
-					}
-				}
-				else
-				{
-					for (int32 alpha = 0; alpha < 3; alpha++) {
-						ParticleResidual[alpha] = -Dt * Dt * ExternalAcceleration[alpha] * Particles.M(p);
-					}
-				}
-			};
-
-			AccelerationTechniquePerParticle = [this](FSolverParticles& Particles, int32 ParticleIndex)
-			{
-				if (Particles.InvM(ParticleIndex) != FSolverReal(0) && CurrentIt > SORStart)
-				{
-					Particles.P(ParticleIndex) = OmegaSOR * (Particles.P(ParticleIndex) - this->X_k_1[ParticleIndex]) + this->X_k_1[ParticleIndex];
-				}
-				if (Particles.InvM(ParticleIndex) != FSolverReal(0))
-				{
-					this->X_k_1[ParticleIndex] = this->X_k[ParticleIndex];
-					this->X_k[ParticleIndex] = Particles.P(ParticleIndex);
-				}
-			};
-		}
+		bool ApplySingleParticle(const int32 ParticleIdx, const FSolverReal Dt, FSolverParticles& Particles);
+
+		void InitializeLambdas();
+
+		static bool IsClean(const TArray<TArray<int32>>& ConstraintsIn, const TArray<TArray<int32>>& IncidentElementsIn, const TArray<TArray<int32>>& IncidentElementsLocalIn);
 
 
 		//Constraints storage:
@@ -432,12 +144,16 @@ namespace Chaos::Softs
 		TArray<TArray<int32>> DynamicIncidentElements;
 		TArray<TArray<int32>> DynamicIncidentElementsLocal;
 
-		//Lambds for spcifying residual/hessian computations:
-		TFunction<void(const FSolverParticles&, const int32, const FSolverReal, TVec3<FSolverReal>&, Chaos::PMatrix<FSolverReal, 3, 3>&)> ComputeInitialResidualAndHessian;
-		TArray<TFunction<void(const FSolverParticles&, const int32, const int32, const FSolverReal, TVec3<FSolverReal>&, Chaos::PMatrix<FSolverReal, 3, 3>&)>> AddStaticConstraintResidualAndHessian;
-		TArray<TFunction<void(const FSolverParticles&, const int32, const int32, const FSolverReal, TVec3<FSolverReal>&, Chaos::PMatrix<FSolverReal, 3, 3>&)>> AddDynamicConstraintResidualAndHessian;
-		TArray<TFunction<void(const FSolverParticles&, const int32, const int32, const FSolverReal, TVec3<FSolverReal>&, Chaos::PMatrix<FSolverReal, 3, 3>&)>> AddTransientConstraintResidualAndHessian;
-		TArray<TFunction<void(const int32, const FSolverReal, Chaos::PMatrix<FSolverReal, 3, 3>&)>> AddPerNodeHessian;
+		//Lambdas for specifying residual/hessian computations:
+		FGSInitialResidualHessianFunc ComputeInitialResidualAndHessian;
+		TArray<FGSConstraintResidualHessianFunc> AddStaticConstraintResidualAndHessian;
+		TArray<FGSConstraintResidualHessianFunc> AddDynamicConstraintResidualAndHessian;
+		TArray<FGSConstraintResidualHessianFunc> AddTransientConstraintResidualAndHessian;
+		TArray<FGSPerNodeHessianFunc> AddPerNodeHessian;
+		TArray<FGSPerNodeHessianFunc> AddTransientPerNodeHessian;
+
+		TArray<FGSCrossBodyResidualHessianFunc> AddCrossBodyStaticConstraintResidualAndHessian;
+		TArray<FGSCrossBodyResidualHessianFunc> AddCrossBodyTransientConstraintResidualAndHessian;
 
 		//Coloring information:
 		TArray<int32> StaticParticleColors;
@@ -467,13 +183,13 @@ namespace Chaos::Softs
 		FDeformableXPBDCorotatedParams CorotatedParams;
 
 		//Newton solver variables:
-		TArray<TFunction<void(const FSolverParticles&, const TArray<TVec3<FSolverReal>>&, TArray<TVec3<FSolverReal>>&)>> AddInternalForceDifferentials;
+		TArray<FGSForceDifferentialFunc> AddInternalForceDifferentials;
 		TUniquePtr<TArray<int32>> use_list;
 
 		int32 NumTotalParticles;
 		TArray<TVec3<FSolverReal>> ReorderedPs;
 
-		FSolverReal MaxDxSize = FSolverReal(0);
+		FSolverReal MaxDxSize = FSolverReal(UE_MAX_FLT);
 
 		public:
 		bool DebugResidual = false;
diff --git a/Engine/Source/Runtime/Experimental/Chaos/Public/Chaos/Deformable/GaussSeidelWeakConstraints.h b/Engine/Source/Runtime/Experimental/Chaos/Public/Chaos/Deformable/GaussSeidelWeakConstraints.h
--- a/Engine/Source/Runtime/Experimental/Chaos/Public/Chaos/Deformable/GaussSeidelWeakConstraints.h
+++ b/Engine/Source/Runtime/Experimental/Chaos/Public/Chaos/Deformable/GaussSeidelWeakConstraints.h
@@ -15,7 +15,6 @@
 #include "Chaos/Triangle.h"
 #include "Chaos/TriangleCollisionPoint.h"
 #include "Chaos/TriangleMesh.h"
-#include <unordered_map>
 namespace Chaos::Softs
 {
 	using Chaos::TVec3;
@@ -36,238 +35,56 @@ namespace Chaos::Softs
 	class FGaussSeidelSpringConstraintData : public TArrayCollection
 	{
 	public:
-		FGaussSeidelSpringConstraintData()
-		{
-			AddArray(&MIndices);
-			AddArray(&MSecondIndices);
-			AddArray(&MWeights);
-			AddArray(&MSecondWeights);
-			AddArray(&MStiffness);
-			AddArray(&MIsAnisotropic);
-			AddArray(&MNormals);
-			AddArray(&MIsZeroRestLength);
-			AddArray(&MRestLength);
-		}
-
-		//Explicit copy constructor
-		FGaussSeidelSpringConstraintData(const FGaussSeidelSpringConstraintData& Other)
-			: FGaussSeidelSpringConstraintData()
-		{
-			AddConstraints(Other.Size());
-			for (int32 ConstraintIdx = 0; ConstraintIdx < Other.Size(); ++ConstraintIdx)
-			{
-				SetSingleConstraint(Other.GetSingleConstraintData(ConstraintIdx), ConstraintIdx);
-			}
-		}
-
-		FGaussSeidelSpringConstraintData(FGaussSeidelSpringConstraintData&& Other)
-		    : TArrayCollection(), MIndices(MoveTemp(Other.MIndices))
-			, MSecondIndices(MoveTemp(Other.MSecondIndices))
-			, MWeights(MoveTemp(Other.MWeights))
-			, MSecondWeights(MoveTemp(Other.MSecondWeights))
-			, MStiffness(MoveTemp(Other.MStiffness))
-			, MIsAnisotropic(MoveTemp(Other.MIsAnisotropic))
-		{
-			AddArray(&MIndices);
-			AddArray(&MSecondIndices);
-			AddArray(&MWeights);
-			AddArray(&MSecondWeights);
-			AddArray(&MStiffness);
-			AddArray(&MIsAnisotropic);
-			AddArray(&MNormals);
-			AddArray(&MIsZeroRestLength);
-			AddArray(&MRestLength);
-			Other.MSize = 0;
-		}
+		CHAOS_API FGaussSeidelSpringConstraintData();
+		CHAOS_API FGaussSeidelSpringConstraintData(const FGaussSeidelSpringConstraintData& Other);
+		CHAOS_API FGaussSeidelSpringConstraintData(FGaussSeidelSpringConstraintData&& Other);
 
 		virtual ~FGaussSeidelSpringConstraintData() = default;
 
-		void AddConstraints(const int32 Num)
-		{
-			AddElementsHelper(Num);
-		}
-
-		void RemoveConstraint(const int32 Idx)
-		{
-			RemoveAtSwapHelper(Idx);
-		}
-
-		void SetSingleConstraint(const FGaussSeidelSpringConstraintSingleData& SingleData, const int32 ConstraintIndex)
-		{
-			MIndices[ConstraintIndex] = SingleData.SingleIndices;
-			MSecondIndices[ConstraintIndex] = SingleData.SingleSecondIndices;
-			MStiffness[ConstraintIndex] = SingleData.SingleStiffness;
-			MWeights[ConstraintIndex] = SingleData.SingleWeights;
-			MSecondWeights[ConstraintIndex] = SingleData.SingleSecondWeights;
-			MNormals[ConstraintIndex] = SingleData.SingleNormal;
-			MIsAnisotropic[ConstraintIndex] = SingleData.bIsAnisotropic;
-			MIsZeroRestLength[ConstraintIndex] = SingleData.bIsZeroRestLength;
-		}
-
-		void AddSingleConstraint(const FGaussSeidelSpringConstraintSingleData& SingleData)
-		{
-			AddConstraints(1);
-			SetSingleConstraint(SingleData, MSize - 1);
-		}
-
-		int32 Size() const 
-		{
-			return static_cast<int32>(MSize);
-		}
-
-		void Resize(const int32 Num)
-		{
-			ResizeHelper(Num);
-		}
-
-		FGaussSeidelSpringConstraintData& operator=(FGaussSeidelSpringConstraintData&& Other)
-		{
-			MIndices = MoveTemp(Other.MIndices);
-			MSecondIndices = MoveTemp(Other.MSecondIndices);
-			MWeights = MoveTemp(Other.MWeights);
-			MSecondWeights = MoveTemp(Other.MSecondWeights);
-			MStiffness = MoveTemp(Other.MStiffness);
-			MIsAnisotropic = MoveTemp(Other.MIsAnisotropic);
-			MNormals = MoveTemp(Other.MNormals);
-			MIsZeroRestLength = MoveTemp(Other.MIsZeroRestLength);
-			ResizeHelper(Other.Size());
-			Other.MSize = 0;
-			return *this;
-		}
-
-		inline const TArrayCollectionArray<TArray<int32>>& Indices() const
-		{
-			return MIndices;
-		}
-
-		const TArray<int32>& GetIndices(const int32 Index) const
-		{
-			return MIndices[Index];
-		}
-
-		void SetIndices(const int32 Index, const TArray<int32>& InIndices)
-		{
-			MIndices[Index] = InIndices;
-		}
-
-		inline const TArrayCollectionArray<TArray<int32>>& SecondIndices() const
-		{
-			return MSecondIndices;
-		}
-
-		const TArray<int32>& GetSecondIndices(const int32 Index) const
-		{
-			return MSecondIndices[Index];
-		}
-
-		void SetSecondIndices(const int32 Index, const TArray<int32>& InIndices)
-		{
-			MSecondIndices[Index] = InIndices;
-		}
-
-		inline const TArrayCollectionArray<TArray<FSolverReal>>& Weights() const
-		{
-			return MWeights;
-		}
-
-		const TArray<FSolverReal>& GetWeights(const int32 Index) const
-		{
-			return MWeights[Index];
-		}
-
-		void SetWeights(const int32 Index, const TArray<FSolverReal>& InWeights)
-		{
-			MWeights[Index] = InWeights;
-		}
+		void AddConstraints(const int32 Num) { AddElementsHelper(Num); }
+		void RemoveConstraint(const int32 Idx) { RemoveAtSwapHelper(Idx); }
 
-		inline const TArrayCollectionArray<TArray<FSolverReal>>& SecondWeights() const
-		{
-			return MSecondWeights;
-		}
+		void SetSingleConstraint(const FGaussSeidelSpringConstraintSingleData& SingleData, const int32 ConstraintIndex);
+		void AddSingleConstraint(const FGaussSeidelSpringConstraintSingleData& SingleData);
 
-		const TArray<FSolverReal>& GetSecondWeights(const int32 Index) const
-		{
-			return MSecondWeights[Index];
-		}
-
-		void SetSecondWeights(const int32 Index, const TArray<FSolverReal>& InWeights)
-		{
-			MSecondWeights[Index] = InWeights;
-		}
-
-		const bool GetIsAnisotropic(const int32 Index) const
-		{
-			return MIsAnisotropic[Index];
-		}
+		int32 Size() const { return static_cast<int32>(MSize); }
+		void Resize(const int32 Num) { ResizeHelper(Num); }
 
-		void SetIsAnisotropic(const int32 Index, const bool InIsAnisotropic)
-		{
-			MIsAnisotropic[Index] = InIsAnisotropic;
-		}
+		CHAOS_API FGaussSeidelSpringConstraintData& operator=(FGaussSeidelSpringConstraintData&& Other);
 
-		inline const TArrayCollectionArray<TVec3<FSolverReal>>& Normals() const
-		{
-			return MNormals;
-		}
+		inline const TArrayCollectionArray<TArray<int32>>& Indices() const { return MIndices; }
+		const TArray<int32>& GetIndices(const int32 Index) const { return MIndices[Index]; }
+		void SetIndices(const int32 Index, const TArray<int32>& InIndices) { MIndices[Index] = InIndices; }
 
-		const TVec3<FSolverReal>& GetNormal(const int32 Index) const
-		{
-			return MNormals[Index];
-		}
+		inline const TArrayCollectionArray<TArray<int32>>& SecondIndices() const { return MSecondIndices; }
+		const TArray<int32>& GetSecondIndices(const int32 Index) const { return MSecondIndices[Index]; }
+		void SetSecondIndices(const int32 Index, const TArray<int32>& InIndices) { MSecondIndices[Index] = InIndices; }
 
-		void SetNormal(const int32 Index, const TVec3<FSolverReal>& InNormal)
-		{
-			MNormals[Index] = InNormal;
-		}
+		inline const TArrayCollectionArray<TArray<FSolverReal>>& Weights() const { return MWeights; }
+		const TArray<FSolverReal>& GetWeights(const int32 Index) const { return MWeights[Index]; }
+		void SetWeights(const int32 Index, const TArray<FSolverReal>& InWeights) { MWeights[Index] = InWeights; }
 
-		inline const TArrayCollectionArray<FSolverReal>& Stiffness() const
-		{
-			return MStiffness;
-		}
+		inline const TArrayCollectionArray<TArray<FSolverReal>>& SecondWeights() const { return MSecondWeights; }
+		const TArray<FSolverReal>& GetSecondWeights(const int32 Index) const { return MSecondWeights[Index]; }
+		void SetSecondWeights(const int32 Index, const TArray<FSolverReal>& InWeights) { MSecondWeights[Index] = InWeights; }
 
-		FSolverReal GetStiffness(const int32 Index) const
-		{
-			return MStiffness[Index];
-		}
+		const bool GetIsAnisotropic(const int32 Index) const { return MIsAnisotropic[Index]; }
+		void SetIsAnisotropic(const int32 Index, const bool InIsAnisotropic) { MIsAnisotropic[Index] = InIsAnisotropic; }
 
-		void SetStiffness(const int32 Index, const FSolverReal InStiffness)
-		{
-			MStiffness[Index] = InStiffness;
-		}
+		inline const TArrayCollectionArray<TVec3<FSolverReal>>& Normals() const { return MNormals; }
+		const TVec3<FSolverReal>& GetNormal(const int32 Index) const { return MNormals[Index]; }
+		void SetNormal(const int32 Index, const TVec3<FSolverReal>& InNormal) { MNormals[Index] = InNormal; }
 
-		const bool GetIsZeroRestLength(const int32 Index) const
-		{
-			return MIsZeroRestLength[Index];
-		}
+		inline const TArrayCollectionArray<FSolverReal>& Stiffness() const { return MStiffness; }
+		FSolverReal GetStiffness(const int32 Index) const { return MStiffness[Index]; }
+		void SetStiffness(const int32 Index, const FSolverReal InStiffness) { MStiffness[Index] = InStiffness; }
 
-		void SetIsZeroRestLength(const int32 Index, const bool InIsZeroRestLength)
-		{
-			MIsZeroRestLength[Index] = InIsZeroRestLength;
-		}
+		const bool GetIsZeroRestLength(const int32 Index) const { return MIsZeroRestLength[Index]; }
+		void SetIsZeroRestLength(const int32 Index, const bool InIsZeroRestLength) { MIsZeroRestLength[Index] = InIsZeroRestLength; }
 
-		void SetRestLength(const int32 Index, const FSolverReal InRestLength)
-		{
-			MRestLength[Index] = InRestLength;
-		}
+		void SetRestLength(const int32 Index, const FSolverReal InRestLength) { MRestLength[Index] = InRestLength; }
 
-		const FGaussSeidelSpringConstraintSingleData GetSingleConstraintData(const int32 ConstraintIndex) const 
-		{
-			FGaussSeidelSpringConstraintSingleData SingleConstraintData;
-			check(static_cast<uint32>(ConstraintIndex) < MSize);
-			if (ConstraintIndex > INDEX_NONE && static_cast<uint32>(ConstraintIndex) < MSize)
-			{
-				SingleConstraintData.SingleIndices = MIndices[ConstraintIndex];
-				SingleConstraintData.SingleSecondIndices = MSecondIndices[ConstraintIndex];
-				SingleConstraintData.SingleStiffness = MStiffness[ConstraintIndex];
-				SingleConstraintData.SingleWeights = MWeights[ConstraintIndex];
-				SingleConstraintData.SingleSecondWeights = MSecondWeights[ConstraintIndex];
-				SingleConstraintData.bIsAnisotropic = MIsAnisotropic[ConstraintIndex];
-				SingleConstraintData.SingleNormal = MNormals[ConstraintIndex];
-				SingleConstraintData.bIsZeroRestLength = MIsZeroRestLength[ConstraintIndex];
-				SingleConstraintData.RestLength = MRestLength[ConstraintIndex];
-			}
-			return SingleConstraintData;
-		}
+		CHAOS_API const FGaussSeidelSpringConstraintSingleData GetSingleConstraintData(const int32 ConstraintIndex) const;
 
 	private:
 		TArrayCollectionArray<TArray<int32>> MIndices;
@@ -284,75 +101,28 @@ namespace Chaos::Softs
 	// Spring constraint for position-based nonlinear Gauss Seidel solver [Chen et al., 2024]
 	struct FGaussSeidelSpringConstraints
 	{
+	public:
 		//TODO(Yizhou): Add unittest for Gauss Seidel Weak Constraints
-		FGaussSeidelSpringConstraints(
+		CHAOS_API FGaussSeidelSpringConstraints(
 			const TArray<TArray<int32>>& InIndices,
 			const TArray<TArray<FSolverReal>>& InWeights,
 			const TArray<FSolverReal>& InStiffness,
 			const TArray<TArray<int32>>& InSecondIndices,
 			const TArray<TArray<FSolverReal>>& InSecondWeights,
-			const TArray<bool>& InIsAnisotrpic,
+			const TArray<bool>& InIsAnisotropic,
 			const TArray<bool>& InZeroRestLength,
 			const FDeformableXPBDWeightedSpringConstraintParams& InParams
-		) : DebugDrawParams(InParams)
-		{
-			ensureMsgf(InIndices.Num() == InSecondIndices.Num(), TEXT("Input Double Bindings have wrong size"));
-
-			ConstraintsData.AddConstraints(InIndices.Num());
-
-			for (int32 SpringIdx = 0; SpringIdx < InIndices.Num(); ++SpringIdx)
-			{
-				FGaussSeidelSpringConstraintSingleData SingleConstraintData;
-				SingleConstraintData.SingleIndices = InIndices[SpringIdx];
-				SingleConstraintData.SingleSecondIndices = InSecondIndices[SpringIdx];
-				SingleConstraintData.SingleWeights = InWeights[SpringIdx];
-				SingleConstraintData.SingleSecondWeights = InSecondWeights[SpringIdx];
-				SingleConstraintData.SingleStiffness = InStiffness[SpringIdx];
-				SingleConstraintData.bIsAnisotropic = InIsAnisotrpic[SpringIdx];
-				SingleConstraintData.bIsZeroRestLength = InZeroRestLength[SpringIdx];
-				ConstraintsData.SetSingleConstraint(SingleConstraintData, SpringIdx);
-				for (int32 SecondLocalIdx = 0; SecondLocalIdx < InSecondIndices[SpringIdx].Num(); ++SecondLocalIdx)
-				{
-					ensureMsgf(!InIndices[SpringIdx].Contains(InSecondIndices[SpringIdx][SecondLocalIdx]), TEXT("Indices and Second Indices overlaps. Currently not supported"));
-				}
-			}
-		}
+		);
 
 		UE_DEPRECATED(5.8, "Replaced with constructor with IsAnisotropic and IsZeroRestLength.")
-			FGaussSeidelSpringConstraints(
-				const TArray<TArray<int32>>& InIndices,
-				const TArray<TArray<FSolverReal>>& InWeights,
-				const TArray<FSolverReal>& InStiffness,
-				const TArray<TArray<int32>>& InSecondIndices,
-				const TArray<TArray<FSolverReal>>& InSecondWeights,
-				const FDeformableXPBDWeightedSpringConstraintParams& InParams
-			) : DebugDrawParams(InParams)
-		{
-			ensureMsgf(InIndices.Num() == InSecondIndices.Num(), TEXT("Input Double Bindings have wrong size"));
-
-			ConstraintsData.AddConstraints(InIndices.Num());
-
-			for (int32 i = 0; i < InIndices.Num(); i++)
-			{
-				FGaussSeidelSpringConstraintSingleData SingleConstraintData;
-				SingleConstraintData.SingleIndices = InIndices[i];
-				SingleConstraintData.SingleSecondIndices = InSecondIndices[i];
-				SingleConstraintData.SingleWeights = InWeights[i];
-				SingleConstraintData.SingleSecondWeights = InSecondWeights[i];
-				SingleConstraintData.SingleStiffness = InStiffness[i];
-				ConstraintsData.SetSingleConstraint(SingleConstraintData, i);
-			}
-
-			for (int32 i = 0; i < ConstraintsData.Size(); i++)
-			{
-				const TArray<int32>& SingleIndices = ConstraintsData.GetIndices(i);
-				const TArray<int32>& SingleSecondIndices = ConstraintsData.GetSecondIndices(i);
-				for (int32 j = 0; j < SingleSecondIndices.Num(); j++)
-				{
-					ensureMsgf(!SingleIndices.Contains(SingleSecondIndices[j]), TEXT("Indices and Second Indices overlaps. Currently not supported"));
-				}
-			}
-		}
+		CHAOS_API FGaussSeidelSpringConstraints(
+			const TArray<TArray<int32>>& InIndices,
+			const TArray<TArray<FSolverReal>>& InWeights,
+			const TArray<FSolverReal>& InStiffness,
+			const TArray<TArray<int32>>& InSecondIndices,
+			const TArray<TArray<FSolverReal>>& InSecondWeights,
+			const FDeformableXPBDWeightedSpringConstraintParams& InParams
+		);
 
 		struct FGaussSeidelConstraintHandle
 		{
@@ -361,765 +131,98 @@ namespace Chaos::Softs
 
 		virtual ~FGaussSeidelSpringConstraints() = default;
 
-		void ComputeInitialWCData(const FSolverParticlesRange& InParticles)
-		{
-			TArray<TArray<int32>> ExtraConstraints;
-			ExtraConstraints.Init(TArray<int32>(), ConstraintsData.Size());
-			for (int32 ConstraintIdx = 0; ConstraintIdx < ExtraConstraints.Num(); ++ConstraintIdx)
-			{
-				ExtraConstraints[ConstraintIdx].SetNum(ConstraintsData.GetIndices(ConstraintIdx).Num() + ConstraintsData.GetSecondIndices(ConstraintIdx).Num());
-				for (int32 LocalIdx = 0; LocalIdx < ConstraintsData.GetIndices(ConstraintIdx).Num(); ++LocalIdx)
-				{
-					ExtraConstraints[ConstraintIdx][LocalIdx] = ConstraintsData.GetIndices(ConstraintIdx)[LocalIdx];
-				}
-				for (int32 LocalSecondIdx = 0; LocalSecondIdx < ConstraintsData.GetSecondIndices(ConstraintIdx).Num(); ++LocalSecondIdx)
-				{
-					ExtraConstraints[ConstraintIdx][LocalSecondIdx + ConstraintsData.GetIndices(ConstraintIdx).Num()] = ConstraintsData.GetSecondIndices(ConstraintIdx)[LocalSecondIdx];
-				}
-			}
-			WCIncidentElements = Chaos::Utilities::ComputeIncidentElements(ExtraConstraints, &WCIncidentElementsLocal);
+		CHAOS_API void ComputeInitialWCData(const FSolverParticlesRange& InParticles);
+		CHAOS_API void ComputeInitialWCData(const FSolverParticles& InParticles);
 
-			//Update rest state normal and nodal weights
-			UpdateTriangleNormalAndNodalWeight(InParticles, /*bUseParticleX = */true);
-
-			NoCollisionNodalWeights = NodalWeights;
-			NoCollisionConstraints = ExtraConstraints;
-			InitialWCSize = ConstraintsData.Size();
-
-			NoCollisionWCIncidentElements = WCIncidentElements;
-			NoCollisionWCIncidentElementsLocal = WCIncidentElementsLocal;
-
-			// Compute rest length
-			for (int32 ConstraintIdx = 0; ConstraintIdx < ConstraintsData.Size(); ++ConstraintIdx)
-			{
-				if (ConstraintsData.GetIsZeroRestLength(ConstraintIdx))
-				{
-					ConstraintsData.SetRestLength(ConstraintIdx, 0);
-				}
-				else
-				{
-					const TArray<int32>& Indices = ConstraintsData.GetIndices(ConstraintIdx);
-					const TArray<int32>& SecondIndices = ConstraintsData.GetSecondIndices(ConstraintIdx);
-					const TArray<FSolverReal>& Weights = ConstraintsData.GetWeights(ConstraintIdx);
-					const TArray<FSolverReal>& SecondWeights = ConstraintsData.GetSecondWeights(ConstraintIdx);
-					const TVec3<FSolverReal> RestSpringEdge = ComputeSpringEdge(InParticles, Indices, SecondIndices, Weights, SecondWeights, /*bUseParticleX =*/true);
-					if (ConstraintsData.GetIsAnisotropic(ConstraintIdx))
-					{
-						//if the spring is anisotropic, rest length could be negative depending on the normal direction
-						ConstraintsData.SetRestLength(ConstraintIdx, TVec3<FSolverReal>::DotProduct(ConstraintsData.GetNormal(ConstraintIdx), RestSpringEdge));
-					}
-					else
-					{
-						ConstraintsData.SetRestLength(ConstraintIdx, RestSpringEdge.Size());
-					}
-				}
-			}
-		}
-
-		void ComputeInitialWCData(const FSolverParticles& InParticles)
-		{
-			ComputeInitialWCData(FSolverParticlesRange(const_cast<FSolverParticles*>(&InParticles), /*InOffset = */ 0, /*InRange = */ InParticles.Size()));
-		}
-
-		TVec3<FSolverReal> ComputeSpringEdge(const FSolverParticlesRange& InParticles, const TArray<int32>& LocalIndices,
-			const TArray<int32>& LocalSecondIndices, const TArray<FSolverReal>& Weight, const TArray<FSolverReal>& SecondWeight, bool bUseParticleX) const
-		{
-			TVec3<FSolverReal> SpringEdge((FSolverReal)0.);
-			if (ensure(LocalIndices.Num() == Weight.Num() && LocalSecondIndices.Num() == SecondWeight.Num()))
-			{
-				if (bUseParticleX)
-				{
-					for (int32 Idx = 0; Idx < LocalIndices.Num(); ++Idx)
-					{
-						for (int32 beta = 0; beta < 3; beta++)
-						{
-							SpringEdge[beta] += Weight[Idx] * InParticles.X(LocalIndices[Idx])[beta];
-						}
-					}
-					for (int32 SecondIdx = 0; SecondIdx < LocalSecondIndices.Num(); ++SecondIdx)
-					{
-						for (int32 beta = 0; beta < 3; beta++)
-						{
-							SpringEdge[beta] -= SecondWeight[SecondIdx] * InParticles.X(LocalSecondIndices[SecondIdx])[beta];
-						}
-					}
-				}
-				else
-				{
-					for (int32 Idx = 0; Idx < LocalIndices.Num(); ++Idx)
-					{
-						for (int32 beta = 0; beta < 3; beta++)
-						{
-							SpringEdge[beta] += Weight[Idx] * InParticles.P(LocalIndices[Idx])[beta];
-						}
-					}
-					for (int32 SecondIdx = 0; SecondIdx < LocalSecondIndices.Num(); ++SecondIdx)
-					{
-						for (int32 beta = 0; beta < 3; beta++)
-						{
-							SpringEdge[beta] -= SecondWeight[SecondIdx] * InParticles.P(LocalSecondIndices[SecondIdx])[beta];
-						}
-					}
-				}
-			}
-			return SpringEdge;
-		}
-
-		TVec3<FSolverReal> ComputeSpringEdge(const FSolverParticles& InParticles, const TArray<int32>& LocalIndices,
-			const TArray<int32>& LocalSecondIndices, const TArray<FSolverReal>& Weight, const TArray<FSolverReal>& SecondWeight, bool bUseParticleX) const
-		{
-			return ComputeSpringEdge(FSolverParticlesRange(const_cast<FSolverParticles*>(&InParticles), /*InOffset = */ 0, /*InRange = */ InParticles.Size()), 
-				LocalIndices, LocalSecondIndices, Weight, SecondWeight, bUseParticleX);
-		}
-
-		void AddWCHessian(const int32 p, const FSolverReal Dt, Chaos::PMatrix<FSolverReal, 3, 3>& ParticleHessian) const
-		{
-			if (NodalWeights[p].Num() > 0)
-			{
-				for (int32 alpha = 0; alpha < 3; alpha++)
-				{
-					ParticleHessian.SetAt(alpha, alpha, ParticleHessian.GetAt(alpha, alpha) + Dt * Dt * NodalWeights[p][alpha]);
-				}
-
-				ParticleHessian.SetAt(0, 1, ParticleHessian.GetAt(0, 1) + Dt * Dt * NodalWeights[p][3]);
-				ParticleHessian.SetAt(0, 2, ParticleHessian.GetAt(0, 2) + Dt * Dt * NodalWeights[p][4]);
-				ParticleHessian.SetAt(1, 2, ParticleHessian.GetAt(1, 2) + Dt * Dt * NodalWeights[p][5]);
-				ParticleHessian.SetAt(1, 0, ParticleHessian.GetAt(1, 0) + Dt * Dt * NodalWeights[p][3]);
-				ParticleHessian.SetAt(2, 0, ParticleHessian.GetAt(2, 0) + Dt * Dt * NodalWeights[p][4]);
-				ParticleHessian.SetAt(2, 1, ParticleHessian.GetAt(2, 1) + Dt * Dt * NodalWeights[p][5]);
-				//TODO(Yizhou): Clean up the following after debugging:
-				//ParticleHessian.SetAt(0, 2) += Dt * Dt * NodalWeights[p][4];
-				//ParticleHessian.SetAt(1, 2) += Dt * Dt * NodalWeights[p][5];
-				//ParticleHessian.SetAt(1, 0) += Dt * Dt * NodalWeights[p][3];
-				//ParticleHessian.SetAt(2, 0) += Dt * Dt * NodalWeights[p][4];
-				//ParticleHessian.SetAt(2, 1) += Dt * Dt * NodalWeights[p][5];
-			}
-		}
+		CHAOS_API void AddWCHessian(const int32 ParticleIdx, const FSolverReal Dt, Chaos::PMatrix<FSolverReal, 3, 3>& ParticleHessian) const;
 
 		void AddExtraConstraints(const TArray<TArray<int32>>& InIndices,
 								const TArray<TArray<FSolverReal>>& InWeights,
 								const TArray<FSolverReal>& InStiffness,
 								const TArray<TArray<int32>>& InSecondIndices,
 								const TArray<TArray<FSolverReal>>& InSecondWeights,
-								const TArray<bool>& InIsAnisotrpic,
-								const TArray<bool>& InIsZeroRestLength)
-		{
-			const int32 Offset = ConstraintsData.Size();
-
-			ConstraintsData.AddConstraints(InIndices.Num());
-
-			for (int32 i = 0; i < InIndices.Num(); i++)
-			{
-				FGaussSeidelSpringConstraintSingleData SingleConstraintData;
-				SingleConstraintData.SingleIndices = InIndices[i];
-				SingleConstraintData.SingleSecondIndices = InSecondIndices[i];
-				SingleConstraintData.SingleWeights = InWeights[i];
-				SingleConstraintData.SingleSecondWeights = InSecondWeights[i];
-				SingleConstraintData.SingleStiffness = InStiffness[i];
-				SingleConstraintData.bIsAnisotropic = InIsAnisotrpic[i];
-				SingleConstraintData.bIsZeroRestLength = InIsZeroRestLength[i];
-				ConstraintsData.SetSingleConstraint(SingleConstraintData, i + Offset);
-			}
-		}
-
-		void Resize(int32 Size)
-		{
-			ConstraintsData.Resize(Size);
-		}
-
-		void UpdatePointTriangleCollisionWCData(const FSolverParticles& Particles)
-		{	
-			FGaussSeidelSpringConstraintData OriginalConstraintsData = ConstraintsData;
-
-			ConstraintsData.Resize(InitialWCSize);
+								const TArray<bool>& InIsAnisotropic,
+								const TArray<bool>& InIsZeroRestLength);
 
-			for (int32 i = InitialWCSize; i < OriginalConstraintsData.Size(); i++)
-			{
-				const TArray<int32>& IndicesTemp = OriginalConstraintsData.GetIndices(i);
-				const TArray<int32>& SecondIndicesTemp = OriginalConstraintsData.GetSecondIndices(i);
-				ensureMsgf(OriginalConstraintsData.GetIndices(i).Num() == 3, TEXT("Collision format is not point-triangle"));
-				ensureMsgf(OriginalConstraintsData.GetSecondIndices(i).Num() == 1, TEXT("Collision format is not point-triangle"));
-				Chaos::TVector<float, 3> TriPos0(Particles.P(IndicesTemp[0])), TriPos1(Particles.P(IndicesTemp[1])), TriPos2(Particles.P(IndicesTemp[2])), ParticlePos(Particles.P(SecondIndicesTemp[0]));
-				Chaos::TVector<FSolverReal, 3> Normal = FVector3f::CrossProduct(TriPos2 - TriPos0, TriPos1 - TriPos0); //triangle normal convention (see FTriangleMesh::GetFaceNormals())
-				if (FVector3f::DotProduct(ParticlePos - TriPos0, Normal) < 0.f) //not resolved, keep the spring
-				{
-					ConstraintsData.AddSingleConstraint(OriginalConstraintsData.GetSingleConstraintData(i));
-				}
-			}
-		}
-
-		void VisualizeAllBindings(const FSolverParticlesRange& InParticles, const FSolverReal Dt) const
-		{
-#if WITH_EDITOR
-			auto DoubleVert = [](Chaos::TVec3<FSolverReal> V) { return FVector3d(V.X, V.Y, V.Z); };
-			for (int32 i = 0; i < ConstraintsData.Size(); i++)
-			{
-				const FGaussSeidelSpringConstraintSingleData& SingleConstraintData = ConstraintsData.GetSingleConstraintData(i);
-				Chaos::TVec3<FSolverReal> SourcePos((FSolverReal)0.), TargetPos((FSolverReal)0.);
-				for (int32 j = 0; j < SingleConstraintData.SingleIndices.Num(); j++)
-				{
-					SourcePos += SingleConstraintData.SingleWeights[j] * InParticles.P(SingleConstraintData.SingleIndices[j]);
-				}
-				for (int32 j = 0; j < SingleConstraintData.SingleSecondIndices.Num(); j++)
-				{
-					TargetPos += SingleConstraintData.SingleSecondWeights[j] * InParticles.P(SingleConstraintData.SingleSecondIndices[j]);
-				}
-
-				float ParticleThickness = DebugDrawParams.DebugParticleWidth;
-				float LineThickness = DebugDrawParams.DebugLineWidth;
-
-				if (SingleConstraintData.SingleIndices.Num() == 1)
-				{
-					Chaos::FDebugDrawQueue::GetInstance().DrawDebugPoint(DoubleVert(SourcePos), FColor::Red, false, Dt, 0, ParticleThickness);
-					for (int32 j = 0; j < SingleConstraintData.SingleSecondIndices.Num(); j++)
-					{
-						Chaos::FDebugDrawQueue::GetInstance().DrawDebugPoint(DoubleVert(InParticles.P(SingleConstraintData.SingleSecondIndices[j])), FColor::Green, false, Dt, 0, ParticleThickness);
-						Chaos::FDebugDrawQueue::GetInstance().DrawDebugLine(DoubleVert(InParticles.P(SingleConstraintData.SingleSecondIndices[j])), DoubleVert(InParticles.P(SingleConstraintData.SingleSecondIndices[(j + 1) % SingleConstraintData.SingleSecondIndices.Num()])), FColor::Green, false, Dt, 0, LineThickness);
-					}
-
-				}
-
-				if (SingleConstraintData.SingleSecondIndices.Num() == 1)
-				{
-					Chaos::FDebugDrawQueue::GetInstance().DrawDebugPoint(DoubleVert(TargetPos), FColor::Red, false, Dt, 0, ParticleThickness);
-					for (int32 j = 0; j < SingleConstraintData.SingleIndices.Num(); j++)
-					{
-						Chaos::FDebugDrawQueue::GetInstance().DrawDebugPoint(DoubleVert(InParticles.P(SingleConstraintData.SingleIndices[j])), FColor::Green, false, Dt, 0, ParticleThickness);
-						Chaos::FDebugDrawQueue::GetInstance().DrawDebugLine(DoubleVert(InParticles.P(SingleConstraintData.SingleIndices[j])), DoubleVert(InParticles.P(SingleConstraintData.SingleIndices[(j + 1) % SingleConstraintData.SingleIndices.Num()])), FColor::Green, false, Dt, 0, LineThickness);
-					}
-				}
-
-				Chaos::FDebugDrawQueue::GetInstance().DrawDebugLine(DoubleVert(SourcePos), DoubleVert(TargetPos), FColor::Yellow, false, Dt, 0, LineThickness);
-			}
-#endif
-		}
-
-		void Init(const FSolverParticlesRange& InParticles, const FSolverReal Dt)
-		{
-			UpdateTriangleNormalAndNodalWeight(InParticles, /*bUseParticleX =*/ false);
-			if (DebugDrawParams.bVisualizeBindings)
-			{
-				VisualizeAllBindings(InParticles, Dt);
-			}
-		}
-
-		void Init(const FSolverParticles& InParticles, const FSolverReal Dt)
-		{
-			Init(FSolverParticlesRange(const_cast<FSolverParticles*>(&InParticles), /*InOffset = */ 0, /*InRange = */ InParticles.Size()), Dt);
-		}
-
-		void UpdateTriangleNormalAndNodalWeight(const FSolverParticlesRange& InParticles, bool bUseParticleX)
-		{
-			for (int32 i = 0; i < ConstraintsData.Size(); i++)
-			{
-				if (ConstraintsData.GetIsAnisotropic(i))
-				{
-					const TArray<int32>& IndicesTemp = ConstraintsData.GetIndices(i);
-					const TArray<int32>& SecondIndicesTemp = ConstraintsData.GetSecondIndices(i);
-					ensureMsgf(ConstraintsData.GetIndices(i).Num() == 3, TEXT("Collision format is not point-triangle"));
-					ensureMsgf(ConstraintsData.GetSecondIndices(i).Num() == 1, TEXT("Collision format is not point-triangle"));
-					Chaos::TVector<FSolverReal, 3> Normal;
-					if (bUseParticleX)
-					{
-						Chaos::TVector<float, 3> TriPos0(InParticles.X(IndicesTemp[0])), TriPos1(InParticles.X(IndicesTemp[1])), TriPos2(InParticles.X(IndicesTemp[2]));
-						Normal = FVector3f::CrossProduct(TriPos2 - TriPos0, TriPos1 - TriPos0).GetSafeNormal(); //triangle normal convention (see FTriangleMesh::GetFaceNormals())
-					}
-					else
-					{
-						Chaos::TVector<float, 3> TriPos0(InParticles.P(IndicesTemp[0])), TriPos1(InParticles.P(IndicesTemp[1])), TriPos2(InParticles.P(IndicesTemp[2]));
-						Normal = FVector3f::CrossProduct(TriPos2 - TriPos0, TriPos1 - TriPos0).GetSafeNormal(); //triangle normal convention (see FTriangleMesh::GetFaceNormals())
-					}
-					ConstraintsData.SetNormal(i, Normal);
-				}
-			}
-
-			NodalWeights.Init({}, InParticles.Size());
-			for (int32 p = 0; p < WCIncidentElements.Num(); ++p)
-			{
-				if (WCIncidentElements[p].Num() > 0)
-				{
-					NodalWeights[p].Init(FSolverReal(0), 6);
-					for (int32 j = 0; j < WCIncidentElements[p].Num(); j++)
-					{
-						int32 ConstraintIndex = WCIncidentElements[p][j];
-						int32 LocalIndex = WCIncidentElementsLocal[p][j];
-
-						FSolverReal Weight = FSolverReal(0);
-						if (LocalIndex >= ConstraintsData.GetIndices(ConstraintIndex).Num())
-						{
-							Weight = ConstraintsData.GetSecondWeights(ConstraintIndex)[LocalIndex - ConstraintsData.GetIndices(ConstraintIndex).Num()];
-						}
-						else
-						{
-							Weight = ConstraintsData.GetWeights(ConstraintIndex)[LocalIndex];
-						}
-
-						if (ConstraintsData.GetIsAnisotropic(ConstraintIndex))
-						{
-							for (int32 alpha = 0; alpha < 3; alpha++)
-							{
-								NodalWeights[p][alpha] += ConstraintsData.GetNormal(ConstraintIndex)[alpha] * ConstraintsData.GetNormal(ConstraintIndex)[alpha] * Weight * Weight * ConstraintsData.GetStiffness(ConstraintIndex);
-							}
-
-							NodalWeights[p][3] += ConstraintsData.GetNormal(ConstraintIndex)[0] * ConstraintsData.GetNormal(ConstraintIndex)[1] * Weight * Weight * ConstraintsData.GetStiffness(ConstraintIndex);
-							NodalWeights[p][4] += ConstraintsData.GetNormal(ConstraintIndex)[0] * ConstraintsData.GetNormal(ConstraintIndex)[2] * Weight * Weight * ConstraintsData.GetStiffness(ConstraintIndex);
-							NodalWeights[p][5] += ConstraintsData.GetNormal(ConstraintIndex)[1] * ConstraintsData.GetNormal(ConstraintIndex)[2] * Weight * Weight * ConstraintsData.GetStiffness(ConstraintIndex);
-						}
-						else
-						{
-							for (int32 alpha = 0; alpha < 3; alpha++)
-							{
-								NodalWeights[p][alpha] += Weight * Weight * ConstraintsData.GetStiffness(ConstraintIndex);
-							}
-						}
-					}
-				}
-			}
-		}
+		CHAOS_API void Init(const FSolverParticlesRange& InParticles, const FSolverReal Dt);
+		CHAOS_API void Init(const FSolverParticles& InParticles, const FSolverReal Dt);
 
 		//CollisionDetectionSpatialHash should be faster than CollisionDetectionBVH
-		void CollisionDetectionBVH(const FSolverParticles& Particles, const TArray<TVec3<int32>>& SurfaceElements, const TArray<int32>& ComponentIndex, float DetectRadius = 1.f, float PositionTargetStiffness = 10000.f, bool UseAnisotropicSpring = true)
-		{
-			TRACE_CPUPROFILER_EVENT_SCOPE(STAT_ChaosGaussSeidelWeakConstraintsCollisionDetection);
-			Resize(InitialWCSize);
-
-			TArray<Chaos::TVector<int32, 3>> SurfaceElementsArray;
-			for (int32 i = 0; i < SurfaceElements.Num(); i++)
-			{
-				Chaos::TVector<int32, 3> CurrentSurfaceElements(0);
-				for (int32 j = 0; j < 3; j++)
-				{
-					CurrentSurfaceElements[j] = SurfaceElements[i][j];
-				}
-				if (CurrentSurfaceElements[0] != INDEX_NONE
-					&& CurrentSurfaceElements[1] != INDEX_NONE
-					&& CurrentSurfaceElements[2] != INDEX_NONE)
-				{
-					SurfaceElementsArray.Emplace(CurrentSurfaceElements);
-				}
-			}
-			TArray<TArray<int32>> LocalIndex;
-			TArray<TArray<int32>>* LocalIndexPtr = &LocalIndex;
-			TArray<TArray<int>> GlobalIndex = Chaos::Utilities::ComputeIncidentElements(SurfaceElementsArray, LocalIndexPtr);
-			int32 ActualParticleCount = 0;
-			for (int32 l = 0; l < GlobalIndex.Num(); l++)
-			{
-				if (GlobalIndex[l].Num() > 0)
-				{
-					ActualParticleCount += 1;
-				}
-			}
-			TArray<Chaos::TVector<float, 3>> SurfaceElementsPositions;
-			SurfaceElementsPositions.SetNum(ActualParticleCount);
-			TArray<int32> SurfaceElementsMap;
-			SurfaceElementsMap.SetNum(ActualParticleCount);
-			int32 CurrentParticleIndex = 0;
-			for (int32 i = 0; i < GlobalIndex.Num(); i++)
-			{
-				if (GlobalIndex[i].Num() > 0)
-				{
-					SurfaceElementsPositions[CurrentParticleIndex] = Particles.P(SurfaceElements[GlobalIndex[i][0]][LocalIndex[i][0]]);
-					SurfaceElementsMap[CurrentParticleIndex] = SurfaceElements[GlobalIndex[i][0]][LocalIndex[i][0]];
-					CurrentParticleIndex += 1;
-				}
-			}
-
-			TArray<Chaos::FSphere*> VertexSpherePtrs;
-			TArray<Chaos::FSphere> VertexSpheres;
-
-			VertexSpheres.Init(Chaos::FSphere(Chaos::TVec3<Chaos::FReal>(0), DetectRadius), SurfaceElementsPositions.Num());
-			VertexSpherePtrs.SetNum(SurfaceElementsPositions.Num());
-
-			for (int32 i = 0; i < SurfaceElementsPositions.Num(); i++)
-			{
-				Chaos::TVec3<Chaos::FReal> SphereCenter(SurfaceElementsPositions[i]);
-				Chaos::FSphere VertexSphere(SphereCenter, DetectRadius);
-				VertexSpheres[i] = Chaos::FSphere(SphereCenter, DetectRadius);
-				VertexSpherePtrs[i] = &VertexSpheres[i];
-			}
-			Chaos::TBoundingVolumeHierarchy<
-				TArray<Chaos::FSphere*>,
-				TArray<int32>,
-				Chaos::FReal,
-				3> VertexBVH(VertexSpherePtrs);
-
-			for (int32 i = 0; i < SurfaceElements.Num(); i++)
-			{
-				TArray<int32> TriangleIntersections0 = VertexBVH.FindAllIntersections(Particles.P(SurfaceElements[i][0]));
-				TArray<int32> TriangleIntersections1 = VertexBVH.FindAllIntersections(Particles.P(SurfaceElements[i][1]));
-				TArray<int32> TriangleIntersections2 = VertexBVH.FindAllIntersections(Particles.P(SurfaceElements[i][2]));
-				TriangleIntersections0.Sort();
-				TriangleIntersections1.Sort();
-				TriangleIntersections2.Sort();
-
-				TArray<int32> TriangleIntersections({});
-				for (int32 k = 0; k < TriangleIntersections0.Num(); k++)
-				{
-					if (TriangleIntersections1.Contains(TriangleIntersections0[k])
-						&& TriangleIntersections2.Contains(TriangleIntersections0[k]))
-					{
-						TriangleIntersections.Emplace(TriangleIntersections0[k]);
-					}
-				}
-
-				int32 TriangleIndex = ComponentIndex[SurfaceElements[i][0]];
-				int32 MinIndex = INDEX_NONE;
-				float MinDis = DetectRadius;
-				Chaos::TVector<float, 3> ClosestBary(0.f);
-				Chaos::TVector<float, 3> FaceNormal;
-				for (int32 j = 0; j < TriangleIntersections.Num(); j++)
-				{
-					if (ComponentIndex[SurfaceElementsMap[TriangleIntersections[j]]] >= 0 && TriangleIndex >= 0 && ComponentIndex[SurfaceElementsMap[TriangleIntersections[j]]] != TriangleIndex)
-					{
-						Chaos::TVector<float, 3> Bary, TriPos0(Particles.P(SurfaceElements[i][0])), TriPos1(Particles.P(SurfaceElements[i][1])), TriPos2(Particles.P(SurfaceElements[i][2])), ParticlePos(Particles.P(SurfaceElementsMap[TriangleIntersections[j]]));
-						Chaos::TVector<Chaos::FRealSingle, 3> ClosestPoint = Chaos::FindClosestPointAndBaryOnTriangle(TriPos0, TriPos1, TriPos2, ParticlePos, Bary);
-						Chaos::FRealSingle CurrentDistance = (Particles.P(SurfaceElementsMap[TriangleIntersections[j]]) - ClosestPoint).Size();
-						if (CurrentDistance < MinDis)
-						{
-							Chaos::TVector<FSolverReal, 3> Normal = FVector3f::CrossProduct(TriPos2 - TriPos0, TriPos1 - TriPos0); //The normal needs to point outwards of the geometry
-							if (FVector3f::DotProduct(ParticlePos - TriPos0, Normal) < 0.f)
-							{
-								Normal.SafeNormalize(1e-8f);
-								MinDis = CurrentDistance;
-								MinIndex = SurfaceElementsMap[TriangleIntersections[j]];
-								ClosestBary = Bary;
-								FaceNormal = Normal;
-							}
-						}
-
-					}
-				}
-				if (MinIndex != INDEX_NONE
-					&& MinIndex != SurfaceElements[i][0]
-					&& MinIndex != SurfaceElements[i][1]
-					&& MinIndex != SurfaceElements[i][2])
-				{
-					FGaussSeidelSpringConstraintSingleData SingleConstraintData;
-					SingleConstraintData.SingleIndices = {SurfaceElements[i][0], SurfaceElements[i][1], SurfaceElements[i][2]};
-					SingleConstraintData.SingleSecondIndices = {MinIndex};
-					SingleConstraintData.SingleWeights = {ClosestBary[0], ClosestBary[1], ClosestBary[2]};
-					SingleConstraintData.SingleSecondWeights = {FSolverReal(1.f)};
-					SingleConstraintData.bIsAnisotropic = UseAnisotropicSpring;
-					SingleConstraintData.SingleNormal = FaceNormal;
-					SingleConstraintData.bIsZeroRestLength = true; //Push-out type collision springs should be zero rest length
-
-					float SpringStiffness = 0.f;
-					for (int32 k = 0; k < 3; k++)
-					{
-						SpringStiffness += ClosestBary[k] * PositionTargetStiffness * Particles.M(SurfaceElements[i][k]);
-					}
-					SpringStiffness += PositionTargetStiffness * Particles.M(MinIndex);
-					SingleConstraintData.SingleStiffness = (FSolverReal)SpringStiffness;
-					ConstraintsData.AddSingleConstraint(SingleConstraintData);
-				}
-			}
-		}
+		void CollisionDetectionBVH(const FSolverParticles& Particles, const TArray<TVec3<int32>>& SurfaceElements, const TArray<int32>& ComponentIndex, float DetectRadius = 1.f, float PositionTargetStiffness = 10000.f, bool UseAnisotropicSpring = true);
 
 		template<typename SpatialAccelerator>
-		void CollisionDetectionSpatialHash(const FSolverParticles& Particles, const TArray<int32>& SurfaceVertices, const FTriangleMesh& TriangleMesh, const TArray<int32>& ComponentIndex, const SpatialAccelerator& Spatial, float DetectRadius = 1.f, float PositionTargetStiffness = 10000.f, bool UseAnisotropicSpring = true)
-		{
-			TRACE_CPUPROFILER_EVENT_SCOPE(STAT_ChaosGaussSeidelWeakConstraintsCollisionDetectionSpatialHash);
-			Resize(InitialWCSize + Particles.Size());
-			std::atomic<int32> ConstraintIndex(InitialWCSize);
-			const TArray<TVec3<int32>>& Elements = TriangleMesh.GetSurfaceElements();
-			const float HalfRadius = DetectRadius / 2;
-			PhysicsParallelFor(SurfaceVertices.Num(),
-				[this, &Spatial, &Particles, &SurfaceVertices, &ConstraintIndex, &TriangleMesh, &Elements, &HalfRadius, &ComponentIndex, &PositionTargetStiffness, &UseAnisotropicSpring](int32 i)
-				{
-					const int32 Index = SurfaceVertices[i];
-					TArray< TTriangleCollisionPoint<FSolverReal>> Result;
-					//PointProximityQuery
-					if (TriangleMesh.PointClosestTriangleQuery(Spatial, static_cast<const TArrayView<const FSolverVec3>&>(Particles.XArray()), Index, Particles.GetX(Index), HalfRadius, HalfRadius,
-						[this, &ComponentIndex, &Elements](const int32 PointIndex, const int32 TriangleIndex)->bool
-						{
-							//Skip particles that are bound in initial springs
-							return ComponentIndex[PointIndex] != ComponentIndex[Elements[TriangleIndex][0]] && (!NoCollisionWCIncidentElements.IsValidIndex(PointIndex) || NoCollisionWCIncidentElements[PointIndex].Num() == 0);
-						},
-						Result))
-					{
-						for (const TTriangleCollisionPoint<FSolverReal>& CollisionPoint : Result)
-						{
-							if (CollisionPoint.Phi < 0)
-							{
-								const TVector<int32, 3>& Elem = Elements[CollisionPoint.Indices[1]];
-								const int32 IndexToWrite = ConstraintIndex.fetch_add(1);
-
-								FGaussSeidelSpringConstraintSingleData SingleConstraintData;
-								SingleConstraintData.SingleIndices = { Elem[0], Elem[1] ,Elem[2] };
-								SingleConstraintData.SingleSecondIndices =  { Index };
-								SingleConstraintData.SingleWeights = { CollisionPoint.Bary[1], CollisionPoint.Bary[2], CollisionPoint.Bary[3] };
-								SingleConstraintData.SingleSecondWeights = {FSolverReal(1.f)};
-								SingleConstraintData.bIsAnisotropic = UseAnisotropicSpring;
-								SingleConstraintData.SingleNormal = CollisionPoint.Normal;
-								SingleConstraintData.bIsZeroRestLength = true; //Push-out type collision springs should be zero rest length
-
-								float SpringStiffness = 0.f;
-								for (int32 k = 0; k < 3; k++)
-								{
-									SpringStiffness += SingleConstraintData.SingleWeights[k] * PositionTargetStiffness * Particles.M(Elem[k]);
-								}
-								SpringStiffness += PositionTargetStiffness * Particles.M(Index);
-								SingleConstraintData.SingleStiffness = (FSolverReal)SpringStiffness;
-								ConstraintsData.SetSingleConstraint(SingleConstraintData, IndexToWrite);
-							}
-						}
-					}
-				}
-			);
-
-			// Shrink the arrays to the actual number of found constraints.
-			const int32 ConstraintNum = ConstraintIndex.load();
-			Resize(ConstraintNum);
-		}
+		void CollisionDetectionSpatialHash(const FSolverParticles& Particles, const TArray<int32>& SurfaceVertices, const FTriangleMesh& TriangleMesh, const TArray<int32>& ComponentIndex, const SpatialAccelerator& Spatial, float DetectRadius = 1.f, float PositionTargetStiffness = 10000.f, bool UseAnisotropicSpring = true);
 
 		template<typename SpatialAccelerator>
-		void CollisionDetectionSpatialHashInComponent(const FSolverParticles& Particles, const TArray<int32>& SurfaceVertices, const FTriangleMesh& TriangleMesh, const TMap<int32, TSet<int32>>& ExcludeMap, const SpatialAccelerator& Spatial, float DetectRadius = 0.f, float PositionTargetStiffness = 10000.f, bool UseAnisotropicSpring = true)
-		{
-			TRACE_CPUPROFILER_EVENT_SCOPE(STAT_ChaosGaussSeidelWeakConstraintsCollisionDetectionSpatialHashInComponent);
-			Resize(InitialWCSize + Particles.Size());
-			std::atomic<int32> ConstraintIndex(InitialWCSize);
-			const TArray<TVec3<int32>>& Elements = TriangleMesh.GetSurfaceElements();
-			const float HalfRadius = DetectRadius/2;
-			PhysicsParallelFor(SurfaceVertices.Num(),
-				[this, &Spatial, &Particles, &SurfaceVertices, &ConstraintIndex, &TriangleMesh, &ExcludeMap, &Elements, &HalfRadius, &PositionTargetStiffness, &UseAnisotropicSpring](int32 i)
-				{
-					const int32 Index = SurfaceVertices[i];
-					TArray< TTriangleCollisionPoint<FSolverReal>> Result;
-					//PointProximityQuery
-					if (TriangleMesh.PointClosestTriangleQuery(Spatial, static_cast<const TArrayView<const FSolverVec3>&>(Particles.XArray()), Index, Particles.GetX(Index), HalfRadius, HalfRadius,
-						[this, &Elements, &ExcludeMap](const int32 PointIndex, const int32 TriangleIndex)->bool
-						{	
-							return  !(ExcludeMap.Find(PointIndex) && ExcludeMap[PointIndex].Contains(TriangleIndex));
-						},
-						Result))
-					{
-						for (const TTriangleCollisionPoint<FSolverReal>& CollisionPoint : Result)
-						{
-							if (CollisionPoint.Phi < 0)
-							{
-								const TVector<int32, 3>& Elem = Elements[CollisionPoint.Indices[1]];
-								const int32 IndexToWrite = ConstraintIndex.fetch_add(1);
+		void CollisionDetectionSpatialHashInComponent(const FSolverParticles& Particles, const TArray<int32>& SurfaceVertices, const FTriangleMesh& TriangleMesh, const TMap<int32, TSet<int32>>& ExcludeMap, const SpatialAccelerator& Spatial, float DetectRadius = 0.f, float PositionTargetStiffness = 10000.f, bool UseAnisotropicSpring = true);
 
-								FGaussSeidelSpringConstraintSingleData SingleConstraintData;
-								SingleConstraintData.SingleIndices = { Elem[0], Elem[1] ,Elem[2] };
-								SingleConstraintData.SingleSecondIndices =  { Index };
-								SingleConstraintData.SingleWeights = { CollisionPoint.Bary[1], CollisionPoint.Bary[2], CollisionPoint.Bary[3] };
-								SingleConstraintData.SingleSecondWeights = {FSolverReal(1.f)};
-								SingleConstraintData.bIsAnisotropic = UseAnisotropicSpring;
-								SingleConstraintData.SingleNormal = CollisionPoint.Normal;
-								SingleConstraintData.bIsZeroRestLength = true; //Push-out type collision springs should be zero rest length
+		void ComputeCollisionWCDataSimplified(TArray<TArray<int32>>& ExtraConstraints, TArray<TArray<int32>>& ExtraWCIncidentElements, TArray<TArray<int32>>& ExtraWCIncidentElementsLocal);
 
-								float SpringStiffness = 0.f;
-								for (int32 k = 0; k < 3; k++)
-								{
-									SpringStiffness += SingleConstraintData.SingleWeights[k] * PositionTargetStiffness * Particles.M(Elem[k]);
-								}
-								SpringStiffness += PositionTargetStiffness * Particles.M(Index);
-								SingleConstraintData.SingleStiffness = (FSolverReal)SpringStiffness;
-								ConstraintsData.SetSingleConstraint(SingleConstraintData, IndexToWrite);
-							}
-						}
-					}
-				}
-			);
+		CHAOS_API const TArray<TArray<int32>>& GetStaticConstraintArrays(TArray<TArray<int32>>& IncidentElements, TArray<TArray<int32>>& IncidentElementsLocal) const;
+		TArray<TArray<int32>> GetDynamicConstraintArrays(TArray<TArray<int32>>& IncidentElements, TArray<TArray<int32>>& IncidentElementsLocal) const;
 
-			// Shrink the arrays to the actual number of found constraints.
-			const int32 ConstraintNum = ConstraintIndex.load();
-			Resize(ConstraintNum);
-		}
-
-		void ComputeCollisionWCDataSimplified(TArray<TArray<int32>>& ExtraConstraints, TArray<TArray<int32>>& ExtraWCIncidentElements, TArray<TArray<int32>>& ExtraWCIncidentElementsLocal)
-		{
-			ensureMsgf(ConstraintsData.Size() >= InitialWCSize, TEXT("The size of Indices is smaller than InitialWCSize"));
-
-			ExtraConstraints.Init(TArray<int32>(), ConstraintsData.Size() - InitialWCSize);
-			for (int32 i = InitialWCSize; i < static_cast<int32>(ConstraintsData.Size()); i++)
-			{
-				const TArray<int32>& LocalIndices = ConstraintsData.GetIndices(i);
-				const TArray<int32>& LocalSecondIndices = ConstraintsData.GetSecondIndices(i);
-
-				ExtraConstraints[i - InitialWCSize].SetNum(LocalIndices.Num() + LocalSecondIndices.Num());
-				for (int32 j = 0; j < LocalIndices.Num(); j++)
-				{
-					ExtraConstraints[i - InitialWCSize][j] = LocalIndices[j];
-				}
-				for (int32 j = 0; j < LocalSecondIndices.Num(); j++)
-				{
-					ExtraConstraints[i - InitialWCSize][j + LocalIndices.Num()] = LocalSecondIndices[j];
-				}
-			}
-
-			ExtraWCIncidentElements = Chaos::Utilities::ComputeIncidentElements(ExtraConstraints, &ExtraWCIncidentElementsLocal);
-
-			NodalWeights = NoCollisionNodalWeights;
-			for (int32 i = 0; i < ExtraWCIncidentElements.Num(); i++)
-			{
-				if (ExtraWCIncidentElements[i].Num() > 0)
-				{
-					int32 p = ExtraConstraints[ExtraWCIncidentElements[i][0]][ExtraWCIncidentElementsLocal[i][0]];
-					if (NodalWeights[p].Num() == 0)
-					{
-						NodalWeights[p].Init(FSolverReal(0), 6);
-					}
-					for (int32 j = 0; j < ExtraWCIncidentElements[i].Num(); j++)
-					{
-						int32 LocalIndex = ExtraWCIncidentElementsLocal[i][j];
-						int32 ConstraintIndex = ExtraWCIncidentElements[i][j] + InitialWCSize;
-
-						const FGaussSeidelSpringConstraintSingleData& SingleData = ConstraintsData.GetSingleConstraintData(ConstraintIndex);
-
-						FSolverReal weight = FSolverReal(0);
-						if (LocalIndex >= SingleData.SingleIndices.Num())
-						{
-							weight = SingleData.SingleWeights[LocalIndex - SingleData.SingleIndices.Num()];
-						}
-						else
-						{
-							weight = SingleData.SingleWeights[LocalIndex];
-						}
-						if (SingleData.bIsAnisotropic)
-						{
-							for (int32 alpha = 0; alpha < 3; alpha++)
-							{
-								NodalWeights[p][alpha] += SingleData.SingleNormal[alpha] * SingleData.SingleNormal[alpha] * weight * weight * SingleData.SingleStiffness;
-							}
+		UE_DEPRECATED(5.8, "Replaced with AddWCResidual for more general cases.")
+		void AddZeroRestLengthWCResidualAndHessian(const FSolverParticles& InParticles, const int32 ConstraintIndex, const int32 LocalIndex, const FSolverReal Dt, TVec3<FSolverReal>& ParticleResidual, Chaos::PMatrix<FSolverReal, 3, 3>& ParticleHessian) const;
 
-							NodalWeights[p][3] += SingleData.SingleNormal[0] * SingleData.SingleNormal[1] * weight * weight * SingleData.SingleStiffness;
-							NodalWeights[p][4] += SingleData.SingleNormal[0] * SingleData.SingleNormal[2] * weight * weight * SingleData.SingleStiffness;
-							NodalWeights[p][5] += SingleData.SingleNormal[1] * SingleData.SingleNormal[2] * weight * weight * SingleData.SingleStiffness;
-						}
-						else
-						{
-							for (int32 alpha = 0; alpha < 3; alpha++)
-							{
-								NodalWeights[p][alpha] += weight * weight * SingleData.SingleStiffness;
-							}
-						}
-					}
-				}
-			}
-		}
+		CHAOS_API void AddWCResidual(const FSolverParticlesRange& InParticles, const int32 ConstraintIndex, const int32 LocalIndex, const FSolverReal Dt, TVec3<FSolverReal>& ParticleResidual, Chaos::PMatrix<FSolverReal, 3, 3>& ParticleHessian) const;
 
+		CHAOS_API void AddWCResidual(const FSolverParticles& InParticles, const int32 ConstraintIndex, const int32 LocalIndex, const FSolverReal Dt, TVec3<FSolverReal>& ParticleResidual, Chaos::PMatrix<FSolverReal, 3, 3>& ParticleHessian) const;
 
-		const TArray<TArray<int32>>& GetStaticConstraintArrays(TArray<TArray<int32>>& IncidentElements, TArray<TArray<int32>>& IncidentElementsLocal) const
+		int32 GetInitialWCSize() const
 		{
-			IncidentElements = NoCollisionWCIncidentElements;
-			IncidentElementsLocal = NoCollisionWCIncidentElementsLocal;
-			return NoCollisionConstraints;
+			return InitialWCSize;
 		}
 
-		TArray<TArray<int32>> GetDynamicConstraintArrays(TArray<TArray<int32>>& IncidentElements, TArray<TArray<int32>>& IncidentElementsLocal) const
-		{
-			TArray<TArray<int32>> ExtraConstraints;
-			ExtraConstraints.Init(TArray<int32>(), ConstraintsData.Size());
-			for (int32 i = InitialWCSize; i < ConstraintsData.Size(); i++)
-			{
-				const TArray<int32>& LocalIndices = ConstraintsData.GetIndices(i);
-				const TArray<int32>& LocalSecondIndices = ConstraintsData.GetSecondIndices(i);
-				ExtraConstraints[i - InitialWCSize].SetNum(LocalIndices.Num() + LocalSecondIndices.Num());
-				for (int32 j = 0; j < LocalIndices.Num(); j++)
-				{
-					ExtraConstraints[i - InitialWCSize][j] = LocalIndices[j];
-				}
-				for (int32 j = 0; j < LocalSecondIndices.Num(); j++)
-				{
-					ExtraConstraints[i - InitialWCSize][j + LocalIndices.Num()] = LocalSecondIndices[j];
-				}
-			}
+	protected:
+		CHAOS_API virtual void SetConstraintsRestLength(const FSolverParticlesRange& InParticles, const FSolverParticlesRange& InSecondParticles);
 
-			IncidentElements = Chaos::Utilities::ComputeIncidentElements(ExtraConstraints, &IncidentElementsLocal);
+		CHAOS_API void UpdateTriangleNormalAndNodalWeight(const FSolverParticlesRange& InParticles, const FSolverParticlesRange& InSecondParticles, bool bUseParticleX);
 
-			return ExtraConstraints;
-		}
+		CHAOS_API TVec3<FSolverReal> ComputeSpringEdge(const FSolverParticlesRange& InParticles, const FSolverParticlesRange& InSecondParticles,
+			const TArray<int32>& LocalIndices, const TArray<int32>& LocalSecondIndices,
+			const TArray<FSolverReal>& Weight, const TArray<FSolverReal>& SecondWeight, bool bUseParticleX) const;
 
-		UE_DEPRECATED(5.8, "Replaced with AddWCResidual for more general cases.")
-		void AddZeroRestLengthWCResidualAndHessian(const FSolverParticles& InParticles, const int32 ConstraintIndex, const int32 LocalIndex, const FSolverReal Dt, TVec3<FSolverReal>& ParticleResidual, Chaos::PMatrix<FSolverReal, 3, 3>& ParticleHessian) const
-		{
-			const FGaussSeidelSpringConstraintSingleData& SingleData = ConstraintsData.GetSingleConstraintData(ConstraintIndex);
+		CHAOS_API TVec3<FSolverReal> ComputeSpringEdge(const FSolverParticles& InParticles, const FSolverParticles& InSecondParticles,
+			const TArray<int32>& LocalIndices, const TArray<int32>& LocalSecondIndices,
+			const TArray<FSolverReal>& Weight, const TArray<FSolverReal>& SecondWeight, bool bUseParticleX) const;
 
-			const TVec3<FSolverReal> SpringEdge = ComputeSpringEdge(InParticles, SingleData.SingleIndices, SingleData.SingleSecondIndices,
-				SingleData.SingleWeights, SingleData.SingleSecondWeights, /*bUseParticleX =*/false);
-			FSolverReal weight = FSolverReal(0);
-			if (LocalIndex >= SingleData.SingleIndices.Num())
-			{
-				weight = -SingleData.SingleSecondWeights[LocalIndex - SingleData.SingleIndices.Num()];
-			}
-			else
-			{
-				weight = SingleData.SingleWeights[LocalIndex];
-			}
-			if (SingleData.bIsAnisotropic)
-			{
-				FSolverReal comp = TVec3<FSolverReal>::DotProduct(SpringEdge, SingleData.SingleNormal);
-				TVec3<FSolverReal> proj = SingleData.SingleNormal * comp;
-				for (int32 alpha = 0; alpha < 3; alpha++)
-				{
-					ParticleResidual[alpha] += Dt * Dt * SingleData.SingleStiffness * proj[alpha] * weight;
-				}
-			}
-			else
-			{
-				for (int32 alpha = 0; alpha < 3; alpha++)
-				{
-					ParticleResidual[alpha] += Dt * Dt * SingleData.SingleStiffness * SpringEdge[alpha] * weight;
-				}
-			}
-		}
+		CHAOS_API TVec3<FSolverReal> ComputeSpringEdge(const FSolverParticlesRange& InParticles,
+			const TArray<int32>& LocalIndices, const TArray<int32>& LocalSecondIndices,
+			const TArray<FSolverReal>& Weight, const TArray<FSolverReal>& SecondWeight, bool bUseParticleX) const;
 
-		void AddWCResidual(const FSolverParticlesRange& InParticles, const int32 ConstraintIndex, const int32 LocalIndex, const FSolverReal Dt, TVec3<FSolverReal>& ParticleResidual, Chaos::PMatrix<FSolverReal, 3, 3>& ParticleHessian) const
-		{
-			const FGaussSeidelSpringConstraintSingleData& SingleData = ConstraintsData.GetSingleConstraintData(ConstraintIndex);
+		CHAOS_API TVec3<FSolverReal> ComputeSpringEdge(const FSolverParticles& InParticles, const TArray<int32>& LocalIndices,
+			const TArray<int32>& LocalSecondIndices, const TArray<FSolverReal>& Weight, const TArray<FSolverReal>& SecondWeight, bool bUseParticleX) const;
 
-			const TVec3<FSolverReal> SpringEdge = ComputeSpringEdge(InParticles, SingleData.SingleIndices, SingleData.SingleSecondIndices,
-				SingleData.SingleWeights, SingleData.SingleSecondWeights, /*bUseParticleX =*/false);
-			FSolverReal weight = FSolverReal(0);
-			if (LocalIndex >= SingleData.SingleIndices.Num())
-			{
-				weight = -SingleData.SingleSecondWeights[LocalIndex - SingleData.SingleIndices.Num()];
-			}
-			else
-			{
-				weight = SingleData.SingleWeights[LocalIndex];
-			}
-			TVec3<FSolverReal> Projection;
-			if (SingleData.bIsAnisotropic)
-			{
-				FSolverReal LengthDiff = TVec3<FSolverReal>::DotProduct(SpringEdge, SingleData.SingleNormal) - SingleData.RestLength;
-				Projection = SingleData.SingleNormal * LengthDiff;
-			}
-			else
-			{
-				Projection = SpringEdge;
-				if (!SingleData.bIsZeroRestLength) // if not zero rest-length, apply repulsion force
-				{
-					Projection -= SingleData.RestLength * SpringEdge.GetSafeNormal();
-				}	
-			}
-			for (int32 alpha = 0; alpha < 3; ++alpha)
-			{
-				ParticleResidual[alpha] += Dt * Dt * SingleData.SingleStiffness * Projection[alpha] * weight;
-			}
-		}
-
-		void AddWCResidual(const FSolverParticles& InParticles, const int32 ConstraintIndex, const int32 LocalIndex, const FSolverReal Dt, TVec3<FSolverReal>& ParticleResidual, Chaos::PMatrix<FSolverReal, 3, 3>& ParticleHessian) const
-		{
-			AddWCResidual(FSolverParticlesRange(const_cast<FSolverParticles*>(&InParticles), /*InOffset = */ 0, /*InRange = */ InParticles.Size()), ConstraintIndex, LocalIndex, Dt, ParticleResidual, ParticleHessian);
-		}
+		CHAOS_API virtual void SetNoCollisionData();
 
 		FGaussSeidelSpringConstraintData ConstraintsData;
-
 		TArray<TArray<FSolverReal>> NodalWeights;
-
 		TArray<TArray<int32>> WCIncidentElements;
 		TArray<TArray<int32>> WCIncidentElementsLocal;
 
-		FDeformableXPBDWeightedSpringConstraintParams DebugDrawParams;
-
-		int32 InitialWCSize;
+		// Static constraint data (before adding collision constraints)
 		TArray<TArray<FSolverReal>> NoCollisionNodalWeights;
 		TArray<TArray<int32>> NoCollisionConstraints;
 		TArray<TArray<int32>> NoCollisionWCIncidentElements;
 		TArray<TArray<int32>> NoCollisionWCIncidentElementsLocal;
+
+	private:
+		CHAOS_API virtual void ComputeIncidentElements();
+
+		void UpdateTriangleNormal(const FSolverParticlesRange& InParticles, const FSolverParticlesRange& InSecondParticles, bool bUseParticleX);
+
+		CHAOS_API virtual void UpdateNodalWeight(const FSolverParticlesRange& InParticles, const FSolverParticlesRange& InSecondParticles, bool bUseParticleX);
+
+		void UpdateTriangleNormalAndNodalWeight(const FSolverParticlesRange& InParticles, bool bUseParticleX);
+
+		void Resize(int32 Size);
+
+		void UpdatePointTriangleCollisionWCData(const FSolverParticles& Particles);
+
+		CHAOS_API void VisualizeAllBindings(const FSolverParticlesRange& InParticles, const FSolverReal Dt) const;
+
+		int32 InitialWCSize;
+		FDeformableXPBDWeightedSpringConstraintParams DebugDrawParams;
 	};
 
 	template <typename T>
@@ -1127,7 +230,7 @@ namespace Chaos::Softs
 
 	template <class T>
 	using FGaussSeidelWeakConstraintData UE_DEPRECATED(5.8, "Deprecated. this class is to be deleted, use FGaussSeidelSpringConstraintData instead") = FGaussSeidelSpringConstraintData;
-	
+
 	template <typename T, typename ParticleType>
 	using FGaussSeidelWeakConstraints UE_DEPRECATED(5.8, "Deprecated. this struct is to be deleted, use FGaussSeidelSpringConstraints instead") = FGaussSeidelSpringConstraints;
 
diff --git a/Engine/Source/Runtime/Experimental/Chaos/Public/Chaos/GraphColoring.h b/Engine/Source/Runtime/Experimental/Chaos/Public/Chaos/GraphColoring.h
--- a/Engine/Source/Runtime/Experimental/Chaos/Public/Chaos/GraphColoring.h
+++ b/Engine/Source/Runtime/Experimental/Chaos/Public/Chaos/GraphColoring.h
@@ -60,9 +60,14 @@ void ComputeGridBasedGraphSubColoringPointer(const TArray<TArray<int32>>& Elemen
 template<typename T>
 CHAOS_API void ComputeGridBasedGraphSubColoringPointer(const TArray<TArray<int32>>& ElementsPerColor, const TMPMGrid<T>& Grid, const int32 GridSize, TUniquePtr<TArray<TArray<int32>>>& PreviousColoring, const TArray<TArray<int32>>& ConstraintsNodesSet, TArray<TArray<TArray<int32>>>& ElementsPerSubColors);
 
+// Both Indices and SecondIndices point to the same InParticles
 template<typename T>
 CHAOS_API void ComputeWeakConstraintsColoring(const TArray<TArray<int32>>& Indices, const TArray<TArray<int32>>& SecondIndices, const Chaos::TDynamicParticles<T, 3>& InParticles, TArray<TArray<int32>>& ConstraintsPerColor);
 
+// if SecondIndices point to InSecondParticles
+template<typename T>
+CHAOS_API void ComputeWeakConstraintsColoring(const TArray<TArray<int32>>& Indices, const TArray<TArray<int32>>& SecondIndices, const Chaos::TDynamicParticles<T, 3>& InParticles, const Chaos::TDynamicParticles<T, 3>& InSecondParticles, TArray<TArray<int32>>& ConstraintsPerColor);
+
 template<typename T>
 CHAOS_API TArray<TArray<int32>> ComputeNodalColoring(const TArray<TVec4<int32>>& Graph, const Chaos::TDynamicParticles<T, 3>& InParticles, const int32 GraphParticlesStart, const int32 GraphParticlesEnd, const TArray<TArray<int32>>& IncidentElements, const TArray<TArray<int32>>& IncidentElementsLocalIndex);
 
diff --git a/Engine/Source/Runtime/Experimental/Chaos/Public/Chaos/Plane.h b/Engine/Source/Runtime/Experimental/Chaos/Public/Chaos/Plane.h
--- a/Engine/Source/Runtime/Experimental/Chaos/Public/Chaos/Plane.h
+++ b/Engine/Source/Runtime/Experimental/Chaos/Public/Chaos/Plane.h
@@ -396,4 +396,106 @@ bool IntersectPlanes2(TVector<T,d>& I, TVector<T,d>& D, const TPlane<T,d>& P1, c
 	return RetVal;
 }
 
+// Returns the closest distance between two edges (X0,X1) and (Y0,Y1) and their barycentric coordinates of closest point pair
+template<typename T, int d>
+T FindClosestPointAndBaryBetweenEdges(const TVector<T, d>& X0, const TVector<T, d>& X1, const TVector<T, d>& Y0, const TVector<T, d>& Y1, 
+	TVector<T, 2>& BaryX, TVector<T, 2>& BaryY, T Eps = UE_SMALL_NUMBER)
+{
+	// Edge directions
+	const TVector<T, d> U = X1 - X0;
+	const TVector<T, d> V = Y1 - Y0;
+	const TVector<T, d> W = X0 - Y0;
+
+	const T a = U.SizeSquared();
+	const T b = U.Dot(V);
+	const T c = V.SizeSquared();
+	const T dval = U.Dot(W);
+	const T eval = V.Dot(W);
+
+	const T D = a * c - b * b;
+
+	T sN, sD = D;
+	T tN, tD = D;
+
+	// Initial solution
+	if (D < Eps) // nearly parallel
+	{
+		sN = (T)0;
+		sD = (T)1;
+		tN = eval;
+		tD = c;
+	}
+	else
+	{
+		sN = (b * eval - c * dval);
+		tN = (a * eval - b * dval);
+	}
+
+	// Clamp s to [0,1]
+	if (sN <= (T)0)
+	{
+		sN = (T)0;
+		tN = eval;
+		tD = c;
+	}
+	else if (sN >= sD)
+	{
+		sN = sD;
+		tN = eval + b;
+		tD = c;
+	}
+
+	// Clamp t to [0,1]
+	if (tN <= (T)0)
+	{
+		tN = (T)0;
+		if (-dval <= (T)0)
+		{
+			sN = (T)0;
+			sD = (T)1;
+		}
+		else if (-dval >= a)
+		{
+			sN = sD;
+		}
+		else
+		{
+			sN = -dval;
+			sD = a;
+		}
+	}
+	else if (tN >= tD)
+	{
+		tN = tD;
+		if ((-dval + b) <= (T)0)
+		{
+			sN = (T)0;
+			sD = (T)1;
+		}
+		else if ((-dval + b) >= a)
+		{
+			sN = sD;
+		}
+		else
+		{
+			sN = (-dval + b);
+			sD = a;
+		}
+	}
+
+	// Barycentric weights
+	const T s = (FMath::Abs(sN) < Eps) ? (T)0 : sN / sD;
+	const T t = (FMath::Abs(tN) < Eps) ? (T)0 : tN / tD;
+	BaryX[0] = (T)1 - s;
+	BaryX[1] = s;
+	BaryY[0] = (T)1 - t;
+	BaryY[1] = t;
+
+	// Closest points
+	const TVector<T, d> ClosestX = X0 + U * s;
+	const TVector<T, d> ClosestY = Y0 + V * t;
+
+	return (ClosestX - ClosestY).Size();
+}
+
 }
diff --git a/Engine/Source/Runtime/Experimental/Chaos/Public/Chaos/TriangleMesh.h b/Engine/Source/Runtime/Experimental/Chaos/Public/Chaos/TriangleMesh.h
--- a/Engine/Source/Runtime/Experimental/Chaos/Public/Chaos/TriangleMesh.h
+++ b/Engine/Source/Runtime/Experimental/Chaos/Public/Chaos/TriangleMesh.h
@@ -335,6 +335,11 @@ namespace Chaos
 		CHAOS_API bool PointClosestTriangleQuery(const TSpatialHashType<T>& SpatialHash, const TConstArrayView<TVec3<T>>& Points, const int32 PointIndex, const TVec3<T>& PointPosition, const T PointThickness, const T ThisThickness,
 			TFunctionRef<bool(const int32 PointIndex, const int32 TriangleIndex)> BroadphaseTest, TArray<TTriangleCollisionPoint<T>>& Result) const;
 
+		// uses GetPosition lambda function to query point position instead of Points array
+		template<typename T>
+		CHAOS_API bool PointClosestTriangleQuery(const TSpatialHashType<T>& SpatialHash, TFunctionRef<const TVec3<T>&(int32 PointIndex)> GetPosition, const int32 PointIndex, const TVec3<T>& PointPosition, const T PointThickness, const T ThisThickness,
+			TFunctionRef<bool(const int32 PointIndex, const int32 TriangleIndex)> BroadphaseTest, TArray<TTriangleCollisionPoint<T>>& Result) const;
+
 		template<typename T>
 		bool EdgeIntersectionQuery(const TSpatialHashType<T>& SpatialHash, const TConstArrayView<TVec3<T>>& Points, const int32 EdgeIndex, const TVec3<T>& EdgePosition1, const TVec3<T>& EdgePosition2,
 			TFunctionRef<bool(const int32 EdgeIndex, const int32 TriangleIndex)> BroadphaseTest, TArray<TTriangleCollisionPoint<T>>& Result) const;
diff --git a/Engine/Source/Runtime/Experimental/Chaos/Public/Chaos/XPBDWeakConstraints.h b/Engine/Source/Runtime/Experimental/Chaos/Public/Chaos/XPBDWeakConstraints.h
--- a/Engine/Source/Runtime/Experimental/Chaos/Public/Chaos/XPBDWeakConstraints.h
+++ b/Engine/Source/Runtime/Experimental/Chaos/Public/Chaos/XPBDWeakConstraints.h
@@ -35,20 +35,17 @@ namespace Chaos::Softs
 	{
 
 	public:
+		static constexpr FSolverReal SoftMaxStiffness = (FSolverReal)1e14; // Stiffnesses greater than this will be treated as "hard" PBD constraints
+
 		// Used for spring constraints without self targets (against a fixed constraint target instead of deformable vertices)
-		FXPBDWeightedSpringConstraints(
+		CHAOS_API FXPBDWeightedSpringConstraints(
 			const FSolverParticles& InParticles,
 			const TArray<TArray<int32>>& InIndices,
 			const TArray<TArray<FSolverReal>>& InWeights,
 			const TArray<FSolverReal>& InStiffness
-		)
-			: Indices(InIndices), Weights(InWeights), Stiffness(InStiffness)
-		{
-			InitColor(InParticles);
-			LambdaArray.Init((FSolverReal)0., Indices.Num() * 3);
-		}
+		);
 
-		FXPBDWeightedSpringConstraints(
+		CHAOS_API FXPBDWeightedSpringConstraints(
 			const FSolverParticles& InParticles,
 			const TArray<TArray<int32>>& InIndices,
 			const TArray<TArray<FSolverReal>>& InWeights,
@@ -56,24 +53,9 @@ namespace Chaos::Softs
 			const TArray<TArray<int32>>& InSecondIndices,
 			const TArray<TArray<FSolverReal>>& InSecondWeights,
 			const FDeformableXPBDWeightedSpringConstraintParams& InParams
-		)
-			: Indices(InIndices), Weights(InWeights), SecondIndices(InSecondIndices), SecondWeights(InSecondWeights), Stiffness(InStiffness), DebugDrawParams(InParams)
-		{
-			ensureMsgf(Indices.Num() == SecondIndices.Num(), TEXT("Input Double Bindings have wrong size"));
-			
-			for (int32 i = 0; i < Indices.Num(); i++)
-			{
-				TSet<int32> IndicesSet = TSet<int32>(Indices[i]);
-				for (int32 j = 0; j < SecondIndices[i].Num(); j++)
-				{
-					ensureMsgf(!IndicesSet.Contains(SecondIndices[i][j]), TEXT("Indices and Second Indices overlaps. Currently not supported"));
-				}
-			}
-			InitColor(InParticles);
-			LambdaArray.Init((FSolverReal)0., Indices.Num() * 3);
-		}
+		);
 
-		FXPBDWeightedSpringConstraints(
+		CHAOS_API FXPBDWeightedSpringConstraints(
 			const FSolverParticles& InParticles,
 			const FSolverParticles& InSecondParticles,
 			const TArray<TArray<int32>>& InIndices,
@@ -82,88 +64,30 @@ namespace Chaos::Softs
 			const TArray<TArray<int32>>& InSecondIndices,
 			const TArray<TArray<FSolverReal>>& InSecondWeights,
 			const FDeformableXPBDWeightedSpringConstraintParams& InParams
-		)
-			: Indices(InIndices), Weights(InWeights), SecondIndices(InSecondIndices), SecondWeights(InSecondWeights), Stiffness(InStiffness), DebugDrawParams(InParams)
-		{
-			ensureMsgf(Indices.Num() == SecondIndices.Num(), TEXT("Input Double Bindings have wrong size"));
-
-			//TODO: write InitColor for both sides of particles
-			//InitColor(InParticles);
-			LambdaArray.Init((FSolverReal)0., Indices.Num() * 3);
-		}
+		);
 
 		virtual ~FXPBDWeightedSpringConstraints() = default;
 
-		void ApplyInParallel(FSolverParticlesRange& Particles, const FSolverReal Dt) const
-		{
-			TRACE_CPUPROFILER_EVENT_SCOPE(STAT_ChaosXPBDWeakConstraintApply);
-
-			if ((ConstraintsPerColor.Num() > 0))
-			{
-				if (SecondIndices.Num() == 0)
-				{
-					for (int32 Color = 0; Color < ConstraintsPerColor.Num(); Color++)
-					{
-						PhysicsParallelFor(ConstraintsPerColor[Color].Num(), [&](const int32 Index)
-							{
-								const int32 ConstraintIndex = ConstraintsPerColor[Color][Index];
-								ApplySingleConstraintWithoutSelfTarget(Particles, Dt, ConstraintIndex);
-							});
-					}
-				}
-				else 
-				{
-					for (int32 Color = 0; Color < ConstraintsPerColor.Num(); Color++)
-					{
-						PhysicsParallelFor(ConstraintsPerColor[Color].Num(), [&](const int32 Index)
-							{
-								const int32 ConstraintIndex = ConstraintsPerColor[Color][Index];
-								ApplySingleConstraintWithSelfTarget(Particles, Dt, ConstraintIndex);
-							});
-					}
-				}
-			}
-		}
+		CHAOS_API void ApplyInParallel(FSolverParticlesRange& Particles, const FSolverReal Dt) const;
 
-		void ApplyInParallel(FSolverParticles& Particles, const FSolverReal Dt) const
-		{
-			FSolverParticlesRange Range(&Particles, /*InOffset = */ 0, /*InRange = */ Particles.Size());
-			ApplyInParallel(Range, Dt);
-		}
+		CHAOS_API void ApplyInParallel(FSolverParticles& Particles, const FSolverReal Dt) const;
 
-		void Apply(FSolverParticlesRange& Particles, const FSolverReal Dt) const
-		{
-			ApplyInParallel(Particles, Dt);
-		}
+		CHAOS_API void Apply(FSolverParticlesRange& Particles, const FSolverReal Dt) const;
 
-		void ApplyInSerial(FSolverParticlesRange& Particles, FSolverParticlesRange& SecondParticles, const FSolverReal Dt) const
-		{
-			TRACE_CPUPROFILER_EVENT_SCOPE(STAT_ChaosXPBDWeakConstraintApply);
+		CHAOS_API void ApplyInSerial(FSolverParticlesRange& Particles, FSolverParticlesRange& SecondParticles, const FSolverReal Dt) const;
 
-			for (int32 ConstraintIdx = 0; ConstraintIdx < Indices.Num(); ++ConstraintIdx)
-			{
-				ApplySingleConstraintWithSelfTarget(Particles, SecondParticles, Dt, ConstraintIdx);
-			}
-		}
+		CHAOS_API void ApplyInParallel(FSolverParticlesRange& Particles, FSolverParticlesRange& SecondParticles, const FSolverReal Dt) const;
 
 		const TArray<TArray<int32>>& GetIndices()
 		{
 			return Indices;
 		}
 
-		void Init(const FSolverParticlesRange& InParticles, const FSolverReal Dt) const
-		{
-			for (FSolverReal& Lambdas : LambdaArray) { Lambdas = (FSolverReal)0.; }
-			if (DebugDrawParams.bVisualizeBindings)
-			{
-				VisualizeAllBindings(InParticles, Dt);
-			}
-		}
+		CHAOS_API void Init(const FSolverParticlesRange& InParticles, const FSolverReal Dt) const;
 
-		void Init(const FSolverParticles& InParticles, const FSolverReal Dt) const
-		{
-			Init(FSolverParticlesRange(const_cast<FSolverParticles*>(&InParticles), /*InOffset = */ 0, /*InRange = */ InParticles.Size()), Dt);
-		}
+		CHAOS_API void Init(const FSolverParticles& InParticles, const FSolverReal Dt) const;
+
+		CHAOS_API void Init(const FSolverParticlesRange& InParticles, const FSolverParticlesRange& InSecondParticles, const FSolverReal Dt) const;
 
 		void UpdateTargets(TArray<TVector<FSolverReal, 3>>&& InTargets)
 		{
@@ -171,217 +95,45 @@ namespace Chaos::Softs
 		}
 
 
-	private:
-
-		void InitColor(const FSolverParticles& Particles)
-		{
-			Chaos::ComputeWeakConstraintsColoring(Indices, SecondIndices, Particles, ConstraintsPerColor);
-		}
-
-		void ApplySingleConstraintWithoutSelfTarget(FSolverParticlesRange& Particles, const FSolverReal Dt, const int32 ConstraintIndex, const FSolverReal Tol = UE_KINDA_SMALL_NUMBER) const
-		{
-			if (Stiffness[ConstraintIndex] < Tol)
-			{
-				return;
-			}
-
-			FSolverReal AlphaTilde = FSolverReal(2) / (Stiffness[ConstraintIndex] * Dt * Dt);
-
-			if (Stiffness[ConstraintIndex] > 1e14)
-			{
-				AlphaTilde = FSolverReal(0);
-			}
-
-			TVec3<FSolverReal> SpringEdge((FSolverReal)0.);
-			for (int32 j = 0; j < Indices[ConstraintIndex].Num(); j++) {
-				for (int32 alpha = 0; alpha < 3; alpha++) {
-					SpringEdge[alpha] += Weights[ConstraintIndex][j] * Particles.P(Indices[ConstraintIndex][j])[alpha];
-				}
-			}
-
-			SpringEdge -= Constraints[ConstraintIndex];
-
-			FSolverReal Denom = AlphaTilde;
-			for (int32 j = 0; j < Indices[ConstraintIndex].Num(); j++) {
-				Denom += Weights[ConstraintIndex][j] * Weights[ConstraintIndex][j] * Particles.InvM(Indices[ConstraintIndex][j]);
-				
-			}
-
-			for (int32 Beta = 0; Beta < 3; Beta++)
-			{
-				FSolverReal Cj = SpringEdge[Beta];
-				FSolverReal DLambda = -Cj - AlphaTilde * LambdaArray[ConstraintIndex * 3 + Beta];
-
-				DLambda /= Denom;
-				LambdaArray[ConstraintIndex * 3 + Beta] += DLambda;
-				for (int32 j = 0; j < Indices[ConstraintIndex].Num(); j++) {
-					Particles.P(Indices[ConstraintIndex][j])[Beta] += DLambda * Weights[ConstraintIndex][j] * Particles.InvM(Indices[ConstraintIndex][j]);
-				}
-			}
-		}
-
-		void ApplySingleConstraintWithSelfTarget(FSolverParticlesRange& Particles, const FSolverReal Dt, const int32 ConstraintIndex, const FSolverReal Tol = UE_KINDA_SMALL_NUMBER) const
-		{
-			if (Stiffness[ConstraintIndex] < Tol)
-			{
-				return;
-			}
-
-			ensure(SecondIndices.Num() > 0);
-
-			FSolverReal AlphaTilde = FSolverReal(2) / (Stiffness[ConstraintIndex] * Dt * Dt);
+		CHAOS_API void VisualizeAllBindings(const FSolverParticlesRange& InParticles, const FSolverParticlesRange& InSecondParticles,
+			const TFunctionRef<void(const FVector& Pos0, const FVector& Pos1, const FLinearColor& Color)>& DrawLine,
+			const TFunctionRef<void(const FVector& Pos, const FLinearColor& Color)>& DrawPoint) const;
 
-			if (Stiffness[ConstraintIndex] > 1e14)
-			{
-				AlphaTilde = FSolverReal(0);
-			}
-
-			TVec3<FSolverReal> SpringEdge((FSolverReal)0.);
-			for (int32 j = 0; j < Indices[ConstraintIndex].Num(); j++) {
-				for (int32 alpha = 0; alpha < 3; alpha++) {
-					SpringEdge[alpha] += Weights[ConstraintIndex][j] * Particles.P(Indices[ConstraintIndex][j])[alpha];
-				}
-			}
-
-			for (int32 j = 0; j < SecondIndices[ConstraintIndex].Num(); j++) {
+	protected:
+		CHAOS_API virtual bool ComputeSpringEdge(const FSolverParticlesRange& InParticles, TVec3<FSolverReal>& OutSpringEdge, int32 ConstraintIndex, bool bUseParticleX, bool bUseConstraintTargetPosition) const;
 
-					SpringEdge -= SecondWeights[ConstraintIndex][j] * Particles.P(SecondIndices[ConstraintIndex][j]);
+		CHAOS_API virtual bool ComputeSpringEdge(const FSolverParticlesRange& InParticles, const FSolverParticlesRange& InSecondParticles,
+			TVec3<FSolverReal>& OutSpringEdge, int32 ConstraintIndex, bool bUseParticleX, bool bUseConstraintTargetPosition) const;
 
-			}
+	protected:
+		TArray<TArray<int32>> Indices;
+		TArray<TArray<FSolverReal>> Weights;
+		TArray<TVector<FSolverReal, 3>> Constraints;
+		TArray<TArray<int32>> SecondIndices;
+		TArray<TArray<FSolverReal>> SecondWeights;
+		TArray<FSolverReal> Stiffness;
 
-			FSolverReal Denom = AlphaTilde;
-			for (int32 j = 0; j < Indices[ConstraintIndex].Num(); j++) {
-				Denom += Weights[ConstraintIndex][j] * Weights[ConstraintIndex][j] * Particles.InvM(Indices[ConstraintIndex][j]);
-			}
+	private:
 
-			for (int32 j = 0; j < SecondIndices[ConstraintIndex].Num(); j++) {
-				Denom += SecondWeights[ConstraintIndex][j] * SecondWeights[ConstraintIndex][j] * Particles.InvM(SecondIndices[ConstraintIndex][j]);
-			}
+		void InitColor(const FSolverParticles& Particles);
 
+		void InitColor(const FSolverParticles& Particles, const FSolverParticles& SecondParticles);
 
-			for (int32 Beta = 0; Beta < 3; Beta++)
-			{
-				FSolverReal Cj = SpringEdge[Beta];
-				FSolverReal DLambda = -Cj - AlphaTilde * LambdaArray[ConstraintIndex * 3 + Beta];
+		void ApplySingleConstraintWithoutSelfTarget(FSolverParticlesRange& Particles, const FSolverReal Dt, const int32 ConstraintIndex, const FSolverReal Tol = UE_KINDA_SMALL_NUMBER) const;
 
-				DLambda /= Denom;
-				LambdaArray[ConstraintIndex * 3 + Beta] += DLambda;
-				for (int32 j = 0; j < Indices[ConstraintIndex].Num(); j++) {
-					Particles.P(Indices[ConstraintIndex][j])[Beta] += DLambda * Weights[ConstraintIndex][j] * Particles.InvM(Indices[ConstraintIndex][j]);
-				}
+		void ApplySingleConstraintWithSelfTarget(FSolverParticlesRange& Particles, const FSolverReal Dt, const int32 ConstraintIndex, const FSolverReal Tol = UE_KINDA_SMALL_NUMBER) const;
 
-				for (int32 j = 0; j < SecondIndices[ConstraintIndex].Num(); j++) {
-					Particles.P(SecondIndices[ConstraintIndex][j])[Beta] -= DLambda * SecondWeights[ConstraintIndex][j] * Particles.InvM(SecondIndices[ConstraintIndex][j]);
-				}
-			}
-		}
+		//Deprecating with ApplySingleConstraintBetweenParticles
+		void ApplySingleConstraintWithSelfTarget(FSolverParticlesRange& Particles, FSolverParticlesRange& SecondParticles, const FSolverReal Dt, const int32 ConstraintIndex, const FSolverReal Tol = UE_KINDA_SMALL_NUMBER) const;
 
-		void ApplySingleConstraintWithSelfTarget(FSolverParticlesRange& Particles, FSolverParticlesRange& SecondParticles, const FSolverReal Dt, const int32 ConstraintIndex, const FSolverReal Tol = UE_KINDA_SMALL_NUMBER) const
-		{
-			if (Stiffness[ConstraintIndex] < Tol)
-			{
-				return;
-			}
-
-			ensure(SecondIndices.Num() > 0);
-			FSolverReal AlphaTilde = FSolverReal(2) / (Stiffness[ConstraintIndex] * Dt * Dt);
-			if (Stiffness[ConstraintIndex] > 1e14)
-			{
-				AlphaTilde = FSolverReal(0);
-			}
-
-			TVec3<FSolverReal> SpringEdge((FSolverReal)0.);
-			for (int32 j = 0; j < Indices[ConstraintIndex].Num(); j++)
-			{
-				for (int32 alpha = 0; alpha < 3; alpha++)
-				{
-					SpringEdge[alpha] += Weights[ConstraintIndex][j] * Particles.P(Indices[ConstraintIndex][j])[alpha];
-				}
-			}
-			for (int32 j = 0; j < SecondIndices[ConstraintIndex].Num(); j++)
-			{
-				SpringEdge -= SecondWeights[ConstraintIndex][j] * SecondParticles.P(SecondIndices[ConstraintIndex][j]);
-			}
-
-			FSolverReal Denom = AlphaTilde;
-			for (int32 j = 0; j < Indices[ConstraintIndex].Num(); j++)
-			{
-				Denom += Weights[ConstraintIndex][j] * Weights[ConstraintIndex][j] * Particles.InvM(Indices[ConstraintIndex][j]);
-			}
-			for (int32 j = 0; j < SecondIndices[ConstraintIndex].Num(); j++)
-			{
-				Denom += SecondWeights[ConstraintIndex][j] * SecondWeights[ConstraintIndex][j] * SecondParticles.InvM(SecondIndices[ConstraintIndex][j]);
-			}
-			for (int32 Beta = 0; Beta < 3; Beta++)
-			{
-				FSolverReal Cj = SpringEdge[Beta];
-				FSolverReal DLambda = -Cj - AlphaTilde * LambdaArray[ConstraintIndex * 3 + Beta];
-				DLambda /= Denom;
-				LambdaArray[ConstraintIndex * 3 + Beta] += DLambda;
-				for (int32 j = 0; j < Indices[ConstraintIndex].Num(); j++)
-				{
-					Particles.P(Indices[ConstraintIndex][j])[Beta] += DLambda * Weights[ConstraintIndex][j] * Particles.InvM(Indices[ConstraintIndex][j]);
-				}
-				for (int32 j = 0; j < SecondIndices[ConstraintIndex].Num(); j++)
-				{
-					SecondParticles.P(SecondIndices[ConstraintIndex][j])[Beta] -= DLambda * SecondWeights[ConstraintIndex][j] * SecondParticles.InvM(SecondIndices[ConstraintIndex][j]);
-				}
-			}
-		}
+		void ApplySingleConstraintBetweenParticles(FSolverParticlesRange& Particles, FSolverParticlesRange& SecondParticles, const FSolverReal Dt, const int32 ConstraintIndex, const FSolverReal Tol = UE_KINDA_SMALL_NUMBER) const;
 
-		void VisualizeAllBindings(const FSolverParticlesRange& InParticles, const FSolverReal Dt) const
-		{
-#if WITH_EDITOR
-			auto DoubleVert = [](Chaos::TVec3<FSolverReal> V) { return FVector3d(V.X, V.Y, V.Z); };
-			for (int32 i = 0; i < Indices.Num(); i++)
-			{
-				Chaos::TVec3<FSolverReal> SourcePos((FSolverReal)0.), TargetPos((FSolverReal)0.);
-				for (int32 j = 0; j < Indices[i].Num(); j++)
-				{
-					SourcePos += Weights[i][j] * InParticles.P(Indices[i][j]);
-				}
-				for (int32 j = 0; j < SecondIndices[i].Num(); j++)
-				{
-					TargetPos += SecondWeights[i][j] * InParticles.P(SecondIndices[i][j]);
-				}
-
-				float ParticleThickness = DebugDrawParams.DebugParticleWidth;
-				float LineThickness = DebugDrawParams.DebugLineWidth;
-
-				if (Indices[i].Num() == 1)
-				{
-					Chaos::FDebugDrawQueue::GetInstance().DrawDebugPoint(DoubleVert(SourcePos), FColor::Red, false, Dt, 0, ParticleThickness);
-					for (int32 j = 0; j < SecondIndices[i].Num(); j++)
-					{
-						Chaos::FDebugDrawQueue::GetInstance().DrawDebugPoint(DoubleVert(InParticles.P(SecondIndices[i][j])), FColor::Green, false, Dt, 0, ParticleThickness);
-						Chaos::FDebugDrawQueue::GetInstance().DrawDebugLine(DoubleVert(InParticles.P(SecondIndices[i][j])), DoubleVert(InParticles.P(SecondIndices[i][(j+1) % SecondIndices[i].Num()])), FColor::Green, false, Dt, 0, LineThickness);
-					}
-				
-				}
-
-				if (SecondIndices[i].Num() == 1)
-				{
-					Chaos::FDebugDrawQueue::GetInstance().DrawDebugPoint(DoubleVert(TargetPos), FColor::Red, false, Dt, 0, ParticleThickness);
-					for (int32 j = 0; j < Indices[i].Num(); j++)
-					{
-						Chaos::FDebugDrawQueue::GetInstance().DrawDebugPoint(DoubleVert(InParticles.P(Indices[i][j])), FColor::Green, false, Dt, 0, ParticleThickness);
-						Chaos::FDebugDrawQueue::GetInstance().DrawDebugLine(DoubleVert(InParticles.P(Indices[i][j])), DoubleVert(InParticles.P(Indices[i][(j + 1) % SecondIndices[i].Num()])), FColor::Green, false, Dt, 0, LineThickness);
-					}
-				}
-
-				Chaos::FDebugDrawQueue::GetInstance().DrawDebugLine(DoubleVert(SourcePos), DoubleVert(TargetPos), FColor::Yellow, false, Dt, 0, LineThickness);
-			}
-#endif
-		}
+		void VisualizeAllBindings(const FSolverParticlesRange& InParticles, const FSolverReal Dt) const;
 
+		void VisualizeAllBindings(const FSolverParticlesRange& InParticles,
+			const TFunctionRef<void(const FVector& Pos0, const FVector& Pos1, const FLinearColor& Color)>& DrawLine,
+			const TFunctionRef<void(const FVector& Pos, const FLinearColor& Color)>& DrawPoint) const;
 
-	protected:
-		TArray<TArray<int32>> Indices;
-		TArray<TArray<FSolverReal>> Weights;
-		TArray<TVector<FSolverReal, 3>> Constraints;
-		TArray<TArray<int32>> SecondIndices;
-		TArray<TArray<FSolverReal>> SecondWeights;
-		TArray<FSolverReal> Stiffness;
 		TArray<TArray<int32>> ConstraintsPerColor;
 		mutable TArray<FSolverReal> LambdaArray;
 		bool VisualizeBindings = false;
@@ -389,7 +141,7 @@ namespace Chaos::Softs
 	};
 
 	using FDeformableXPBDWeakConstraintParams = FDeformableXPBDWeightedSpringConstraintParams;
-	
+
 	template <typename T, typename ParticleType>
 	using FXPBDWeakConstraints UE_DEPRECATED(5.8, "Deprecated. this class is to be deleted, use FXPBDWeightedSpringConstraints instead") = FXPBDWeightedSpringConstraints;
 
