Dataflow:
- Added more functionality to PointsFacade

#Synced-CL 51837085

diff --git a/Engine/Source/Runtime/Experimental/Chaos/Private/GeometryCollection/Facades/PointsFacade.cpp b/Engine/Source/Runtime/Experimental/Chaos/Private/GeometryCollection/Facades/PointsFacade.cpp
--- a/Engine/Source/Runtime/Experimental/Chaos/Private/GeometryCollection/Facades/PointsFacade.cpp
+++ b/Engine/Source/Runtime/Experimental/Chaos/Private/GeometryCollection/Facades/PointsFacade.cpp
@@ -23,6 +23,12 @@ namespace GeometryCollection::Facades
 
 	/* ------------------------------------------------------------------------------------------------------------- */
 
+	const FManagedArrayCollection& FPointsFacade::GetConstCollection() const
+	{
+		return ConstCollection;
+	}
+	/* ------------------------------------------------------------------------------------------------------------- */
+
 	bool FPointsFacade::IsValid() const
 	{
 		if (ConstCollection.HasGroup(FPointsFacade::PointsGroup) &&
@@ -160,121 +166,201 @@ namespace GeometryCollection::Facades
 
 	/* ------------------------------------------------------------------------------------------------------------- */
 
-	void FPointsFacade::AddFloatAttribute(const FName InAttribute, const TArray<float>& InValues)
+	void FPointsFacade::DeletePoints(const TArray<int32>& InSortedDeletionList)
+	{
+		if (ensure(MutableCollection))
+		{
+			MutableCollection->RemoveElements(FPointsFacade::PointsGroup, InSortedDeletionList);
+		}
+	}
+
+	/* ------------------------------------------------------------------------------------------------------------- */
+
+	bool FPointsFacade::AddFloatAttribute(const FName InAttribute, const TArray<float>& InValues)
 	{
 		if (ensure(MutableCollection))
 		{
 			if (IsValid())
 			{
 				const TManagedArray<FVector3f>& Point = MutableCollection->GetAttribute<FVector3f>(FPointsFacade::PointAttribute, FPointsFacade::PointsGroup);
-				const int32 NumPoints = Point.Num();
+				const int32 NumPoints = FMath::Min(Point.Num(), InValues.Num());
 
 				if (NumPoints == InValues.Num())
 				{
-					TManagedArray<float>& Attr = MutableCollection->AddAttribute<float>(InAttribute, FPointsFacade::PointsGroup);
-
-					for (int32 Idx = 0; Idx < NumPoints; ++Idx)
+					TManagedArray<float>* Attr = MutableCollection->FindAttributeTyped<float>(InAttribute, FPointsFacade::PointsGroup);
+					if (!Attr)
 					{
-						Attr[Idx] = InValues[Idx];
+						TManagedArray<float>& AttrArr = MutableCollection->AddAttribute<float>(InAttribute, FPointsFacade::PointsGroup);
+
+						for (int32 Idx = 0; Idx < NumPoints; ++Idx)
+						{
+							AttrArr[Idx] = InValues[Idx];
+						}
+
+						return true;
 					}
 				}
 			}
 		}
+
+		return false;
 	}
 
 	/* ------------------------------------------------------------------------------------------------------------- */
 
-	void FPointsFacade::AddBoolAttribute(const FName InAttribute, const TArray<bool>& InValues)
+	bool FPointsFacade::AddBoolAttribute(const FName InAttribute, const TArray<bool>& InValues)
 	{
 		if (ensure(MutableCollection))
 		{
 			if (IsValid())
 			{
 				const TManagedArray<FVector3f>& Point = MutableCollection->GetAttribute<FVector3f>(FPointsFacade::PointAttribute, FPointsFacade::PointsGroup);
-				const int32 NumPoints = Point.Num();
+				const int32 NumPoints = FMath::Min(Point.Num(), InValues.Num());
 
 				if (NumPoints == InValues.Num())
 				{
-					TManagedArray<bool>& Attr = MutableCollection->AddAttribute<bool>(InAttribute, FPointsFacade::PointsGroup);
-
-					for (int32 Idx = 0; Idx < NumPoints; ++Idx)
+					TManagedArray<bool>* Attr = MutableCollection->FindAttributeTyped<bool>(InAttribute, FPointsFacade::PointsGroup);
+					if (!Attr)
 					{
-						Attr[Idx] = InValues[Idx];
+						TManagedArray<bool>& AttrArr = MutableCollection->AddAttribute<bool>(InAttribute, FPointsFacade::PointsGroup);
+
+						for (int32 Idx = 0; Idx < NumPoints; ++Idx)
+						{
+							AttrArr[Idx] = InValues[Idx];
+						}
+
+						return true;
 					}
 				}
 			}
 		}
+
+		return false;
 	}
 
 	/* ------------------------------------------------------------------------------------------------------------- */
 
-	void FPointsFacade::AddIntAttribute(const FName InAttribute, const TArray<int32>& InValues)
+	bool FPointsFacade::AddBoolAttribute(const FName InAttribute, const TBitArray<>& InValues)
 	{
 		if (ensure(MutableCollection))
 		{
 			if (IsValid())
 			{
 				const TManagedArray<FVector3f>& Point = MutableCollection->GetAttribute<FVector3f>(FPointsFacade::PointAttribute, FPointsFacade::PointsGroup);
-				const int32 NumPoints = Point.Num();
+				const int32 NumPoints = FMath::Min(Point.Num(), InValues.Num());
 
 				if (NumPoints == InValues.Num())
 				{
-					TManagedArray<int32>& Attr = MutableCollection->AddAttribute<int32>(InAttribute, FPointsFacade::PointsGroup);
-
-					for (int32 Idx = 0; Idx < NumPoints; ++Idx)
+					TManagedArray<bool>* Attr = MutableCollection->FindAttributeTyped<bool>(InAttribute, FPointsFacade::PointsGroup);
+					if (!Attr)
 					{
-						Attr[Idx] = InValues[Idx];
+						TManagedArray<bool>& AttrArr = MutableCollection->AddAttribute<bool>(InAttribute, FPointsFacade::PointsGroup);
+
+						for (int32 Idx = 0; Idx < NumPoints; ++Idx)
+						{
+							AttrArr[Idx] = InValues[Idx];
+						}
+
+						return true;
 					}
 				}
 			}
 		}
+
+		return false;
 	}
 
 	/* ------------------------------------------------------------------------------------------------------------- */
 
-	void FPointsFacade::AddVector2Attribute(const FName InAttribute, const TArray<FVector2f>& InValues)
+	bool FPointsFacade::AddIntAttribute(const FName InAttribute, const TArray<int32>& InValues)
 	{
 		if (ensure(MutableCollection))
 		{
 			if (IsValid())
 			{
 				const TManagedArray<FVector3f>& Point = MutableCollection->GetAttribute<FVector3f>(FPointsFacade::PointAttribute, FPointsFacade::PointsGroup);
-				const int32 NumPoints = Point.Num();
+				const int32 NumPoints = FMath::Min(Point.Num(), InValues.Num());
 
 				if (NumPoints == InValues.Num())
 				{
-					TManagedArray<FVector2f>& Attr = MutableCollection->AddAttribute<FVector2f>(InAttribute, FPointsFacade::PointsGroup);
-
-					for (int32 Idx = 0; Idx < NumPoints; ++Idx)
+					TManagedArray<int32>* Attr = MutableCollection->FindAttributeTyped<int32>(InAttribute, FPointsFacade::PointsGroup);
+					if (!Attr)
 					{
-						Attr[Idx] = InValues[Idx];
+						TManagedArray<int32>& AttrArr = MutableCollection->AddAttribute<int32>(InAttribute, FPointsFacade::PointsGroup);
+
+						for (int32 Idx = 0; Idx < NumPoints; ++Idx)
+						{
+							AttrArr[Idx] = InValues[Idx];
+						}
+
+						return true;
 					}
 				}
 			}
 		}
+
+		return false;
 	}
 
 	/* ------------------------------------------------------------------------------------------------------------- */
 
-	void FPointsFacade::AddVector3Attribute(const FName InAttribute, const TArray<FVector3f>& InValues)
+	bool FPointsFacade::AddVector2Attribute(const FName InAttribute, const TArray<FVector2f>& InValues)
 	{
 		if (ensure(MutableCollection))
 		{
 			if (IsValid())
 			{
 				const TManagedArray<FVector3f>& Point = MutableCollection->GetAttribute<FVector3f>(FPointsFacade::PointAttribute, FPointsFacade::PointsGroup);
-				const int32 NumPoints = Point.Num();
+				const int32 NumPoints = FMath::Min(Point.Num(), InValues.Num());
+
+				if (NumPoints == InValues.Num())
+				{
+					TManagedArray<FVector2f>* Attr = MutableCollection->FindAttributeTyped<FVector2f>(InAttribute, FPointsFacade::PointsGroup);
+					if (!Attr)
+					{
+						TManagedArray<FVector2f>& AttrArr = MutableCollection->AddAttribute<FVector2f>(InAttribute, FPointsFacade::PointsGroup);
 
-				const int32 Num = FMath::Min(NumPoints, InValues.Num());
+						for (int32 Idx = 0; Idx < NumPoints; ++Idx)
+						{
+							AttrArr[Idx] = InValues[Idx];
+						}
 
-				TManagedArray<FVector3f>& Attr = MutableCollection->AddAttribute<FVector3f>(InAttribute, FPointsFacade::PointsGroup);
+						return true;
+					}
+				}
+			}
+		}
 
-				for (int32 Idx = 0; Idx < Num; ++Idx)
+		return false;
+	}
+
+	/* ------------------------------------------------------------------------------------------------------------- */
+
+	bool FPointsFacade::AddVector3Attribute(const FName InAttribute, const TArray<FVector3f>& InValues)
+	{
+		if (ensure(MutableCollection))
+		{
+			if (IsValid())
+			{
+				const TManagedArray<FVector3f>& Point = MutableCollection->GetAttribute<FVector3f>(FPointsFacade::PointAttribute, FPointsFacade::PointsGroup);
+				const int32 NumPoints = FMath::Min(Point.Num(), InValues.Num());
+
+				TManagedArray<FVector3f>* Attr = MutableCollection->FindAttributeTyped<FVector3f>(InAttribute, FPointsFacade::PointsGroup);
+				if (!Attr)
 				{
-					Attr[Idx] = InValues[Idx];
+					TManagedArray<FVector3f>& AttrArr = MutableCollection->AddAttribute<FVector3f>(InAttribute, FPointsFacade::PointsGroup);
+
+					for (int32 Idx = 0; Idx < NumPoints; ++Idx)
+					{
+						AttrArr[Idx] = InValues[Idx];
+					}
+
+					return true;
 				}
 			}
 		}
+
+		return false;
 	}
 
 	/* ------------------------------------------------------------------------------------------------------------- */
diff --git a/Engine/Source/Runtime/Experimental/Chaos/Public/GeometryCollection/Facades/PointsFacade.h b/Engine/Source/Runtime/Experimental/Chaos/Public/GeometryCollection/Facades/PointsFacade.h
--- a/Engine/Source/Runtime/Experimental/Chaos/Public/GeometryCollection/Facades/PointsFacade.h
+++ b/Engine/Source/Runtime/Experimental/Chaos/Public/GeometryCollection/Facades/PointsFacade.h
@@ -42,6 +42,9 @@ namespace GeometryCollection::Facades
 		FPointsFacade(const FPointsFacade&) = delete;
 		FPointsFacade& operator= (const FPointsFacade&) = delete;
 
+		// Get Access to the ConstCollection
+		CHAOS_API const FManagedArrayCollection& GetConstCollection() const;
+
 		/** Is the facade defined constant. */
 		bool IsConst() const { return MutableCollection == nullptr; }
 
@@ -73,14 +76,20 @@ namespace GeometryCollection::Facades
 		CHAOS_API void AddPoints(const TArray<FVector3f>& InPoints);
 		CHAOS_API void AddPoints(const TArray<FVector>& InPoints);
 
+		/**
+		* Delete points from the collection
+		*/
+		CHAOS_API void DeletePoints(const TArray<int32>& InSortedDeletionList);
+
 		/**
 		* Add attribute to the points
 		*/
-		CHAOS_API void AddFloatAttribute(const FName InAttribute, const TArray<float>& InValues);
-		CHAOS_API void AddBoolAttribute(const FName InAttribute, const TArray<bool>& InValues);
-		CHAOS_API void AddIntAttribute(const FName InAttribute, const TArray<int32>& InValues);
-		CHAOS_API void AddVector2Attribute(const FName InAttribute, const TArray<FVector2f>& InValues);
-		CHAOS_API void AddVector3Attribute(const FName InAttribute, const TArray<FVector3f>& InValues);
+		CHAOS_API bool AddFloatAttribute(const FName InAttribute, const TArray<float>& InValues);
+		CHAOS_API bool AddBoolAttribute(const FName InAttribute, const TArray<bool>& InValues);
+		CHAOS_API bool AddBoolAttribute(const FName InAttribute, const TBitArray<>& InValues);
+		CHAOS_API bool AddIntAttribute(const FName InAttribute, const TArray<int32>& InValues);
+		CHAOS_API bool AddVector2Attribute(const FName InAttribute, const TArray<FVector2f>& InValues);
+		CHAOS_API bool AddVector3Attribute(const FName InAttribute, const TArray<FVector3f>& InValues);
 
 		/**
 		* Check if attribute exist
