HeadlessChaosFlesh: Fix access to protected ConstraintsData member via public GetConstraintsData() accessor

#Synced-CL 51662528

#Synced-CL 51662528

diff --git a/Engine/Source/Runtime/Experimental/Chaos/Public/Chaos/Deformable/GaussSeidelWeakConstraints.h b/Engine/Source/Runtime/Experimental/Chaos/Public/Chaos/Deformable/GaussSeidelWeakConstraints.h
--- a/Engine/Source/Runtime/Experimental/Chaos/Public/Chaos/Deformable/GaussSeidelWeakConstraints.h
+++ b/Engine/Source/Runtime/Experimental/Chaos/Public/Chaos/Deformable/GaussSeidelWeakConstraints.h
@@ -172,6 +172,7 @@ namespace Chaos::Softs
 		{
 			return InitialWCSize;
 		}
+		const FGaussSeidelSpringConstraintData& GetConstraintsData() const { return ConstraintsData; }
 
 	protected:
 		CHAOS_API virtual void SetConstraintsRestLength(const FSolverParticlesRange& InParticles, const FSolverParticlesRange& InSecondParticles);
