[HWRT] Fix crash when using static ray tracing instances for geometry collections.
- FGeometryCollectionSceneProxyBase holds an array view of an array in which was being cleared on the game thread.

- Changes:
- Moved initializing/emptying PartGeometries to the render thread.
- Store all render thread data in a new struct FGeometryCollectionRayTracingResources.
- Modified proxy to use that new struct and only access it on the render thread.

#Synced-CL 51889772

diff --git a/Engine/Source/Runtime/Experimental/GeometryCollectionEngine/Private/GeometryCollection/GeometryCollectionRenderData.cpp b/Engine/Source/Runtime/Experimental/GeometryCollectionEngine/Private/GeometryCollection/GeometryCollectionRenderData.cpp
--- a/Engine/Source/Runtime/Experimental/GeometryCollectionEngine/Private/GeometryCollection/GeometryCollectionRenderData.cpp
+++ b/Engine/Source/Runtime/Experimental/GeometryCollectionEngine/Private/GeometryCollection/GeometryCollectionRenderData.cpp
@@ -794,10 +794,8 @@ void FGeometryCollectionRenderData::InitResources(UGeometryCollection const& Own
 		MeshResource.InitResources(Owner);
 
 #if RHI_RAYTRACING
-		if (IsRayTracingAllowed() && Owner.bSupportRayTracing && MeshDescription.NumVertices != 0)
+		if (IsRayTracingAllowed() && Owner.bSupportRayTracing && MeshDescription.NumVertices > 0)
 		{
-			PartRayTracingGeometries.SetNum(MeshDescription.SubSections.Num());
-
 			FName OwnerName = Owner.GetFName();
 
 			ENQUEUE_RENDER_COMMAND(InitGeometryCollectiomRayTracingGeometry)(
@@ -805,11 +803,13 @@ void FGeometryCollectionRenderData::InitResources(UGeometryCollection const& Own
 				{
 					// Create a single ray tracing geometry group to be used by parts
 
-					check(RayTracingGeometryGroupHandle.IsNull());
-					RayTracingGeometryGroupHandle = GRayTracingGeometryManager->RegisterRayTracingGeometryGroup(/*NumLODs*/ 0, /*CurrentFirstLODIdx*/ 0);
+					check(RayTracingResource.GroupHandle.IsNull());
+					RayTracingResource.GroupHandle = GRayTracingGeometryManager->RegisterRayTracingGeometryGroup(/*NumLODs*/ 0, /*CurrentFirstLODIdx*/ 0);
 
 					// TODO: Combine sections using the same transform into a single BLAS
 
+					RayTracingResource.PartGeometries.SetNum(MeshDescription.SubSections.Num());
+
 					for (int32 SectionIndex = 0; SectionIndex < MeshDescription.SubSections.Num(); ++SectionIndex)
 					{
 						FRayTracingGeometryInitializer Initializer;
@@ -834,10 +834,10 @@ void FGeometryCollectionRenderData::InitResources(UGeometryCollection const& Own
 							Initializer.TotalPrimitiveCount += Section.TriangleCount;
 						}
 
-						PartRayTracingGeometries[SectionIndex].GroupHandle = RayTracingGeometryGroupHandle;
-						PartRayTracingGeometries[SectionIndex].LODIndex = FRayTracingGeometry::AdditionalGeometryLODIndex;
-						PartRayTracingGeometries[SectionIndex].SetInitializer(MoveTemp(Initializer));
-						PartRayTracingGeometries[SectionIndex].InitResource(RHICmdList);
+						RayTracingResource.PartGeometries[SectionIndex].GroupHandle = RayTracingResource.GroupHandle;
+						RayTracingResource.PartGeometries[SectionIndex].LODIndex = FRayTracingGeometry::AdditionalGeometryLODIndex;
+						RayTracingResource.PartGeometries[SectionIndex].SetInitializer(MoveTemp(Initializer));
+						RayTracingResource.PartGeometries[SectionIndex].InitResource(RHICmdList);
 					}
 				}
 				);
@@ -862,22 +862,22 @@ void FGeometryCollectionRenderData::ReleaseResources()
 
 	if (bHasMeshData)
 	{
-		for (FRayTracingGeometry& PartRayTracingGeometry : PartRayTracingGeometries)
-		{
-			PartRayTracingGeometry.ReleaseResource();
-		}
-		PartRayTracingGeometries.Empty();
-
 #if RHI_RAYTRACING
 		if (IsRayTracingAllowed())
 		{
 			ENQUEUE_RENDER_COMMAND(CmdReleaseRayTracingGeometryGroup)(
 				[this](FRHICommandListImmediate&)
 				{
-					if (!RayTracingGeometryGroupHandle.IsNull())
+					for (FRayTracingGeometry& PartRayTracingGeometry : RayTracingResource.PartGeometries)
+					{
+						PartRayTracingGeometry.ReleaseResource();
+					}
+					RayTracingResource.PartGeometries.Empty();
+
+					if (!RayTracingResource.GroupHandle.IsNull())
 					{
-						GRayTracingGeometryManager->ReleaseRayTracingGeometryGroup(RayTracingGeometryGroupHandle);
-						RayTracingGeometryGroupHandle = {};
+						GRayTracingGeometryManager->ReleaseRayTracingGeometryGroup(RayTracingResource.GroupHandle);
+						RayTracingResource.GroupHandle = {};
 					}
 				});
 		}
diff --git a/Engine/Source/Runtime/Experimental/GeometryCollectionEngine/Private/GeometryCollection/GeometryCollectionRenderData.h b/Engine/Source/Runtime/Experimental/GeometryCollectionEngine/Private/GeometryCollection/GeometryCollectionRenderData.h
--- a/Engine/Source/Runtime/Experimental/GeometryCollectionEngine/Private/GeometryCollection/GeometryCollectionRenderData.h
+++ b/Engine/Source/Runtime/Experimental/GeometryCollectionEngine/Private/GeometryCollection/GeometryCollectionRenderData.h
@@ -106,6 +106,15 @@ struct FGeometryCollectionMeshDescription
 	void Serialize(FArchive& Ar);
 };
 
+#if RHI_RAYTRACING
+/** Set of resources required to ray trace a geometry collection. */
+struct FGeometryCollectionRayTracingResources
+{
+	RayTracing::FGeometryGroupHandle GroupHandle;
+	TArray<FRayTracingGeometry> PartGeometries;
+};
+#endif
+
 /** Built render data compiled and ready for use by the geometry collection scene proxy. */
 class FGeometryCollectionRenderData
 {
@@ -158,9 +167,8 @@ public:
 	FBoxSphereBounds PreSkinnedBounds = FBoxSphereBounds(ForceInitToZero);
 
 #if RHI_RAYTRACING
-	RayTracing::FGeometryGroupHandle RayTracingGeometryGroupHandle;
+	FGeometryCollectionRayTracingResources RayTracingResource;
 #endif
-	TArray<FRayTracingGeometry> PartRayTracingGeometries;
 
 private:
 	bool bIsInitialized = false;
diff --git a/Engine/Source/Runtime/Experimental/GeometryCollectionEngine/Private/GeometryCollection/GeometryCollectionSceneProxy.cpp b/Engine/Source/Runtime/Experimental/GeometryCollectionEngine/Private/GeometryCollection/GeometryCollectionSceneProxy.cpp
--- a/Engine/Source/Runtime/Experimental/GeometryCollectionEngine/Private/GeometryCollection/GeometryCollectionSceneProxy.cpp
+++ b/Engine/Source/Runtime/Experimental/GeometryCollectionEngine/Private/GeometryCollection/GeometryCollectionSceneProxy.cpp
@@ -164,8 +164,8 @@ FGeometryCollectionSceneProxyBase::FGeometryCollectionSceneProxyBase(UGeometryCo
 	, bUseShaderBoneTransform(VertexFactory.UseShaderBoneTransform(Component->GetScene()->GetShaderPlatform()))
 	, bSupportsTripleBufferVertexUpload(GRHISupportsMapWriteNoOverwrite)
 #if RHI_RAYTRACING
-	, bHasRayTracingRepresentation(GRayTracingGeometryCollection && IsRayTracingEnabled() && Component->GetRestCollection()->RenderData->PartRayTracingGeometries.Num() > 0)
-	, RayTracingGeometryGroupHandle(Component->GetRestCollection()->RenderData->RayTracingGeometryGroupHandle)
+	, bHasRayTracingRepresentation(GRayTracingGeometryCollection && IsRayTracingEnabled() && Component->GetRestCollection()->bSupportRayTracing && Component->GetRestCollection()->RenderData->MeshDescription.NumVertices > 0)
+	, RayTracingResources(Component->GetRestCollection()->RenderData->RayTracingResource)
 #endif
 {
 	if (!bIsNanite || bHasRayTracingRepresentation)
@@ -193,13 +193,6 @@ FGeometryCollectionSceneProxyBase::FGeometryCollectionSceneProxyBase(UGeometryCo
 		}
 	}
 
-#if RHI_RAYTRACING
-	if (bHasRayTracingRepresentation)
-	{
-		PartRayTracingGeometries = Component->GetRestCollection()->RenderData->PartRayTracingGeometries;
-	}
-#endif
-
 	PreSkinnedBounds = Component->GetRestCollection()->RenderData->PreSkinnedBounds;
 
 	DynamicData = Component->InitDynamicData(true);
@@ -706,14 +699,14 @@ void FGeometryCollectionSceneProxyBase::GetDynamicRayTracingInstances(FRayTracin
 			}
 		}
 
-		checkf(PartRayTracingGeometries.Num() == SectionArray.Num(), TEXT("Expected to have one part per section (%d vs %d)."), PartRayTracingGeometries.Num(), SectionArray.Num());
+		checkf(RayTracingResources.PartGeometries.Num() == SectionArray.Num(), TEXT("Expected to have one part per section (%d vs %d)."), RayTracingResources.PartGeometries.Num(), SectionArray.Num());
 
-		for (int32 SectionIndex = 0; SectionIndex < PartRayTracingGeometries.Num(); ++SectionIndex)
+		for (int32 SectionIndex = 0; SectionIndex < RayTracingResources.PartGeometries.Num(); ++SectionIndex)
 		{
 			const FGeometryCollectionMeshElement& Section = SectionArray[SectionIndex];
 
 			FRayTracingInstance RayTracingInstance;
-			RayTracingInstance.Geometry = &PartRayTracingGeometries[SectionIndex];
+			RayTracingInstance.Geometry = &RayTracingResources.PartGeometries[SectionIndex];
 			RayTracingInstance.InstanceTransformsView = MakeConstArrayView(&PartInstanceTransforms[Section.TransformIndex], 1);
 			RayTracingInstance.MaterialsView = MakeConstArrayView(&CachedPartRayTracingMaterials[SectionIndex], 1);
 			RayTracingInstance.bInstanceMaskAndFlagsDirty = bInstanceMaskAndFlagsDirty;
@@ -747,18 +740,18 @@ void FGeometryCollectionSceneProxyBase::GetStaticRayTracingInstances(FStaticRayT
 	FMaterialArrayType MaterialProxies;
 	GetSectionMaterialProxies(SectionArray, MaterialProxies);
 
-	checkf(PartRayTracingGeometries.Num() == SectionArray.Num(), TEXT("Expected to have one part per section (%d vs %d)."), PartRayTracingGeometries.Num(), SectionArray.Num());
+	checkf(RayTracingResources.PartGeometries.Num() == SectionArray.Num(), TEXT("Expected to have one part per section (%d vs %d)."), RayTracingResources.PartGeometries.Num(), SectionArray.Num());
 
 	PrimitiveInstanceIndices.Reset();
-	PrimitiveInstanceIndices.SetNumUninitialized(PartRayTracingGeometries.Num());
+	PrimitiveInstanceIndices.SetNumUninitialized(RayTracingResources.PartGeometries.Num());
 
 	// One ray tracing LOD with N instances (one per part)
 
 	FStaticRayTracingInstancesLOD& StaticRayTracingInstancesLOD = OutRayTracingInstances.LODs.AddDefaulted_GetRef();
 
-	for (int32 SectionIndex = 0; SectionIndex < PartRayTracingGeometries.Num(); ++SectionIndex)
+	for (int32 SectionIndex = 0; SectionIndex < RayTracingResources.PartGeometries.Num(); ++SectionIndex)
 	{
-		if (!PartRayTracingGeometries[SectionIndex].IsValid() || PartRayTracingGeometries[SectionIndex].IsEvicted())
+		if (!RayTracingResources.PartGeometries[SectionIndex].IsValid() || RayTracingResources.PartGeometries[SectionIndex].IsEvicted())
 		{
 			continue;
 		}
@@ -775,7 +768,7 @@ void FGeometryCollectionSceneProxyBase::GetStaticRayTracingInstances(FStaticRayT
 		PrimitiveInstanceIndices[SectionIndex] = InstanceIndex;
 
 		FRayTracingInstance& RayTracingInstance = StaticRayTracingInstancesLOD.Instances.AddDefaulted_GetRef();
-		RayTracingInstance.Geometry = &PartRayTracingGeometries[SectionIndex];
+		RayTracingInstance.Geometry = &RayTracingResources.PartGeometries[SectionIndex];
 		RayTracingInstance.NumTransforms = 1;
 		RayTracingInstance.PrimitiveInstanceIndicesView = MakeConstArrayView(&PrimitiveInstanceIndices[SectionIndex], 1); // parts are stored in GPU Scene as instances
 
@@ -787,7 +780,7 @@ void FGeometryCollectionSceneProxyBase::GetStaticRayTracingInstances(FStaticRayT
 RayTracing::FGeometryGroupHandle FGeometryCollectionSceneProxyBase::GetRayTracingGeometryGroupHandle() const
 {
 	check(IsInRenderingThread() || IsInParallelRenderingThread());
-	return RayTracingGeometryGroupHandle;
+	return RayTracingResources.GroupHandle;
 }
 
 #endif
@@ -797,7 +790,7 @@ uint32 FGeometryCollectionSceneProxyBase::GetAllocatedSize() const
 #if RHI_RAYTRACING
 	uint32 RayTracingGeometryAllocatedSize = RayTracingGeometry.RawData.GetAllocatedSize();
 
-	for (const FRayTracingGeometry& PartRayTracingGeometry : PartRayTracingGeometries)
+	for (const FRayTracingGeometry& PartRayTracingGeometry : RayTracingResources.PartGeometries)
 	{
 		RayTracingGeometryAllocatedSize += PartRayTracingGeometry.RawData.GetAllocatedSize();
 	}
diff --git a/Engine/Source/Runtime/Experimental/GeometryCollectionEngine/Private/GeometryCollection/GeometryCollectionSceneProxy.h b/Engine/Source/Runtime/Experimental/GeometryCollectionEngine/Private/GeometryCollection/GeometryCollectionSceneProxy.h
--- a/Engine/Source/Runtime/Experimental/GeometryCollectionEngine/Private/GeometryCollection/GeometryCollectionSceneProxy.h
+++ b/Engine/Source/Runtime/Experimental/GeometryCollectionEngine/Private/GeometryCollection/GeometryCollectionSceneProxy.h
@@ -245,14 +245,13 @@ private:
 	const bool bHasRayTracingRepresentation = false;
 
 #if RHI_RAYTRACING
+	const FGeometryCollectionRayTracingResources& RayTracingResources;
+
 	FRayTracingGeometry RayTracingGeometry;
 	FRWBuffer RayTracingDynamicVertexBuffer;
 
-	TConstArrayView<FRayTracingGeometry> PartRayTracingGeometries;
 	TArray<FMeshBatch> CachedPartRayTracingMaterials;
 	TArray<uint32> PrimitiveInstanceIndices;
-
-	RayTracing::FGeometryGroupHandle RayTracingGeometryGroupHandle;
 #endif
 
 	bool bRenderResourcesCreated = false;
