add a dataflow node to merge dataflowskeletons

#Synced-CL 51921485

diff --git a/Engine/Source/Runtime/Experimental/Dataflow/Engine/Private/Dataflow/DataflowEngineModule.cpp b/Engine/Source/Runtime/Experimental/Dataflow/Engine/Private/Dataflow/DataflowEngineModule.cpp
--- a/Engine/Source/Runtime/Experimental/Dataflow/Engine/Private/Dataflow/DataflowEngineModule.cpp
+++ b/Engine/Source/Runtime/Experimental/Dataflow/Engine/Private/Dataflow/DataflowEngineModule.cpp
@@ -9,6 +9,7 @@
 #include "Dataflow/DataflowAssetPathNodes.h"
 #include "Dataflow/DataflowVariableNodes.h"
 #include "Dataflow/DataflowSubGraphNodes.h"
+#include "Dataflow/DataflowSkeletonNodes.h"
 
 /**
  * The public interface to this module
@@ -25,6 +26,7 @@ public:
 		UE::Dataflow::RegisterAssetPathNodes();
 		UE::Dataflow::RegisterVariableNodes();
 		UE::Dataflow::RegisterSubGraphNodes();
+		UE::Dataflow::RegisterSkeletonNodes();
 	}
 
 	virtual void ShutdownModule() override { }
diff --git a/Engine/Source/Runtime/Experimental/Dataflow/Engine/Private/Dataflow/DataflowSkeletonNodes.cpp b/Engine/Source/Runtime/Experimental/Dataflow/Engine/Private/Dataflow/DataflowSkeletonNodes.cpp
new file mode 100644
--- /dev/null
+++ b/Engine/Source/Runtime/Experimental/Dataflow/Engine/Private/Dataflow/DataflowSkeletonNodes.cpp
@@ -0,0 +1,93 @@
+// Copyright Epic Games, Inc. All Rights Reserved.
+
+#include "Dataflow/DataflowSkeletonNodes.h"
+
+#include "BoneIndices.h"
+#include "ReferenceSkeleton.h"
+#include "Dataflow/DataflowNodeFactory.h"
+
+#include UE_INLINE_GENERATED_CPP_BY_NAME(DataflowSkeletonNodes)
+
+namespace UE::Dataflow
+{
+	void RegisterSkeletonNodes()
+	{
+		DATAFLOW_NODE_REGISTER_CREATION_FACTORY(FMergeSkeletonsDataflowNode);
+	}
+}
+
+FMergeSkeletonsDataflowNode::FMergeSkeletonsDataflowNode(
+	const UE::Dataflow::FNodeParameters& InParam, FGuid InGuid)
+	: FDataflowNode(InParam, InGuid)
+{
+	RegisterInputConnection(&Skeletons);
+	RegisterInputConnection(&SharedRootPosition).SetCanHidePin(true).SetPinIsHidden(true);
+	RegisterInputConnection(&bPreserveBoneNames).SetCanHidePin(true).SetPinIsHidden(true);
+	RegisterOutputConnection(&MergedSkeleton);
+}
+
+void FMergeSkeletonsDataflowNode::Evaluate(
+	UE::Dataflow::FContext& Context, const FDataflowOutput* Out) const
+{
+	if (Out->IsA(&MergedSkeleton))
+	{
+		const TArray<FDataflowSkeleton>& InSkeletons = GetValue(Context, &Skeletons);
+		const FVector InRootPos = GetValue(Context, &SharedRootPosition);
+		const bool bInPreserveBoneNames = GetValue(Context, &bPreserveBoneNames);
+
+		FDataflowSkeleton OutSkeleton;
+		FReferenceSkeletonModifier Modifier = OutSkeleton.ModifySkeleton();
+
+		const FTransform SharedRootTransform(InRootPos);
+		FName RootName("Root");
+		Modifier.Add(FMeshBoneInfo(RootName, RootName.ToString(), INDEX_NONE),
+		             SharedRootTransform);
+		TSet<FName> UsedBoneNames;
+		UsedBoneNames.Add(RootName);
+		auto SetUniqueBoneName = [&UsedBoneNames](FName& NameToMakeUnique) -> bool
+			{
+				int32 NameInt = NameToMakeUnique.GetNumber();
+				bool bWasModified = false;
+				while (UsedBoneNames.Contains(NameToMakeUnique))
+				{
+					bWasModified = true;
+					NameToMakeUnique.SetNumber(++NameInt);
+				}
+				UsedBoneNames.Add(NameToMakeUnique);
+				return bWasModified;
+			};
+
+		int32 BoneCount = 0;
+		for (const FDataflowSkeleton& SrcSkeleton : InSkeletons)
+		{
+			const FReferenceSkeleton& SrcRef = SrcSkeleton.GetRefSkeleton();
+			const TArray<FMeshBoneInfo>& SrcInfos = SrcRef.GetRawRefBoneInfo();
+			const TArray<FTransform>& SrcPoses = SrcRef.GetRawRefBonePose();
+
+			const int32 Offset = Modifier.GetReferenceSkeleton().GetRawBoneNum();
+
+			for (int32 BoneIdx = 0; BoneIdx < SrcInfos.Num(); ++BoneIdx)
+			{
+				FMeshBoneInfo BoneInfo = SrcInfos[BoneIdx];
+				bool bNameChanged = !bInPreserveBoneNames;
+				if (bInPreserveBoneNames)
+				{
+					bNameChanged = SetUniqueBoneName(BoneInfo.Name);
+				}
+				else
+				{
+					BoneInfo.Name = FName("Bone", BoneCount++);
+				}
+				if (bNameChanged)
+				{
+					BoneInfo.ExportName = BoneInfo.Name.ToString();
+				}
+				BoneInfo.ParentIndex = (BoneInfo.ParentIndex == INDEX_NONE) ? 0 : BoneInfo.ParentIndex + Offset;
+
+				Modifier.Add(BoneInfo, SrcPoses[BoneIdx]);
+			}
+		}
+
+		SetValue(Context, OutSkeleton, &MergedSkeleton);
+	}
+}
diff --git a/Engine/Source/Runtime/Experimental/Dataflow/Engine/Private/Dataflow/DataflowSkeletonNodes.h b/Engine/Source/Runtime/Experimental/Dataflow/Engine/Private/Dataflow/DataflowSkeletonNodes.h
new file mode 100644
--- /dev/null
+++ b/Engine/Source/Runtime/Experimental/Dataflow/Engine/Private/Dataflow/DataflowSkeletonNodes.h
@@ -0,0 +1,47 @@
+// Copyright Epic Games, Inc. All Rights Reserved.
+
+#pragma once
+
+#include "Dataflow/DataflowNode.h"
+#include "Dataflow/DataflowSkeleton.h"
+
+#include "DataflowSkeletonNodes.generated.h"
+
+/**
+ * Merges an array of skeletons into a single skeleton
+ */
+USTRUCT()
+struct FMergeSkeletonsDataflowNode : public FDataflowNode
+{
+	GENERATED_USTRUCT_BODY()
+	DATAFLOW_NODE_DEFINE_INTERNAL(FMergeSkeletonsDataflowNode, "MergeSkeletons", "Skeleton", "Append Merge Combine")
+
+private:
+	/** Skeletons to merge together */
+	UPROPERTY(EditAnywhere, Category = Skeleton, meta = (DataflowInput))
+	TArray<FDataflowSkeleton> Skeletons;
+
+	/** World-space position of the shared root bone added to the merged skeleton */
+	UPROPERTY(EditAnywhere, Category = Skeleton, meta = (DataflowInput))
+	FVector SharedRootPosition = FVector::ZeroVector;
+
+	/** Whether to attempt to preserve the names of the bones from the input skeleton. Otherwise, bones will be simply numbered. */
+	UPROPERTY(EditAnywhere, Category = Skeleton, meta = (DataflowInput))
+	bool bPreserveBoneNames = true;
+
+	/** The merged skeleton */
+	UPROPERTY(meta = (DataflowOutput))
+	FDataflowSkeleton MergedSkeleton;
+
+	virtual void Evaluate(UE::Dataflow::FContext& Context, const FDataflowOutput* Out) const override;
+
+public:
+	FMergeSkeletonsDataflowNode(const UE::Dataflow::FNodeParameters& InParam, FGuid InGuid = FGuid::NewGuid());
+};
+
+//////////////////////////////////////////////////////////////////////////////////////////////////////////////////////////
+
+namespace UE::Dataflow
+{
+	void RegisterSkeletonNodes();
+}
