[Backout] - CL51925510
#rnx
Original CL Desc
-----------------------------------------------------------------
add a dataflow node to merge dataflowskeletons

#Synced-CL 51925774

diff --git a/Engine/Source/Runtime/Experimental/Dataflow/Engine/Private/Dataflow/DataflowEngineModule.cpp b/Engine/Source/Runtime/Experimental/Dataflow/Engine/Private/Dataflow/DataflowEngineModule.cpp
--- a/Engine/Source/Runtime/Experimental/Dataflow/Engine/Private/Dataflow/DataflowEngineModule.cpp
+++ b/Engine/Source/Runtime/Experimental/Dataflow/Engine/Private/Dataflow/DataflowEngineModule.cpp
@@ -9,7 +9,6 @@
 #include "Dataflow/DataflowAssetPathNodes.h"
 #include "Dataflow/DataflowVariableNodes.h"
 #include "Dataflow/DataflowSubGraphNodes.h"
-#include "Dataflow/DataflowSkeletonNodes.h"
 
 /**
  * The public interface to this module
@@ -26,7 +25,6 @@ public:
 		UE::Dataflow::RegisterAssetPathNodes();
 		UE::Dataflow::RegisterVariableNodes();
 		UE::Dataflow::RegisterSubGraphNodes();
-		UE::Dataflow::RegisterSkeletonNodes();
 	}
 
 	virtual void ShutdownModule() override { }
diff --git a/Engine/Source/Runtime/Experimental/Dataflow/Engine/Private/Dataflow/DataflowSkeletonNodes.cpp b/Engine/Source/Runtime/Experimental/Dataflow/Engine/Private/Dataflow/DataflowSkeletonNodes.cpp
deleted file mode 100644
diff --git a/Engine/Source/Runtime/Experimental/Dataflow/Engine/Private/Dataflow/DataflowSkeletonNodes.h b/Engine/Source/Runtime/Experimental/Dataflow/Engine/Private/Dataflow/DataflowSkeletonNodes.h
deleted file mode 100644
