Dataflow:
- Added FDataflowMeshSelectionTypes (FDataflowVertexSelection, FDataflowFaceSelection)
- Removed FDataflowPointSelection from FDataflowSelectionTypes

#Synced-CL 51928837

diff --git a/Engine/Source/Runtime/Experimental/Dataflow/Core/Private/Dataflow/DataflowAnyType.cpp b/Engine/Source/Runtime/Experimental/Dataflow/Core/Private/Dataflow/DataflowAnyType.cpp
--- a/Engine/Source/Runtime/Experimental/Dataflow/Core/Private/Dataflow/DataflowAnyType.cpp
+++ b/Engine/Source/Runtime/Experimental/Dataflow/Core/Private/Dataflow/DataflowAnyType.cpp
@@ -23,6 +23,7 @@ namespace UE::Dataflow
 		UE_DATAFLOW_REGISTER_ANYTYPE(FDataflowUObjectConvertibleTypes);
 		UE_DATAFLOW_REGISTER_ANYTYPE(FDataflowSelectionTypes);
 		UE_DATAFLOW_REGISTER_ANYTYPE(FDataflowRotationTypes);
+		UE_DATAFLOW_REGISTER_ANYTYPE(FDataflowMeshSelectionTypes);
 
 		// array types 
 		UE_DATAFLOW_REGISTER_ANYTYPE(FDataflowVectorArrayTypes);
diff --git a/Engine/Source/Runtime/Experimental/Dataflow/Core/Public/Dataflow/DataflowAnyType.h b/Engine/Source/Runtime/Experimental/Dataflow/Core/Public/Dataflow/DataflowAnyType.h
--- a/Engine/Source/Runtime/Experimental/Dataflow/Core/Public/Dataflow/DataflowAnyType.h
+++ b/Engine/Source/Runtime/Experimental/Dataflow/Core/Public/Dataflow/DataflowAnyType.h
@@ -253,6 +253,19 @@ struct FDataflowRotationTypes : public FDataflowAnyType
 	FRotator Value = FRotator(EForceInit::ForceInitToZero);
 };
 
+/** Mesh Selection types */
+USTRUCT()
+struct FDataflowMeshSelectionTypes : public FDataflowAnyType
+{
+	using FPolicyType = FDataflowMeshSelectionTypePolicy;
+	using FStorageType = FDataflowSelection;
+
+	GENERATED_USTRUCT_BODY()
+
+	UPROPERTY()
+	FDataflowSelection Value;
+};
+
 namespace UE::Dataflow
 {
 	DATAFLOWCORE_API bool AreTypesCompatible(FName TypeA, FName TypeB);
diff --git a/Engine/Source/Runtime/Experimental/Dataflow/Core/Public/Dataflow/DataflowTypePolicy.h b/Engine/Source/Runtime/Experimental/Dataflow/Core/Public/Dataflow/DataflowTypePolicy.h
--- a/Engine/Source/Runtime/Experimental/Dataflow/Core/Public/Dataflow/DataflowTypePolicy.h
+++ b/Engine/Source/Runtime/Experimental/Dataflow/Core/Public/Dataflow/DataflowTypePolicy.h
@@ -390,7 +390,7 @@ struct FDataflowStringTypePolicy :
 };
 
 struct FDataflowSelectionTypePolicy :
-	public TDataflowMultiTypePolicy<FDataflowTransformSelection, FDataflowVertexSelection, FDataflowFaceSelection, FDataflowGeometrySelection, FDataflowMaterialSelection, FDataflowCurveSelection, FDataflowPointsSelection>
+	public TDataflowMultiTypePolicy<FDataflowTransformSelection, FDataflowVertexSelection, FDataflowFaceSelection, FDataflowGeometrySelection, FDataflowMaterialSelection, FDataflowCurveSelection>
 {
 };
 
@@ -404,6 +404,11 @@ struct FDataflowRotationTypePolicy :
 {
 };
 
+struct FDataflowMeshSelectionTypePolicy :
+	public TDataflowMultiTypePolicy<FDataflowVertexSelection, FDataflowFaceSelection>
+{
+};
+
 /**
 * string comvertible types
 * - FString / Fname
