[HWRT] Fix ray tracing geometry group reference not being released when proxy is destroyed after FGeometryCollectionRenderData is released.
- also avoid issues if RHI is shutdown before FGeometryCollectionRenderData::ReleaseResources() is called.

#Synced-CL 51889772

diff --git a/Engine/Source/Runtime/Experimental/GeometryCollectionEngine/Private/GeometryCollection/GeometryCollectionRenderData.cpp b/Engine/Source/Runtime/Experimental/GeometryCollectionEngine/Private/GeometryCollection/GeometryCollectionRenderData.cpp
--- a/Engine/Source/Runtime/Experimental/GeometryCollectionEngine/Private/GeometryCollection/GeometryCollectionRenderData.cpp
+++ b/Engine/Source/Runtime/Experimental/GeometryCollectionEngine/Private/GeometryCollection/GeometryCollectionRenderData.cpp
@@ -863,24 +863,21 @@ void FGeometryCollectionRenderData::ReleaseResources()
 	if (bHasMeshData)
 	{
 #if RHI_RAYTRACING
-		if (IsRayTracingAllowed())
-		{
-			ENQUEUE_RENDER_COMMAND(CmdReleaseRayTracingGeometryGroup)(
-				[this](FRHICommandListImmediate&)
+		ENQUEUE_RENDER_COMMAND(CmdReleaseRayTracingGeometryGroup)(
+			[this](FRHICommandListImmediate&)
+			{
+				for (FRayTracingGeometry& PartRayTracingGeometry : RayTracingResource.PartGeometries)
 				{
-					for (FRayTracingGeometry& PartRayTracingGeometry : RayTracingResource.PartGeometries)
-					{
-						PartRayTracingGeometry.ReleaseResource();
-					}
-					RayTracingResource.PartGeometries.Empty();
+					PartRayTracingGeometry.ReleaseResource();
+				}
+				RayTracingResource.PartGeometries.Empty();
 
-					if (!RayTracingResource.GroupHandle.IsNull())
-					{
-						GRayTracingGeometryManager->ReleaseRayTracingGeometryGroup(RayTracingResource.GroupHandle);
-						RayTracingResource.GroupHandle = {};
-					}
-				});
-		}
+				if (!RayTracingResource.GroupHandle.IsNull())
+				{
+					GRayTracingGeometryManager->ReleaseRayTracingGeometryGroup(RayTracingResource.GroupHandle);
+					RayTracingResource.GroupHandle = {};
+				}
+			});
 #endif
 
 		MeshResource.ReleaseResources();
diff --git a/Engine/Source/Runtime/Experimental/GeometryCollectionEngine/Private/GeometryCollection/GeometryCollectionSceneProxy.cpp b/Engine/Source/Runtime/Experimental/GeometryCollectionEngine/Private/GeometryCollection/GeometryCollectionSceneProxy.cpp
--- a/Engine/Source/Runtime/Experimental/GeometryCollectionEngine/Private/GeometryCollection/GeometryCollectionSceneProxy.cpp
+++ b/Engine/Source/Runtime/Experimental/GeometryCollectionEngine/Private/GeometryCollection/GeometryCollectionSceneProxy.cpp
@@ -209,6 +209,14 @@ FGeometryCollectionSceneProxyBase::~FGeometryCollectionSceneProxyBase()
 
 void FGeometryCollectionSceneProxyBase::CreateRenderThreadResources(FRHICommandListBase& RHICmdList)
 {
+#if RHI_RAYTRACING
+	if (IsRayTracingAllowed() && bHasRayTracingRepresentation)
+	{
+		// copy RayTracingGeometryGroupHandle from FGeometryCollectionRayTracingResources since FGeometryCollectionRenderData can be released before the proxy is destroyed
+		RayTracingGeometryGroupHandle = RayTracingResources.GroupHandle;
+	}
+#endif
+
 	if (!bIsNanite || bHasRayTracingRepresentation)
 	{
 		if (bUseShaderBoneTransform)
@@ -780,7 +788,7 @@ void FGeometryCollectionSceneProxyBase::GetStaticRayTracingInstances(FStaticRayT
 RayTracing::FGeometryGroupHandle FGeometryCollectionSceneProxyBase::GetRayTracingGeometryGroupHandle() const
 {
 	check(IsInRenderingThread() || IsInParallelRenderingThread());
-	return RayTracingResources.GroupHandle;
+	return RayTracingGeometryGroupHandle;
 }
 
 #endif
diff --git a/Engine/Source/Runtime/Experimental/GeometryCollectionEngine/Private/GeometryCollection/GeometryCollectionSceneProxy.h b/Engine/Source/Runtime/Experimental/GeometryCollectionEngine/Private/GeometryCollection/GeometryCollectionSceneProxy.h
--- a/Engine/Source/Runtime/Experimental/GeometryCollectionEngine/Private/GeometryCollection/GeometryCollectionSceneProxy.h
+++ b/Engine/Source/Runtime/Experimental/GeometryCollectionEngine/Private/GeometryCollection/GeometryCollectionSceneProxy.h
@@ -246,6 +246,7 @@ private:
 
 #if RHI_RAYTRACING
 	const FGeometryCollectionRayTracingResources& RayTracingResources;
+	RayTracing::FGeometryGroupHandle RayTracingGeometryGroupHandle;
 
 	FRayTracingGeometry RayTracingGeometry;
 	FRWBuffer RayTracingDynamicVertexBuffer;
