UDataflowEdNode: Fix some functions that should be toolonly as they are toolonly in UEdGraphNode

#rnx
#Synced-CL 51944828

diff --git a/Engine/Source/Runtime/Experimental/Dataflow/Engine/Private/Dataflow/DataflowEdNode.cpp b/Engine/Source/Runtime/Experimental/Dataflow/Engine/Private/Dataflow/DataflowEdNode.cpp
--- a/Engine/Source/Runtime/Experimental/Dataflow/Engine/Private/Dataflow/DataflowEdNode.cpp
+++ b/Engine/Source/Runtime/Experimental/Dataflow/Engine/Private/Dataflow/DataflowEdNode.cpp
@@ -148,29 +148,7 @@ UEdGraphPin* UDataflowEdNode::CreateEdPin(const UE::Dataflow::FPin& Pin)
 }
 #endif
 
-void UDataflowEdNode::AllocateDefaultPins()
-{
-	UE_LOG(DATAFLOWNODE_LOG, Verbose, TEXT("UDataflowEdNode::AllocateDefaultPins()"));
-	// called on node creation from UI. 
-
-#if WITH_EDITOR && !(UE_BUILD_SHIPPING || UE_BUILD_TEST)
-	if (DataflowGraph)
-	{
-		if (DataflowNodeGuid.IsValid())
-		{
-			if (TSharedPtr<FDataflowNode> DataflowNode = DataflowGraph->FindBaseNode(DataflowNodeGuid))
-			{
-				for (const UE::Dataflow::FPin& Pin : DataflowNode->GetPins())
-				{
-					CreateEdPin(Pin);
-				}
-			}
-		}
-	}
 
-	UpdatePinsDefaultValuesFromNode();
-#endif // WITH_EDITOR && !UE_BUILD_SHIPPING
-}
 
 void UDataflowEdNode::UpdatePinsFromDataflowNode()
 {
@@ -378,6 +356,31 @@ void UDataflowEdNode::UpdatePinsConnectionsFromDataflowNode()
 }
 
 #if WITH_EDITOR
+
+void UDataflowEdNode::AllocateDefaultPins()
+{
+	UE_LOG(DATAFLOWNODE_LOG, Verbose, TEXT("UDataflowEdNode::AllocateDefaultPins()"));
+	// called on node creation from UI. 
+
+#if !(UE_BUILD_SHIPPING || UE_BUILD_TEST)
+	if (DataflowGraph)
+	{
+		if (DataflowNodeGuid.IsValid())
+		{
+			if (TSharedPtr<FDataflowNode> DataflowNode = DataflowGraph->FindBaseNode(DataflowNodeGuid))
+			{
+				for (const UE::Dataflow::FPin& Pin : DataflowNode->GetPins())
+				{
+					CreateEdPin(Pin);
+				}
+			}
+		}
+	}
+
+	UpdatePinsDefaultValuesFromNode();
+#endif // WITH_EDITOR && !UE_BUILD_SHIPPING
+}
+
 bool UDataflowEdNode::SupportsEditablePinType(const UEdGraphPin& Pin)
 {
 	if (const FDataflowConnection* Connection = GetConnectionFromPin(&Pin))
@@ -596,10 +599,7 @@ void UDataflowEdNode::WatchPin(const UEdGraphPin* Pin, bool bWatch)
 	}
 }
 
-FText UDataflowEdNode::GetNodeTitle(ENodeTitleType::Type TitleType) const
-{
-	return FText::FromString(GetName());
-}
+
 
 #if WITH_EDITOR && !(UE_BUILD_SHIPPING || UE_BUILD_TEST)
 void UDataflowEdNode::PinConnectionListChanged(UEdGraphPin* Pin)
@@ -775,6 +775,11 @@ void UDataflowEdNode::RemoveSpacesInAllPinTypes()
 
 #if WITH_EDITOR
 
+FText UDataflowEdNode::GetNodeTitle(ENodeTitleType::Type TitleType) const
+{
+	return FText::FromString(GetName());
+}
+
 FSlateIcon UDataflowEdNode::GetIconAndTint(FLinearColor& OutColor) const
 {
 	FSlateIcon Icon;
diff --git a/Engine/Source/Runtime/Experimental/Dataflow/Engine/Public/Dataflow/DataflowEdNode.h b/Engine/Source/Runtime/Experimental/Dataflow/Engine/Public/Dataflow/DataflowEdNode.h
--- a/Engine/Source/Runtime/Experimental/Dataflow/Engine/Public/Dataflow/DataflowEdNode.h
+++ b/Engine/Source/Runtime/Experimental/Dataflow/Engine/Public/Dataflow/DataflowEdNode.h
@@ -34,12 +34,12 @@ public:
 	virtual ~UDataflowEdNode() { UnRegisterDelegateHandle(); };
 
 	// UEdGraphNode interface
-	DATAFLOWENGINE_API virtual void AllocateDefaultPins();
-	DATAFLOWENGINE_API virtual FText GetNodeTitle(ENodeTitleType::Type TitleType) const;
-#if WITH_EDITOR && !(UE_BUILD_SHIPPING || UE_BUILD_TEST)
+#if WITH_EDITOR
+	DATAFLOWENGINE_API virtual void AllocateDefaultPins() override;
+	DATAFLOWENGINE_API virtual FText GetNodeTitle(ENodeTitleType::Type TitleType) const override;
+#if !(UE_BUILD_SHIPPING || UE_BUILD_TEST)
 	DATAFLOWENGINE_API virtual void PinConnectionListChanged(UEdGraphPin* Pin) override;
 #endif // WITH_EDITOR && !UE_BUILD_SHIPPING
-#if WITH_EDITOR
 	DATAFLOWENGINE_API virtual FSlateIcon GetIconAndTint(FLinearColor& OutColor) const override;
 	DATAFLOWENGINE_API virtual bool ShowPaletteIconOnNode() const override;
 	DATAFLOWENGINE_API virtual FLinearColor GetNodeTitleColor() const override;
@@ -56,7 +56,7 @@ public:
 	// End of UEdGraphNode interface
 
 	// UObject interface
-	DATAFLOWENGINE_API void Serialize(FArchive& Ar);
+	DATAFLOWENGINE_API void Serialize(FArchive& Ar) override;
 	// End UObject interface
 
 	bool IsBound() const { return DataflowGraph && DataflowNodeGuid.IsValid(); }
