Fix CVD usage of NAME_NO_NUMBER.

#Synced-CL 52002942

diff --git a/Engine/Source/Runtime/Experimental/Chaos/Private/ChaosVisualDebugger/ChaosVDSerializedNameTable.cpp b/Engine/Source/Runtime/Experimental/Chaos/Private/ChaosVisualDebugger/ChaosVDSerializedNameTable.cpp
--- a/Engine/Source/Runtime/Experimental/Chaos/Private/ChaosVisualDebugger/ChaosVDSerializedNameTable.cpp
+++ b/Engine/Source/Runtime/Experimental/Chaos/Private/ChaosVisualDebugger/ChaosVDSerializedNameTable.cpp
@@ -52,7 +52,10 @@ namespace Chaos::VisualDebugger
 	uint64 FChaosVDSerializableNameTable::AddNameToTable(const FChaosVDSerializedNameEntry& InNameEntry)
 	{
 		UE::TWriteScopeLock WriteLock(NamesByIDLock);
-		NamesByID.Add(InNameEntry.NameID, FName(InNameEntry.Name, InNameEntry.NameNumber));
+
+		int32 NameNumber = InNameEntry.NameNumber == NAME_NO_NUMBER ? NAME_NO_NUMBER_INTERNAL : InNameEntry.NameNumber;
+
+		NamesByID.Add(InNameEntry.NameID, FName(InNameEntry.Name, NameNumber));
 
 		return InNameEntry.NameID;
 	}
diff --git a/Engine/Source/Runtime/Experimental/Chaos/Public/ChaosVisualDebugger/ChaosVDSerializedNameTable.h b/Engine/Source/Runtime/Experimental/Chaos/Public/ChaosVisualDebugger/ChaosVDSerializedNameTable.h
--- a/Engine/Source/Runtime/Experimental/Chaos/Public/ChaosVisualDebugger/ChaosVDSerializedNameTable.h
+++ b/Engine/Source/Runtime/Experimental/Chaos/Public/ChaosVisualDebugger/ChaosVDSerializedNameTable.h
@@ -23,7 +23,7 @@ namespace Chaos::VisualDebugger
 		static UE_API FStringView WrapperTypeName;
 
 		uint64 NameID = 0;
-		int32 NameNumber = NAME_NO_NUMBER;
+		int32 NameNumber = NAME_NO_NUMBER_INTERNAL;
 		FString Name;
 	};
 
