Add RotatedRadialVector field node

#Synced-CL 51996809

diff --git a/Engine/Source/Runtime/Experimental/Chaos/Private/Field/FieldSystemNodes.cpp b/Engine/Source/Runtime/Experimental/Chaos/Private/Field/FieldSystemNodes.cpp
--- a/Engine/Source/Runtime/Experimental/Chaos/Private/Field/FieldSystemNodes.cpp
+++ b/Engine/Source/Runtime/Experimental/Chaos/Private/Field/FieldSystemNodes.cpp
@@ -43,6 +43,8 @@ FFieldNodeBase * FieldNodeFactory(FFieldNodeBase::EFieldType BaseType,FFieldNode
 		return new FUniformVector();
 	case FFieldNodeBase::ESerializationType::FieldNode_FRadialVector:
 		return new FRadialVector();
+	case FFieldNodeBase::ESerializationType::FieldNode_FRotatedRadialVector:
+		return new FRotatedRadialVector();
 	case FFieldNodeBase::ESerializationType::FieldNode_FRandomVector:
 		return new FRandomVector();
 	case FFieldNodeBase::ESerializationType::FieldNode_FSumScalar:
@@ -1029,6 +1031,78 @@ void FRadialVector::ComputeFieldBounds(FVector& MinBounds, FVector& MaxBounds, F
 }
 
 
+/**
+* FRotatedRadialVector
+*/
+
+void FRotatedRadialVector::Evaluate(FFieldContext& Context, TFieldArrayView<FVector>& Results) const
+{
+	const FQuat Rotation = FQuat(Axis.GetSafeNormal(), FMath::DegreesToRadians(Angle));
+	const int32 NumSamples = Context.SampleIndices.Num();
+	for (int32 SampleIndex = 0; SampleIndex < NumSamples; SampleIndex++)
+	{
+		const FFieldContextIndex& Index = Context.SampleIndices[SampleIndex];
+		{
+			const FVector RadialDir = (Context.SamplePositions[Index.Sample] - Position).GetSafeNormal();
+			Results[Index.Result] = Magnitude * Rotation.RotateVector(RadialDir);
+		}
+	}
+}
+
+void FRotatedRadialVector::Serialize(FArchive& Ar)
+{
+	Super::Serialize(Ar);
+	Ar << Magnitude;
+	Ar << Position;
+	Ar << Angle;
+	Ar << Axis;
+}
+
+bool FRotatedRadialVector::operator==(const FFieldNodeBase& Node)
+{
+	if (Node.SerializationType() == SerializationType())
+	{
+		const FRotatedRadialVector* Other = static_cast<const FRotatedRadialVector*>(&Node);
+		return Super::operator==(Node)
+			&& Magnitude == Other->Magnitude
+			&& Position == Other->Position
+			&& Angle == Other->Angle
+			&& Axis == Other->Axis;
+	}
+	return false;
+}
+
+void FRotatedRadialVector::FillSetupCount(int32& NumOffsets, int32& NumParams) const
+{
+	Super::FillSetupCount(NumOffsets, NumParams);
+	NumParams += 8; // Magnitude + Position(3) + Angle + Axis(3)
+}
+
+void FRotatedRadialVector::FillSetupDatas(TArray<int32>& NodesOffsets, TArray<float>& NodesParams, const float CommandTime) const
+{
+	Super::FillSetupDatas(NodesOffsets, NodesParams, CommandTime);
+	NodesParams.Add(Magnitude);
+	NodesParams.Add(static_cast<float>(Position.X));
+	NodesParams.Add(static_cast<float>(Position.Y));
+	NodesParams.Add(static_cast<float>(Position.Z));
+	NodesParams.Add(Angle);
+	NodesParams.Add(static_cast<float>(Axis.X));
+	NodesParams.Add(static_cast<float>(Axis.Y));
+	NodesParams.Add(static_cast<float>(Axis.Z));
+}
+
+float FRotatedRadialVector::EvalMaxMagnitude() const
+{
+	return Magnitude;
+}
+
+void FRotatedRadialVector::ComputeFieldBounds(FVector& MinBounds, FVector& MaxBounds, FVector& CenterPosition) const
+{
+	MinBounds = FVector(-FLT_MAX);
+	MaxBounds = FVector(FLT_MAX);
+	CenterPosition = Position;
+}
+
 /**
 * FRandomVector
 */
diff --git a/Engine/Source/Runtime/Experimental/Chaos/Public/Field/FieldSystem.h b/Engine/Source/Runtime/Experimental/Chaos/Public/Field/FieldSystem.h
--- a/Engine/Source/Runtime/Experimental/Chaos/Public/Field/FieldSystem.h
+++ b/Engine/Source/Runtime/Experimental/Chaos/Public/Field/FieldSystem.h
@@ -432,7 +432,8 @@ public:
 		FieldNode_FConversionField,
 		FieldNode_FCullingField,
 		FieldNode_FWaveScalar,
-		FieldNode_FReturnResultsTerminal
+		FieldNode_FReturnResultsTerminal,
+		FieldNode_FRotatedRadialVector
 	};
 
 	FFieldNodeBase() : Name("") {}
diff --git a/Engine/Source/Runtime/Experimental/Chaos/Public/Field/FieldSystemNodes.h b/Engine/Source/Runtime/Experimental/Chaos/Public/Field/FieldSystemNodes.h
--- a/Engine/Source/Runtime/Experimental/Chaos/Public/Field/FieldSystemNodes.h
+++ b/Engine/Source/Runtime/Experimental/Chaos/Public/Field/FieldSystemNodes.h
@@ -493,6 +493,54 @@ public:
 
 };
 
+/**
+* RotatedRadialVector
+**/
+
+class FRotatedRadialVector : public FFieldNode<FVector>
+{
+	typedef FFieldNode<FVector> Super;
+
+public:
+
+	FRotatedRadialVector(float MagnitudeIn = 1.f,
+		FVector PositionIn = FVector(0, 0, 0),
+		float AngleIn = 0.f,
+		FVector AxisIn = FVector(0, 0, 1))
+		: Super()
+		, Magnitude(MagnitudeIn)
+		, Position(PositionIn)
+		, Angle(AngleIn)
+		, Axis(AxisIn)
+	{}
+	virtual FFieldNodeBase* NewCopy() const override { return new FRotatedRadialVector(Magnitude, Position, Angle, Axis); }
+	virtual ~FRotatedRadialVector() = default;
+
+	CHAOS_API virtual void Evaluate(FFieldContext&, TFieldArrayView<FVector>& Results) const override;
+	CHAOS_API virtual bool operator==(const FFieldNodeBase& Node) override;
+
+	/** Serialization API */
+	virtual FFieldNodeBase::ESerializationType SerializationType() const override { return FFieldNodeBase::ESerializationType::FieldNode_FRotatedRadialVector; }
+	CHAOS_API virtual void Serialize(FArchive& Ar) override;
+
+	/** Count the number of offsets/params that will be used by the world physics field */
+	CHAOS_API virtual void FillSetupCount(int32& NumOffsets, int32& NumParams) const override;
+
+	/** Fill the offsets/params arrays that will be used by the world physics field */
+	CHAOS_API virtual void FillSetupDatas(TArray<int32>& NodesOffsets, TArray<float>& NodesParams, const float CommandTime) const override;
+
+	/** Evaluate the maximum magnitude of the field graph  */
+	CHAOS_API virtual float EvalMaxMagnitude() const override;
+
+	/** Compute the min/max spatial bounds of the field */
+	CHAOS_API virtual void ComputeFieldBounds(FVector& MinBounds, FVector& MaxBounds, FVector& CenterPosition) const override;
+
+	float Magnitude;
+	FVector Position;
+	float Angle;
+	FVector Axis;
+};
+
 /**
 * RandomVector
 **/
diff --git a/Engine/Source/Runtime/Experimental/FieldSystem/Source/FieldSystemEngine/Private/Field/FieldSystemObjects.cpp b/Engine/Source/Runtime/Experimental/FieldSystem/Source/FieldSystemEngine/Private/Field/FieldSystemObjects.cpp
--- a/Engine/Source/Runtime/Experimental/FieldSystem/Source/FieldSystemEngine/Private/Field/FieldSystemObjects.cpp
+++ b/Engine/Source/Runtime/Experimental/FieldSystem/Source/FieldSystemEngine/Private/Field/FieldSystemObjects.cpp
@@ -267,6 +267,27 @@ URadialVector* URadialVector::SetRadialVector(float InMagnitude, FVector InPosit
 }
 
 
+// RotatedRadialVector
+FFieldNodeBase* URotatedRadialVector::NewEvaluationGraph(TArray<const UFieldNodeBase*>& Nodes) const
+{
+	if (ensureMsgf(!Nodes.Contains(this), TEXT("Cycle Dependency Error : Graph nodes may not be reused in a single chain.")))
+	{
+		Nodes.Add(this);
+		return new FRotatedRadialVector(Magnitude, Position, Angle, Axis);
+	}
+	return nullptr;
+};
+
+URotatedRadialVector* URotatedRadialVector::SetRotatedRadialVector(float InMagnitude, FVector InPosition, float InAngle, FVector InAxis)
+{
+	this->Magnitude = InMagnitude;
+	this->Position = InPosition;
+	this->Angle = InAngle;
+	this->Axis = InAxis;
+	return this;
+}
+
+
 // RadialVector
 FFieldNodeBase* URandomVector::NewEvaluationGraph(TArray<const UFieldNodeBase*>& Nodes) const
 {
diff --git a/Engine/Source/Runtime/Experimental/FieldSystem/Source/FieldSystemEngine/Public/Field/FieldSystemObjects.h b/Engine/Source/Runtime/Experimental/FieldSystem/Source/FieldSystemEngine/Public/Field/FieldSystemObjects.h
--- a/Engine/Source/Runtime/Experimental/FieldSystem/Source/FieldSystemEngine/Public/Field/FieldSystemObjects.h
+++ b/Engine/Source/Runtime/Experimental/FieldSystem/Source/FieldSystemEngine/Public/Field/FieldSystemObjects.h
@@ -703,6 +703,58 @@ public:
 	FVector Position;
 };
 
+/**
+* RotatedRadialVector
+*/
+UCLASS(ClassGroup = "Field", meta = (BlueprintSpawnableComponent, ToolTip = "Set a radial vector value rotated around an axis. The direction is the normalized vector from the field position to the sample, rotated by Angle degrees around Axis. The output is equal to magnitude * rotated direction"), ShowCategories = ("Field"), MinimalAPI)
+class URotatedRadialVector : public UFieldNodeVector
+{
+	GENERATED_BODY()
+public:
+
+	URotatedRadialVector()
+		: Super()
+		, Magnitude(1.0)
+		, Position(FVector(0, 0, 0))
+		, Angle(0.0f)
+		, Axis(FVector(0, 0, 1))
+	{}
+
+	virtual ~URotatedRadialVector() = default;
+
+	FIELDSYSTEMENGINE_API virtual FFieldNodeBase* NewEvaluationGraph(TArray<const UFieldNodeBase*>& Nodes) const override;
+
+	/**
+	 * Set a rotated radial vector value. The direction is the normalized vector from the field position to the sample, then rotated by Angle degrees around Axis.
+	 * @param    Magnitude Magnitude of the radial vector field
+	 * @param    Position Center position of the radial vector field
+	 * @param    Angle Rotation angle in degrees
+	 * @param    Axis Axis to rotate around
+	 */
+	UFUNCTION(BlueprintPure, Category = "Field", meta = (Magnitude = "1.0", Angle = "0.0"))
+	FIELDSYSTEMENGINE_API URotatedRadialVector* SetRotatedRadialVector(
+		UPARAM(DisplayName = "Field Magnitude") float Magnitude,
+		UPARAM(DisplayName = "Center Position") FVector Position,
+		UPARAM(DisplayName = "Rotation Angle") float Angle,
+		UPARAM(DisplayName = "Rotation Axis") FVector Axis);
+
+	/** Magnitude of the radial vector field */
+	UPROPERTY(BlueprintReadWrite, EditAnywhere, Category = "Field", DisplayName = "Field Magnitude")
+	float Magnitude;
+
+	/** Center position of the radial vector field */
+	UPROPERTY(BlueprintReadWrite, EditAnywhere, Category = "Field", DisplayName = "Center Position")
+	FVector Position;
+
+	/** Rotation angle in degrees applied to the radial direction */
+	UPROPERTY(BlueprintReadWrite, EditAnywhere, Category = "Field", DisplayName = "Rotation Angle")
+	float Angle;
+
+	/** Axis around which the radial direction is rotated */
+	UPROPERTY(BlueprintReadWrite, EditAnywhere, Category = "Field", DisplayName = "Rotation Axis")
+	FVector Axis;
+};
+
 /**
 * URandomVector
 */
