Chaos API: More fully implementing TriangleMeshGeometry.
#Synced-CL 52060956

diff --git a/Engine/Plugins/Experimental/ChaosRigidPhysicsAsync/Tests/Private/AnyGeometryTests.cpp b/Engine/Plugins/Experimental/ChaosRigidPhysicsAsync/Tests/Private/AnyGeometryTests.cpp
--- a/Engine/Plugins/Experimental/ChaosRigidPhysicsAsync/Tests/Private/AnyGeometryTests.cpp
+++ b/Engine/Plugins/Experimental/ChaosRigidPhysicsAsync/Tests/Private/AnyGeometryTests.cpp
@@ -5,6 +5,7 @@
 #include "ChaosRigidPhysicsAsyncTest.h"
 
 #include "RigidPhysics/Geometry/AnyGeometry.h"
+#include "RigidTestUtils.h"
 
 #if UE_RIGIDPHYSICS_API_ENABLED
 
@@ -13,6 +14,7 @@ namespace Chaos::LowLevelTest
 	using FSphereGeometry = UE::Physics::FSphereGeometry;
 	using FBoxGeometry = UE::Physics::FBoxGeometry;
 	using FCapsuleGeometry = UE::Physics::FCapsuleGeometry;
+	using FTriangleMeshGeometry = UE::Physics::FTriangleMeshGeometry;
 	using FAnyGeometry = UE::Physics::FAnyGeometry;
 
 	template <typename ExpectedGeometryType, typename ActualGeometryType>
@@ -171,6 +173,45 @@ namespace Chaos::LowLevelTest
 			CHECK(AnyGeometry == Capsule1);
 		}
 	}
+
+	TEST_CASE("AnyGeometry: TriangleMesh", "[Chaos][API][Geometry][unit]")
+	{
+		FBoxGeometry Box0(FVector3f(1, 2, 3));
+		FTriangleMeshGeometry TriMesh0(UE::Physics::MakeQuadTriangleGeometrySetup());
+
+		SECTION("IsA")
+		{
+			FAnyGeometry AnyGeometry(TriMesh0);
+			CHECK(AnyGeometry.IsA<FTriangleMeshGeometry>());
+			CHECK(!AnyGeometry.IsA<FBoxGeometry>());
+		}
+		SECTION("TryGet/Get")
+		{
+			FAnyGeometry AnyGeometry(TriMesh0);
+			REQUIRE(AnyGeometry.TryGet<FTriangleMeshGeometry>() != nullptr);
+			CHECK(*AnyGeometry.TryGet<FTriangleMeshGeometry>() == TriMesh0);
+			CHECK(AnyGeometry.Get<FTriangleMeshGeometry>() == TriMesh0);
+			CHECK(AnyGeometry.TryGet<FBoxGeometry>() == nullptr);
+		}
+		SECTION("Visit")
+		{
+			FAnyGeometry AnyGeometry(TriMesh0);
+			AnyGeometry.Visit([&TriMesh0]<typename GeometryType>(const GeometryType& Geometry) { CheckGeometry(TriMesh0, Geometry); });
+		}
+		SECTION("Comparison")
+		{
+			FAnyGeometry AnyGeometry(TriMesh0);
+			CHECK(AnyGeometry == TriMesh0);
+			CHECK(AnyGeometry != Box0);
+		}
+		SECTION("Assignment")
+		{
+			FAnyGeometry AnyGeometry(Box0);
+			AnyGeometry = TriMesh0;
+			CHECK(AnyGeometry != Box0);
+			CHECK(AnyGeometry == TriMesh0);
+		}
+	}
 } // namespace Chaos::LowLevelTest
 
 #endif // UE_RIGIDPHYSICS_API_ENABLED
diff --git a/Engine/Plugins/Experimental/ChaosRigidPhysicsAsync/Tests/Private/RigidTestUtils.cpp b/Engine/Plugins/Experimental/ChaosRigidPhysicsAsync/Tests/Private/RigidTestUtils.cpp
--- a/Engine/Plugins/Experimental/ChaosRigidPhysicsAsync/Tests/Private/RigidTestUtils.cpp
+++ b/Engine/Plugins/Experimental/ChaosRigidPhysicsAsync/Tests/Private/RigidTestUtils.cpp
@@ -8,6 +8,46 @@
 
 namespace UE::Physics
 {
+	FTriangleMeshGeometrySetup MakeSingleTriangleGeometrySetup(const FVector3f& InScale, const FVector3f& InTranslation)
+	{
+		FTriangleMeshGeometrySetup Setup
+		{
+			.Vertices
+			{
+				FVector3f(0, 0, 0) * InScale + InTranslation,
+				FVector3f(1, 0, 0) * InScale + InTranslation,
+				FVector3f(0, 1, 0) * InScale + InTranslation,
+			},
+			.TriangleIndices
+			{
+				FInt32Vector3(0, 1, 2),
+			},
+			.MaterialIndices{0},
+		};
+		return Setup;
+	}
+
+	FTriangleMeshGeometrySetup MakeQuadTriangleGeometrySetup(const FVector3f& InScale, const FVector3f& InTranslation)
+	{
+		FTriangleMeshGeometrySetup Setup
+		{
+			.Vertices
+			{
+				FVector3f(0, 0, 0) * InScale + InTranslation,
+				FVector3f(1, 0, 0) * InScale + InTranslation,
+				FVector3f(1, 1, 0) * InScale + InTranslation,
+				FVector3f(0, 1, 0) * InScale + InTranslation,
+			},
+			.TriangleIndices
+			{
+				FInt32Vector3(0, 1, 2),
+				FInt32Vector3(0, 2, 3),
+			},
+			.MaterialIndices{0, 1}
+		};
+		return Setup;
+	}
+
 	FRigidShapeInstanceSetup MakeBoxShape(const FVector3f& InSize, const FVector3f& InCenter, UPhysicalMaterial* InMaterial)
 	{
 		FBoxGeometry BoxGeom(InSize);
diff --git a/Engine/Plugins/Experimental/ChaosRigidPhysicsAsync/Tests/Private/RigidTestUtils.h b/Engine/Plugins/Experimental/ChaosRigidPhysicsAsync/Tests/Private/RigidTestUtils.h
--- a/Engine/Plugins/Experimental/ChaosRigidPhysicsAsync/Tests/Private/RigidTestUtils.h
+++ b/Engine/Plugins/Experimental/ChaosRigidPhysicsAsync/Tests/Private/RigidTestUtils.h
@@ -19,6 +19,8 @@ namespace Chaos
 
 namespace UE::Physics
 {
+	FTriangleMeshGeometrySetup MakeSingleTriangleGeometrySetup(const FVector3f& InScale = FVector3f::One(), const FVector3f& InTranslation = FVector3f::Zero());
+	FTriangleMeshGeometrySetup MakeQuadTriangleGeometrySetup(const FVector3f& InScale = FVector3f::One(), const FVector3f& InTranslation = FVector3f::Zero());
 	FRigidShapeInstanceSetup MakeBoxShape(const FVector3f& InSize, const FVector3f& InCenter = FVector3f::Zero(), UPhysicalMaterial* InMaterial = nullptr);
 	FVector3f MakeSolidBoxInertia(float InMass, const FVector3f& InSize);
 } // namespace UE::Physics
diff --git a/Engine/Plugins/Experimental/ChaosRigidPhysicsAsync/Tests/Private/TriangleMeshGeometryTests.cpp b/Engine/Plugins/Experimental/ChaosRigidPhysicsAsync/Tests/Private/TriangleMeshGeometryTests.cpp
new file mode 100644
--- /dev/null
+++ b/Engine/Plugins/Experimental/ChaosRigidPhysicsAsync/Tests/Private/TriangleMeshGeometryTests.cpp
@@ -0,0 +1,123 @@
+// Copyright Epic Games, Inc. All Rights Reserved.
+
+#include "Chaos/LowLevelTest/ChaosTestHarness.h"
+#include "ChaosRigidPhysicsAsyncTest.h"
+
+#include "RigidPhysics/Geometry/TriangleMeshGeometry.h"
+#include "RigidTestUtils.h"
+
+#if UE_RIGIDPHYSICS_API_ENABLED
+
+namespace Chaos::LowLevelTest
+{
+	using FTriangleMeshGeometry = UE::Physics::FTriangleMeshGeometry;
+	using FTriangleMeshGeometrySetup = UE::Physics::FTriangleMeshGeometrySetup;
+
+	TEST_CASE("TriangleMeshGeometry::Basic", "[Chaos][API][Geometry][unit]")
+	{
+		const FTriangleMeshGeometrySetup Setup = GENERATE(
+			UE::Physics::MakeSingleTriangleGeometrySetup(),
+			UE::Physics::MakeQuadTriangleGeometrySetup()
+		);
+		FTriangleMeshGeometrySetup MutableSetup = Setup;
+		const FTriangleMeshGeometry Geometry(MoveTemp(MutableSetup));
+
+		REQUIRE(Setup.Vertices.Num() == Geometry.GetNumVertices());
+		for (int32 I = 0; I < Setup.Vertices.Num(); ++I)
+		{
+			CHECK(Setup.Vertices[I] == Geometry.GetVertex(I));
+		}
+		REQUIRE(Setup.TriangleIndices.Num() == Geometry.GetNumTriangleIndices());
+		for (int32 I = 0; I < Setup.TriangleIndices.Num(); ++I)
+		{
+			CHECK(Setup.TriangleIndices[I] == Geometry.GetTriangleIndices(I));
+		}
+		REQUIRE(Setup.MaterialIndices.Num() == Geometry.GetNumMaterialIndices());
+		for (int32 I = 0; I < Setup.MaterialIndices.Num(); ++I)
+		{
+			CHECK(Setup.MaterialIndices[I] == Geometry.GetMaterialIndex(I));
+		}
+		REQUIRE(Setup.ExternalFaceIndexMap.IsEmpty() == !Geometry.HasExternalFaceIndices());
+		for (int32 I = 0; I < Setup.ExternalFaceIndexMap.Num(); ++I)
+		{
+			CHECK(Setup.ExternalFaceIndexMap[I] == Geometry.GetExternalFaceIndex(I));
+		}
+		REQUIRE(Setup.ExternalVertexIndexMap.IsEmpty() == !Geometry.HasExternalVertexIndices());
+		for (int32 I = 0; I < Setup.ExternalVertexIndexMap.Num(); ++I)
+		{
+			CHECK(Setup.ExternalVertexIndexMap[I] == Geometry.GetExternalVertexIndex(I));
+		}
+		CHECK(Setup.bCullsBackFaceRaycast == Geometry.GetCullsBackFaceRaycast());
+		CHECK(Setup.PerFaceBoundsRepresentation == Geometry.GetPerFaceBoundsRepresentation());
+	}
+
+	TEST_CASE("TriangleMeshGeometry::ExternalFaceIndexMap", "[Chaos][API][Geometry][unit]")
+	{
+		FTriangleMeshGeometrySetup Setup = UE::Physics::MakeQuadTriangleGeometrySetup();
+		SECTION("Empty")
+		{
+			const FTriangleMeshGeometry Geometry(MoveTemp(Setup));
+			CHECK(!Geometry.HasExternalFaceIndices());
+		}
+		SECTION("Filled")
+		{
+			const TArray<int32> ExternalFaceIndices{ 10, 11 };
+			Setup.ExternalFaceIndexMap = ExternalFaceIndices;
+			const FTriangleMeshGeometry Geometry(MoveTemp(Setup));
+
+			REQUIRE(Geometry.HasExternalFaceIndices());
+			for (int32 I = 0; I < ExternalFaceIndices.Num(); ++I)
+			{
+				CHECK(ExternalFaceIndices[I] == Geometry.GetExternalFaceIndex(I));
+			}
+		}
+	}
+
+	TEST_CASE("TriangleMeshGeometry::ExternalVertexIndexMap", "[Chaos][API][Geometry][unit]")
+	{
+		FTriangleMeshGeometrySetup Setup = UE::Physics::MakeQuadTriangleGeometrySetup();
+		SECTION("Empty")
+		{
+			const FTriangleMeshGeometry Geometry(MoveTemp(Setup));
+
+			CHECK(!Geometry.HasExternalVertexIndices());
+		}
+		SECTION("Filled")
+		{
+			const TArray<int32> ExternalVertexIndices{ 10, 11, 12, 13 };
+			Setup.ExternalVertexIndexMap = ExternalVertexIndices;
+			const FTriangleMeshGeometry Geometry(MoveTemp(Setup));
+
+			REQUIRE(Geometry.HasExternalVertexIndices());
+			for (int32 I = 0; I < ExternalVertexIndices.Num(); ++I)
+			{
+				CHECK(ExternalVertexIndices[I] == Geometry.GetExternalVertexIndex(I));
+			}
+		}
+	}
+
+	TEST_CASE("TriangleMeshGeometry::CullsBackFaceRaycast", "[Chaos][API][Geometry][unit]")
+	{
+		const bool bCullsBackFaceRaycast = GENERATE(false, true);
+
+		FTriangleMeshGeometrySetup Setup = UE::Physics::MakeSingleTriangleGeometrySetup();
+		Setup.bCullsBackFaceRaycast = bCullsBackFaceRaycast;
+		const FTriangleMeshGeometry Geometry(MoveTemp(Setup));
+
+		CHECK(Geometry.GetCullsBackFaceRaycast() == bCullsBackFaceRaycast);
+	}
+
+	TEST_CASE("TriangleMeshGeometry::PerFaceBoundsRepresentation", "[Chaos][API][Geometry][unit]")
+	{
+		using EPerFaceBoundsRepresentation = FTriangleMeshGeometrySetup::EPerFaceBoundsRepresentation;
+		const EPerFaceBoundsRepresentation PerFaceBoundsRepresentation = GENERATE(EPerFaceBoundsRepresentation::VectorizedFloat, EPerFaceBoundsRepresentation::Byte);
+
+		FTriangleMeshGeometrySetup Setup = UE::Physics::MakeSingleTriangleGeometrySetup();
+		Setup.PerFaceBoundsRepresentation = PerFaceBoundsRepresentation;
+		const FTriangleMeshGeometry Geometry(MoveTemp(Setup));
+
+		CHECK(Geometry.GetPerFaceBoundsRepresentation() == PerFaceBoundsRepresentation);
+	}
+} // namespace Chaos::LowLevelTest
+
+#endif // UE_RIGIDPHYSICS_API_ENABLED
diff --git a/Engine/Source/Runtime/Experimental/Chaos/Public/Chaos/TriangleMeshImplicitObject.h b/Engine/Source/Runtime/Experimental/Chaos/Public/Chaos/TriangleMeshImplicitObject.h
--- a/Engine/Source/Runtime/Experimental/Chaos/Public/Chaos/TriangleMeshImplicitObject.h
+++ b/Engine/Source/Runtime/Experimental/Chaos/Public/Chaos/TriangleMeshImplicitObject.h
@@ -26,6 +26,11 @@ namespace Chaos::Private
 	class FMeshContactGenerator;
 }
 
+namespace UE::Physics
+{
+	class FTriangleMeshGeometry;
+}
+
 namespace Chaos
 {
 	extern CHAOS_API bool TriMeshPerPolySupport;
@@ -1014,6 +1019,7 @@ private:
 
 		// Required by implicit object serialization, disabled for general use.
 		friend class FImplicitObject;
+		friend class UE::Physics::FTriangleMeshGeometry;
 
 		FTriangleMeshImplicitObject()
 		    : FImplicitObject(EImplicitObject::HasBoundingBox, ImplicitObjectType::TriangleMesh){};
diff --git a/Engine/Source/Runtime/Experimental/RigidPhysics/Private/RigidPhysics/Geometry/TriangleMeshGeometry.cpp b/Engine/Source/Runtime/Experimental/RigidPhysics/Private/RigidPhysics/Geometry/TriangleMeshGeometry.cpp
--- a/Engine/Source/Runtime/Experimental/RigidPhysics/Private/RigidPhysics/Geometry/TriangleMeshGeometry.cpp
+++ b/Engine/Source/Runtime/Experimental/RigidPhysics/Private/RigidPhysics/Geometry/TriangleMeshGeometry.cpp
@@ -5,15 +5,205 @@
 
 namespace UE::Physics
 {
-	FTriangleMeshGeometry::FTriangleMeshGeometry(Chaos::FTriangleMeshImplicitObject* InImplicit)
+	Chaos::FTrimeshBVH::EPerFaceBoundsRepresentation Convert(FTriangleMeshGeometrySetup::EPerFaceBoundsRepresentation InValue)
+	{
+		switch (InValue)
+		{
+			case FTriangleMeshGeometrySetup::EPerFaceBoundsRepresentation::VectorizedFloat:
+			{
+				return Chaos::FTrimeshBVH::EPerFaceBoundsRepresentation::VectorizedFloat;
+			}
+			case FTriangleMeshGeometrySetup::EPerFaceBoundsRepresentation::Byte:
+			{
+				return Chaos::FTrimeshBVH::EPerFaceBoundsRepresentation::Byte;
+			}
+			default:
+			{
+				ensureMsgf(false, TEXT("Unhandled conversion for FTriangleMeshGeometrySetup::EPerFaceBoundsRepresentation"));
+				return Chaos::FTrimeshBVH::EPerFaceBoundsRepresentation::VectorizedFloat;
+			}
+		}
+	}
+
+	FTriangleMeshGeometry::FTriangleMeshGeometry(FTriangleMeshImplicitObjectPtr InImplicit)
 		: Implicit(InImplicit)
 	{
 	}
 
+	FTriangleMeshGeometry::FTriangleMeshGeometry(FTriangleMeshGeometrySetup&& InSetup)
+	{
+		if (InSetup.Vertices.Num() < TNumericLimits<uint16>::Max())
+		{
+			BuildInternal<uint16>(MoveTemp(InSetup));
+		}
+		else
+		{
+			BuildInternal<int32>(MoveTemp(InSetup));
+		}
+	}
+
+	int32 FTriangleMeshGeometry::GetNumVertices() const
+	{
+		return (int32)Implicit->Particles().Size();
+	}
+
+	FVector3f FTriangleMeshGeometry::GetVertex(const int32 InIndex) const
+	{
+		return FVector3f(Implicit->Particles().GetX(InIndex));
+	}
+
+	int32 FTriangleMeshGeometry::GetNumTriangleIndices() const
+	{
+		return Implicit->Elements().GetNumTriangles();
+	}
+
+	FInt32Vector3 FTriangleMeshGeometry::GetTriangleIndices(const int32 InIndex) const
+	{
+		const Chaos::FTrimeshIndexBuffer& Elements = Implicit->Elements();
+		if (Elements.RequiresLargeIndices())
+		{
+			const Chaos::TVec3<int32>& TriIndices = Elements.GetLargeIndexBuffer()[InIndex];
+			return FInt32Vector3(TriIndices.X, TriIndices.Y, TriIndices.Z);
+		}
+		else
+		{
+			const Chaos::TVec3<uint16>& TriIndices = Elements.GetSmallIndexBuffer()[InIndex];
+			return FInt32Vector3(TriIndices.X, TriIndices.Y, TriIndices.Z);
+		}
+	}
+
+	int32 FTriangleMeshGeometry::GetNumMaterialIndices() const
+	{
+		return Implicit->MaterialIndices.Num();
+	}
+
+	FMaterialIndex FTriangleMeshGeometry::GetMaterialIndex(const int32 InIndex) const
+	{
+		return Implicit->GetMaterialIndex(InIndex);
+	}
+
+	bool FTriangleMeshGeometry::HasExternalFaceIndices() const
+	{
+		return Implicit->ExternalFaceIndexMap.IsValid();
+	}
+
+	int32 FTriangleMeshGeometry::GetExternalFaceIndex(const int32 InIndex) const
+	{
+		if (InIndex > INDEX_NONE && HasExternalFaceIndices())
+		{
+			if (ensure(InIndex < Implicit->ExternalFaceIndexMap->Num()))
+			{
+				return (*Implicit->ExternalFaceIndexMap)[InIndex];
+			}
+		}
+		return INDEX_NONE;
+	}
+
+	bool FTriangleMeshGeometry::HasExternalVertexIndices() const
+	{
+		return Implicit->ExternalVertexIndexMap.IsValid();
+	}
+
+	int32 FTriangleMeshGeometry::GetExternalVertexIndex(const int32 InIndex) const
+	{
+		if (InIndex > INDEX_NONE && HasExternalVertexIndices())
+		{
+			if (ensure(InIndex < Implicit->ExternalVertexIndexMap->Num()))
+			{
+				return (*Implicit->ExternalVertexIndexMap)[InIndex];
+			}
+		}
+		return INDEX_NONE;
+	}
+
+	bool FTriangleMeshGeometry::GetCullsBackFaceRaycast() const
+	{
+		return Implicit->GetCullsBackFaceRaycast();
+	}
+
+	FTriangleMeshGeometrySetup::EPerFaceBoundsRepresentation FTriangleMeshGeometry::GetPerFaceBoundsRepresentation() const
+	{
+		const bool bQuantizedFaceBounds = EnumHasAnyFlags(Implicit->MeshSettingsFlags, Chaos::FTriangleMeshImplicitObject::ETriangleMeshSettingsFlags::QuantizePerFaceBounds);
+		return bQuantizedFaceBounds ? EPerFaceBoundsRepresentation::Byte : EPerFaceBoundsRepresentation::VectorizedFloat;
+	}
+
 	FString FTriangleMeshGeometry::ToString() const
 	{
 		return TEXT("TriangleMesh");
 	}
+
+	Chaos::FTriangleMeshImplicitObjectPtr FTriangleMeshGeometry::GetImplicit() const
+	{
+		return Implicit;
+	}
+
+	template <typename IndexType>
+	void FTriangleMeshGeometry::BuildInternal(FTriangleMeshGeometrySetup&& InSetup)
+	{
+		const int32 VerticesNum = InSetup.Vertices.Num();
+		const int32 TriIndicesNum = InSetup.TriangleIndices.Num();
+
+		// Material indices must match triangle indices in count.
+		if (!ensure(InSetup.MaterialIndices.Num() == TriIndicesNum))
+		{
+			InSetup.MaterialIndices.SetNumZeroed(TriIndicesNum);
+		}
+		// External index maps must either match the size of their corresponding array or be empty
+		if (!ensure(InSetup.ExternalFaceIndexMap.IsEmpty() || InSetup.ExternalFaceIndexMap.Num() == TriIndicesNum))
+		{
+			InSetup.ExternalFaceIndexMap.SetNumZeroed(TriIndicesNum);
+		}
+		if (!ensure(InSetup.ExternalVertexIndexMap.IsEmpty() || InSetup.ExternalVertexIndexMap.Num() == VerticesNum))
+		{
+			InSetup.ExternalVertexIndexMap.SetNumZeroed(VerticesNum);
+		}
+
+		TArray<Chaos::FVec3f> Vertices;
+		Vertices.SetNum(VerticesNum);
+		for (int32 I = 0; I < VerticesNum; ++I)
+		{
+			Vertices[I] = InSetup.Vertices[I];
+		}
+
+		TArray<Chaos::TVec3<IndexType>> TriangleIndices;
+		TriangleIndices.SetNum(TriIndicesNum);
+		for (int32 I = 0; I < TriIndicesNum; ++I)
+		{
+			TriangleIndices[I][0] = IndexType(InSetup.TriangleIndices[I][0]);
+			TriangleIndices[I][1] = IndexType(InSetup.TriangleIndices[I][1]);
+			TriangleIndices[I][2] = IndexType(InSetup.TriangleIndices[I][2]);
+		}
+
+		TArray<uint16> MaterialIndices;
+		MaterialIndices.SetNum(InSetup.MaterialIndices.Num());
+		for (int32 I = 0; I < InSetup.MaterialIndices.Num(); ++I)
+		{
+			MaterialIndices[I] = InSetup.MaterialIndices[I];
+		}
+
+		TUniquePtr<TArray<int32>> ExternalFaceIndexMap;
+		TUniquePtr<TArray<int32>> ExternalVertexIndexMap;
+		if (!InSetup.ExternalFaceIndexMap.IsEmpty())
+		{
+			ExternalFaceIndexMap = MakeUnique<TArray<int32>>(MoveTemp(InSetup.ExternalFaceIndexMap));
+		}
+		if (!InSetup.ExternalVertexIndexMap.IsEmpty())
+		{
+			ExternalVertexIndexMap = MakeUnique<TArray<int32>>(MoveTemp(InSetup.ExternalVertexIndexMap));
+		}
+
+		Chaos::FTriangleMeshImplicitObject::FBuildSettings BuildSettings;
+		BuildSettings.PerFaceBounds = Convert(InSetup.PerFaceBoundsRepresentation);
+
+		Implicit = new Chaos::FTriangleMeshImplicitObject(
+			MoveTemp(Vertices),
+			MoveTemp(TriangleIndices),
+			MoveTemp(MaterialIndices),
+			MoveTemp(ExternalFaceIndexMap),
+			MoveTemp(ExternalVertexIndexMap),
+			InSetup.bCullsBackFaceRaycast,
+			BuildSettings);
+	}
 } // namespace UE::Physics
 
 #endif // UE_RIGIDPHYSICS_API_ENABLED
diff --git a/Engine/Source/Runtime/Experimental/RigidPhysics/Private/RigidPhysics/RigidMaterialIndex.cpp b/Engine/Source/Runtime/Experimental/RigidPhysics/Private/RigidPhysics/RigidMaterialIndex.cpp
new file mode 100644
--- /dev/null
+++ b/Engine/Source/Runtime/Experimental/RigidPhysics/Private/RigidPhysics/RigidMaterialIndex.cpp
@@ -0,0 +1,19 @@
+// Copyright Epic Games, Inc. All Rights Reserved.
+#include "RigidPhysics/RigidMaterialIndex.h"
+
+#if UE_RIGIDPHYSICS_API_ENABLED
+
+namespace UE::Physics
+{
+	FMaterialIndex::FMaterialIndex(const uint16 InIndex)
+		: Index(InIndex)
+	{
+	}
+
+	FMaterialIndex::operator uint16() const
+	{
+		return Index;
+	}
+} // namespace UE::Physics
+
+#endif // UE_RIGIDPHYSICS_API_ENABLED
diff --git a/Engine/Source/Runtime/Experimental/RigidPhysics/Public/RigidPhysics/Geometry/TriangleMeshGeometry.h b/Engine/Source/Runtime/Experimental/RigidPhysics/Public/RigidPhysics/Geometry/TriangleMeshGeometry.h
--- a/Engine/Source/Runtime/Experimental/RigidPhysics/Public/RigidPhysics/Geometry/TriangleMeshGeometry.h
+++ b/Engine/Source/Runtime/Experimental/RigidPhysics/Public/RigidPhysics/Geometry/TriangleMeshGeometry.h
@@ -6,22 +6,68 @@
 #if UE_RIGIDPHYSICS_API_ENABLED
 
 #include "Chaos/TriangleMeshImplicitObject.h"
+#include "RigidPhysics/RigidMaterialIndex.h"
 
 UE_EXPERIMENTAL(5.8, "The new Chaos API is experimental")
 namespace UE::Physics
 {
+	class FTriangleMeshGeometrySetup
+	{
+	public:
+		enum class EPerFaceBoundsRepresentation : uint8
+		{
+			VectorizedFloat = 0,
+			Byte,
+		};
+
+		TArray<FVector3f> Vertices;
+		TArray<FInt32Vector3> TriangleIndices;
+		TArray<FMaterialIndex> MaterialIndices;
+		TArray<int32> ExternalFaceIndexMap;
+		TArray<int32> ExternalVertexIndexMap;
+		bool bCullsBackFaceRaycast = false;
+		EPerFaceBoundsRepresentation PerFaceBoundsRepresentation = EPerFaceBoundsRepresentation::VectorizedFloat;
+	};
+
 	class FTriangleMeshGeometry
 	{
+		using FTriangleMeshImplicitObjectPtr = Chaos::FTriangleMeshImplicitObjectPtr;
+
 	public:
-		UE_INTERNAL FTriangleMeshGeometry() = default;
-		UE_INTERNAL RIGIDPHYSICS_API FTriangleMeshGeometry(Chaos::FTriangleMeshImplicitObject* InImplicit);
+		using EPerFaceBoundsRepresentation = FTriangleMeshGeometrySetup::EPerFaceBoundsRepresentation;
+
+		UE_INTERNAL RIGIDPHYSICS_API FTriangleMeshGeometry(FTriangleMeshImplicitObjectPtr InImplicit);
+		RIGIDPHYSICS_API FTriangleMeshGeometry(FTriangleMeshGeometrySetup&& InSetup);
 
 		bool operator==(const FTriangleMeshGeometry&) const = default;
 		bool operator!=(const FTriangleMeshGeometry&) const = default;
 
+		RIGIDPHYSICS_API int32 GetNumVertices() const;
+		RIGIDPHYSICS_API FVector3f GetVertex(const int32 InIndex) const;
+
+		RIGIDPHYSICS_API int32 GetNumTriangleIndices() const;
+		RIGIDPHYSICS_API FInt32Vector3 GetTriangleIndices(const int32 InIndex) const;
+
+		RIGIDPHYSICS_API int32 GetNumMaterialIndices() const;
+		RIGIDPHYSICS_API FMaterialIndex GetMaterialIndex(const int32 InIndex) const;
+
+		RIGIDPHYSICS_API bool HasExternalFaceIndices() const;
+		RIGIDPHYSICS_API int32 GetExternalFaceIndex(const int32 InIndex) const;
+
+		RIGIDPHYSICS_API bool HasExternalVertexIndices() const;
+		RIGIDPHYSICS_API int32 GetExternalVertexIndex(const int32 InIndex) const;
+
+		RIGIDPHYSICS_API bool GetCullsBackFaceRaycast() const;
+		RIGIDPHYSICS_API EPerFaceBoundsRepresentation GetPerFaceBoundsRepresentation() const;
+
 		RIGIDPHYSICS_API FString ToString() const;
 
+		UE_INTERNAL FTriangleMeshImplicitObjectPtr GetImplicit() const;
+
 	private:
+		template <typename IndexType>
+		void BuildInternal(FTriangleMeshGeometrySetup&& InSetup);
+
 		Chaos::FTriangleMeshImplicitObjectPtr Implicit;
 	};
 } // namespace UE::Physics
diff --git a/Engine/Source/Runtime/Experimental/RigidPhysics/Public/RigidPhysics/RigidMaterialIndex.h b/Engine/Source/Runtime/Experimental/RigidPhysics/Public/RigidPhysics/RigidMaterialIndex.h
new file mode 100644
--- /dev/null
+++ b/Engine/Source/Runtime/Experimental/RigidPhysics/Public/RigidPhysics/RigidMaterialIndex.h
@@ -0,0 +1,22 @@
+// Copyright Epic Games, Inc. All Rights Reserved.
+#pragma once
+
+#include "RigidPhysics/RigidFwd.h"
+
+#if UE_RIGIDPHYSICS_API_ENABLED
+
+UE_EXPERIMENTAL(5.8, "The new Chaos API is experimental")
+namespace UE::Physics
+{
+	class FMaterialIndex
+	{
+	public:
+		FMaterialIndex(const uint16 InIndex);
+		operator uint16() const;
+
+	private:
+		uint16 Index = INDEX_NONE;
+	};
+} // namespace UE::Physics
+
+#endif // UE_RIGIDPHYSICS_API_ENABLED
