[HWRT] Move PartGeometries size tracking from FGeometryCollectionSceneProxyBase to UGeometryCollection.
- GC part geometry collection is shared between primitives so allocated size should be tracked on the asset side.

#Synced-CL 52071961

diff --git a/Engine/Source/Runtime/Experimental/GeometryCollectionEngine/Private/GeometryCollection/GeometryCollectionObject.cpp b/Engine/Source/Runtime/Experimental/GeometryCollectionEngine/Private/GeometryCollection/GeometryCollectionObject.cpp
--- a/Engine/Source/Runtime/Experimental/GeometryCollectionEngine/Private/GeometryCollection/GeometryCollectionObject.cpp
+++ b/Engine/Source/Runtime/Experimental/GeometryCollectionEngine/Private/GeometryCollection/GeometryCollectionObject.cpp
@@ -1735,6 +1735,14 @@ FGeometryCollectionRenderResourceSizeInfo UGeometryCollection::GetRenderResource
 
 	InfoOut.NaniteResourcesSize += GetNaniteResourcesSize(*RenderData->NaniteResourcesPtr);
 
+#if RHI_RAYTRACING
+	for (const FRayTracingGeometry& PartRayTracingGeometry : RenderData->RayTracingResource.PartGeometries)
+	{
+		// TODO: Have a dedicated size for ray tracing?
+		InfoOut.MeshResourcesSize += PartRayTracingGeometry.RawData.GetAllocatedSize();
+	}
+#endif
+
 	return InfoOut;
 }
 
diff --git a/Engine/Source/Runtime/Experimental/GeometryCollectionEngine/Private/GeometryCollection/GeometryCollectionSceneProxy.cpp b/Engine/Source/Runtime/Experimental/GeometryCollectionEngine/Private/GeometryCollection/GeometryCollectionSceneProxy.cpp
--- a/Engine/Source/Runtime/Experimental/GeometryCollectionEngine/Private/GeometryCollection/GeometryCollectionSceneProxy.cpp
+++ b/Engine/Source/Runtime/Experimental/GeometryCollectionEngine/Private/GeometryCollection/GeometryCollectionSceneProxy.cpp
@@ -795,21 +795,12 @@ RayTracing::FGeometryGroupHandle FGeometryCollectionSceneProxyBase::GetRayTracin
 
 uint32 FGeometryCollectionSceneProxyBase::GetAllocatedSize() const
 {
-#if RHI_RAYTRACING
-	uint32 RayTracingGeometryAllocatedSize = RayTracingGeometry.RawData.GetAllocatedSize();
-
-	for (const FRayTracingGeometry& PartRayTracingGeometry : RayTracingResources.PartGeometries)
-	{
-		RayTracingGeometryAllocatedSize += PartRayTracingGeometry.RawData.GetAllocatedSize();
-	}
-#endif
-
 	return Materials.GetAllocatedSize()
 		+ MeshDescription.Sections.GetAllocatedSize()
 		+ MeshDescription.SubSections.GetAllocatedSize()
 		+ (SkinnedPositionVertexBuffer.GetAllowCPUAccess() ? SkinnedPositionVertexBuffer.GetStride() * SkinnedPositionVertexBuffer.GetNumVertices() : 0)
 #if RHI_RAYTRACING
-		+ RayTracingGeometryAllocatedSize
+		+ RayTracingGeometry.RawData.GetAllocatedSize()
 #endif
 		;
 }
