Optimize ManagedArray copy by using MemoryOps's CopyAssignItems. This will do a Memcpy when the elements are trivially copy assignable.

#Synced-CL 52021766

diff --git a/Engine/Source/Runtime/Experimental/Chaos/Public/GeometryCollection/ManagedArray.h b/Engine/Source/Runtime/Experimental/Chaos/Public/GeometryCollection/ManagedArray.h
--- a/Engine/Source/Runtime/Experimental/Chaos/Public/GeometryCollection/ManagedArray.h
+++ b/Engine/Source/Runtime/Experimental/Chaos/Public/GeometryCollection/ManagedArray.h
@@ -12,6 +12,7 @@
 #include "Chaos/BVHParticles.h"
 #include "Math/Vector.h"
 #include "Chaos/Matrix.h"
+#include "Templates/MemoryOps.h"
 #include "Templates/TypeHash.h"
 
 struct FManagedArrayCollection;
@@ -666,10 +667,8 @@ inline uint32 GetTypeHash(const TManagedArrayBase<T>& ManagedArray)
 template <typename T>
 void InitHelper(TArray<T>& Array, const TManagedArrayBase<T>& NewTypedArray, int32 Size)
 {
-	for (int32 Index = 0; Index < Size; Index++)
-	{
-		Array[Index] = NewTypedArray[Index];
-	}
+	check(Size <= Array.Num() && Size <= NewTypedArray.Num());
+	CopyAssignItems(Array.GetData(), NewTypedArray.GetData(), Size);
 }
 
 
@@ -703,9 +702,11 @@ void InitHelper(TArray<TRefCountPtr<T>>& Array, const TManagedArrayBase<TRefCoun
 template <typename T>
 void CopyRangeHelper(TArray<T>& Target, const TManagedArrayBase<T>& Source, int32 Start, int32 Stop, int32 Offset)
 {
-	for (int32 Sdx = Start, Tdx = Start + Offset; Sdx < Source.Num() && Tdx < Target.Num() && Sdx < Stop; Sdx++, Tdx++)
+	const int32 NumCopy = FMath::Min(FMath::Min(Source.Num() - Start, Stop - Start), Target.Num() - Start - Offset);
+	if (NumCopy > 0)
 	{
-		Target[Tdx] = Source[Sdx];
+		check(Start >= 0 && Start + Offset >= 0);
+		CopyAssignItems(Target.GetData() + Start + Offset, Source.GetData() + Start, NumCopy);
 	}
 }
 
