Chaos API: Implement a basic public API for ConvexGeometry.

#Synced-CL 52100670

diff --git a/Engine/Plugins/Experimental/ChaosRigidPhysicsAsync/Tests/Private/ConvexGeometryTests.cpp b/Engine/Plugins/Experimental/ChaosRigidPhysicsAsync/Tests/Private/ConvexGeometryTests.cpp
new file mode 100644
--- /dev/null
+++ b/Engine/Plugins/Experimental/ChaosRigidPhysicsAsync/Tests/Private/ConvexGeometryTests.cpp
@@ -0,0 +1,26 @@
+// Copyright Epic Games, Inc. All Rights Reserved.
+
+#include "Chaos/LowLevelTest/ChaosTestHarness.h"
+#include "ChaosRigidPhysicsAsyncTest.h"
+
+#include "RigidPhysics/Geometry/ConvexGeometry.h"
+#include "RigidTestUtils.h"
+
+#if UE_RIGIDPHYSICS_API_ENABLED
+
+namespace Chaos::LowLevelTest
+{
+	using FConvexGeometry = UE::Physics::FConvexGeometry;
+	using FConvexGeometrySetup = UE::Physics::FConvexGeometrySetup;
+
+	TEST_CASE("ConvexGeometry::Basic", "[Chaos][API][Geometry][unit]")
+	{
+		const FConvexGeometrySetup Setup = UE::Physics::MakeConvexBoxGeometrySetup();
+		FConvexGeometrySetup MutableSetup = Setup;
+		FConvexGeometry Convex(MoveTemp(MutableSetup));
+		CHECK(Convex.GetNumVertices() == Setup.Vertices.Num());
+		CHECK(Convex.GetMargin() == Setup.Margin);
+	}
+} // namespace Chaos::LowLevelTest
+
+#endif // UE_RIGIDPHYSICS_API_ENABLED
diff --git a/Engine/Plugins/Experimental/ChaosRigidPhysicsAsync/Tests/Private/RigidTestUtils.cpp b/Engine/Plugins/Experimental/ChaosRigidPhysicsAsync/Tests/Private/RigidTestUtils.cpp
--- a/Engine/Plugins/Experimental/ChaosRigidPhysicsAsync/Tests/Private/RigidTestUtils.cpp
+++ b/Engine/Plugins/Experimental/ChaosRigidPhysicsAsync/Tests/Private/RigidTestUtils.cpp
@@ -47,6 +47,26 @@ namespace UE::Physics
 		};
 		return Setup;
 	}
+	
+	FConvexGeometrySetup MakeConvexBoxGeometrySetup(const FVector3f& InCenter, const FVector3f& InHalfExtent, const float InMargin)
+	{
+		FConvexGeometrySetup Setup
+		{
+			.Vertices
+			{
+				InCenter + InHalfExtent * FVector3f(-1, -1, -1),
+				InCenter + InHalfExtent * FVector3f( 1, -1, -1),
+				InCenter + InHalfExtent * FVector3f(-1,  1, -1),
+				InCenter + InHalfExtent * FVector3f( 1,  1, -1),
+				InCenter + InHalfExtent * FVector3f(-1, -1,  1),
+				InCenter + InHalfExtent * FVector3f( 1, -1,  1),
+				InCenter + InHalfExtent * FVector3f(-1,  1,  1),
+				InCenter + InHalfExtent * FVector3f( 1,  1,  1),
+			},
+			.Margin = InMargin
+		};
+		return Setup;
+	}
 
 	FRigidShapeInstanceSetup MakeBoxShape(const FVector3f& InSize, const FVector3f& InCenter, UPhysicalMaterial* InMaterial)
 	{
diff --git a/Engine/Plugins/Experimental/ChaosRigidPhysicsAsync/Tests/Private/RigidTestUtils.h b/Engine/Plugins/Experimental/ChaosRigidPhysicsAsync/Tests/Private/RigidTestUtils.h
--- a/Engine/Plugins/Experimental/ChaosRigidPhysicsAsync/Tests/Private/RigidTestUtils.h
+++ b/Engine/Plugins/Experimental/ChaosRigidPhysicsAsync/Tests/Private/RigidTestUtils.h
@@ -21,6 +21,7 @@ namespace UE::Physics
 {
 	FTriangleMeshGeometrySetup MakeSingleTriangleGeometrySetup(const FVector3f& InScale = FVector3f::One(), const FVector3f& InTranslation = FVector3f::Zero());
 	FTriangleMeshGeometrySetup MakeQuadTriangleGeometrySetup(const FVector3f& InScale = FVector3f::One(), const FVector3f& InTranslation = FVector3f::Zero());
+	FConvexGeometrySetup MakeConvexBoxGeometrySetup(const FVector3f& InCenter = FVector3f::Zero(), const FVector3f& InHalfExtent = FVector3f::One(), const float InMargin = 0);
 	FRigidShapeInstanceSetup MakeBoxShape(const FVector3f& InSize, const FVector3f& InCenter = FVector3f::Zero(), UPhysicalMaterial* InMaterial = nullptr);
 	FVector3f MakeSolidBoxInertia(float InMass, const FVector3f& InSize);
 } // namespace UE::Physics
diff --git a/Engine/Source/Developer/ChaosTestHarness/Private/Init.cpp b/Engine/Source/Developer/ChaosTestHarness/Private/Init.cpp
--- a/Engine/Source/Developer/ChaosTestHarness/Private/Init.cpp
+++ b/Engine/Source/Developer/ChaosTestHarness/Private/Init.cpp
@@ -3,6 +3,7 @@
 #include "CoreMinimal.h"
 #include "Chaos/LowLevelTest/ChaosTestHarness.h"
 
+#include "CompGeom/ExactPredicates.h"
 #include "Logging/LogScopedVerbosityOverride.h"
 #include "TestCommon/Initialization.h"
 
@@ -19,6 +20,7 @@ GROUP_BEFORE_GLOBAL(Catch::DefaultGroup)
 {
 	// Find the engine directory before InitAll() enables the platform file stub. Prevents warnings.
 	FPlatformMisc::EngineDir();
+	UE::Geometry::ExactPredicates::GlobalInit();
 
 	// Temporarily disable warnings during init
 	LOG_SCOPE_VERBOSITY_OVERRIDE(LogSlate, ELogVerbosity::Error);
diff --git a/Engine/Source/Runtime/Experimental/RigidPhysics/Private/RigidPhysics/Geometry/ConvexGeometry.cpp b/Engine/Source/Runtime/Experimental/RigidPhysics/Private/RigidPhysics/Geometry/ConvexGeometry.cpp
--- a/Engine/Source/Runtime/Experimental/RigidPhysics/Private/RigidPhysics/Geometry/ConvexGeometry.cpp
+++ b/Engine/Source/Runtime/Experimental/RigidPhysics/Private/RigidPhysics/Geometry/ConvexGeometry.cpp
@@ -3,13 +3,64 @@
 
 #if UE_RIGIDPHYSICS_API_ENABLED
 
+#include "Chaos/CollisionConvexMesh.h"
+#include "Chaos/Convex.h"
+
 namespace UE::Physics
 {
-	FConvexGeometry::FConvexGeometry(Chaos::FImplicitConvex3* InImplicit)
+	static Chaos::FConvexBuilder::EBuildMethod Convert(const FConvexGeometrySetup::EBuildMethod InMethod)
+	{
+		switch (InMethod)
+		{
+			case FConvexGeometrySetup::EBuildMethod::Default:             return Chaos::FConvexBuilder::EBuildMethod::Default;
+			case FConvexGeometrySetup::EBuildMethod::Original:            return Chaos::FConvexBuilder::EBuildMethod::Original;
+			case FConvexGeometrySetup::EBuildMethod::ConvexHull3:         return Chaos::FConvexBuilder::EBuildMethod::ConvexHull3;
+			case FConvexGeometrySetup::EBuildMethod::ConvexHull3Simplified: return Chaos::FConvexBuilder::EBuildMethod::ConvexHull3Simplified;
+			default:
+			{
+				ensureMsgf(false, TEXT("Unhandled conversion for FConvexGeometrySetup::EBuildMethod"));
+				return Chaos::FConvexBuilder::EBuildMethod::Default;
+			}
+		}
+	}
+
+	FConvexGeometry::~FConvexGeometry() = default;
+	FConvexGeometry::FConvexGeometry(const FConvexGeometry&) = default;
+	FConvexGeometry::FConvexGeometry(FConvexGeometry&&) = default;
+	FConvexGeometry& FConvexGeometry::operator=(const FConvexGeometry&) = default;
+	FConvexGeometry& FConvexGeometry::operator=(FConvexGeometry&&) = default;
+
+	FConvexGeometry::FConvexGeometry(FConvexGeometrySetup&& InSetup)
+	{
+		TArray<Chaos::FVec3f> Vertices;
+		Vertices.SetNum(InSetup.Vertices.Num());
+		for (int32 I = 0; I < InSetup.Vertices.Num(); ++I)
+		{
+			Vertices[I] = InSetup.Vertices[I];
+		}
+		Implicit = new Chaos::FConvex(Vertices, InSetup.Margin, Convert(InSetup.BuildMethod));
+	}
+
+	FConvexGeometry::FConvexGeometry(Chaos::FConvex* InImplicit)
 		: Implicit(InImplicit)
 	{
 	}
 
+	int32 FConvexGeometry::GetNumVertices() const
+	{
+		return Implicit->NumVertices();
+	}
+
+	FVector3f FConvexGeometry::GetVertex(const int32 InIndex) const
+	{
+		return FVector3f(Implicit->GetVertex(InIndex));
+	}
+
+	float FConvexGeometry::GetMargin() const
+	{
+		return (float)Implicit->GetMargin();
+	}
+
 	FString FConvexGeometry::ToString() const
 	{
 		return TEXT("Convex");
diff --git a/Engine/Source/Runtime/Experimental/RigidPhysics/Public/RigidPhysics/Geometry/ConvexGeometry.h b/Engine/Source/Runtime/Experimental/RigidPhysics/Public/RigidPhysics/Geometry/ConvexGeometry.h
--- a/Engine/Source/Runtime/Experimental/RigidPhysics/Public/RigidPhysics/Geometry/ConvexGeometry.h
+++ b/Engine/Source/Runtime/Experimental/RigidPhysics/Public/RigidPhysics/Geometry/ConvexGeometry.h
@@ -5,24 +5,53 @@
 
 #if UE_RIGIDPHYSICS_API_ENABLED
 
-#include "Chaos/Convex.h"
+namespace Chaos
+{
+	class FConvex;
+} // namespace Chaos
 
 UE_EXPERIMENTAL(5.8, "The new Chaos API is experimental")
 namespace UE::Physics
 {
+	class FConvexGeometrySetup
+	{
+	public:
+		enum class EBuildMethod : uint8
+		{
+			Default = 0,
+			Original = 1,
+			ConvexHull3 = 2,
+			ConvexHull3Simplified = 3,
+		};
+
+		TArray<FVector3f> Vertices;
+		float Margin = 0.f;
+		EBuildMethod BuildMethod = EBuildMethod::Default;
+	};
+
 	class FConvexGeometry
 	{
 	public:
-		UE_INTERNAL FConvexGeometry() = default;
-		UE_INTERNAL RIGIDPHYSICS_API FConvexGeometry(Chaos::FImplicitConvex3* InImplicit);
+		FConvexGeometry() = delete;
+		RIGIDPHYSICS_API ~FConvexGeometry();
+		RIGIDPHYSICS_API FConvexGeometry(const FConvexGeometry&);
+		RIGIDPHYSICS_API FConvexGeometry(FConvexGeometry&&);
+		RIGIDPHYSICS_API FConvexGeometry& operator=(const FConvexGeometry&);
+		RIGIDPHYSICS_API FConvexGeometry& operator=(FConvexGeometry&&);
+		RIGIDPHYSICS_API FConvexGeometry(FConvexGeometrySetup&& InSetup);
+		UE_INTERNAL RIGIDPHYSICS_API FConvexGeometry(Chaos::FConvex* InImplicit);
 
 		bool operator==(const FConvexGeometry&) const = default;
 		bool operator!=(const FConvexGeometry&) const = default;
 
+		RIGIDPHYSICS_API int32 GetNumVertices() const;
+		RIGIDPHYSICS_API FVector3f GetVertex(const int32 InIndex) const;
+		RIGIDPHYSICS_API float GetMargin() const;
+
 		RIGIDPHYSICS_API FString ToString() const;
 
 	private:
-		Chaos::FConvexPtr Implicit;
+		TRefCountPtr<Chaos::FConvex> Implicit;
 	};
 } // namespace UE::Physics
 
