Chaos API: Create HeightFieldGeometry with a basic public API.

#Synced-CL 52100670

diff --git a/Engine/Plugins/Experimental/ChaosRigidPhysicsAsync/Tests/Private/HeightFieldGeometryTests.cpp b/Engine/Plugins/Experimental/ChaosRigidPhysicsAsync/Tests/Private/HeightFieldGeometryTests.cpp
new file mode 100644
--- /dev/null
+++ b/Engine/Plugins/Experimental/ChaosRigidPhysicsAsync/Tests/Private/HeightFieldGeometryTests.cpp
@@ -0,0 +1,56 @@
+// Copyright Epic Games, Inc. All Rights Reserved.
+
+#include "Chaos/LowLevelTest/ChaosTestHarness.h"
+#include "ChaosRigidPhysicsAsyncTest.h"
+
+#include "RigidPhysics/Geometry/HeightFieldGeometry.h"
+#include "RigidTestUtils.h"
+
+#if UE_RIGIDPHYSICS_API_ENABLED
+
+namespace Chaos::LowLevelTest
+{
+	using FHeightFieldGeometry = UE::Physics::FHeightFieldGeometry;
+	using FHeightFieldGeometrySetup = UE::Physics::FHeightFieldGeometrySetup;
+
+	FHeightFieldGeometrySetup MakeHeightFieldGeometrySetup()
+	{
+		FHeightFieldGeometrySetup Setup
+		{
+			.Heights
+			{
+				0, 1, 2,
+				3, 4, 5,
+				6, 7, 8
+			},
+			.MaterialIndices = { 0, 1, 2, 3 },
+			.NumRows = 3,
+			.NumCols = 3,
+			.Scale = FVector3f(100.f, 100.f, 100.f),
+		};
+		return Setup;
+	}
+
+	TEST_CASE("HeightFieldGeometry::Basic", "[Chaos][API][Geometry][unit]")
+	{
+		const FHeightFieldGeometrySetup Setup = MakeHeightFieldGeometrySetup();
+		FHeightFieldGeometrySetup MutableSetup = Setup;
+
+		FHeightFieldGeometry Geometry(MoveTemp(MutableSetup));
+
+		const int32 NumQuads = (Setup.NumRows - 1) * (Setup.NumCols - 1);
+		REQUIRE(Geometry.GetNumRows() == Setup.NumRows);
+		REQUIRE(Geometry.GetNumCols() == Setup.NumCols);
+		CHECK_THAT(Geometry.GetScale(), Catch::ApproxEq(Setup.Scale));
+		for (int32 I = 0; I < Setup.Heights.Num(); ++I)
+		{
+			CHECK_THAT(Setup.Heights[I], Catch::Matchers::WithinRel(Geometry.GetHeight(I), 0.001f));
+		}
+		for (int32 I = 0; I < Setup.MaterialIndices.Num(); ++I)
+		{
+			CHECK(Setup.MaterialIndices[I] == Geometry.GetMaterialIndex(I));
+		}
+	}
+} // namespace Chaos::LowLevelTest
+
+#endif // UE_RIGIDPHYSICS_API_ENABLED
diff --git a/Engine/Source/Runtime/Experimental/Chaos/Public/Chaos/HeightField.h b/Engine/Source/Runtime/Experimental/Chaos/Public/Chaos/HeightField.h
--- a/Engine/Source/Runtime/Experimental/Chaos/Public/Chaos/HeightField.h
+++ b/Engine/Source/Runtime/Experimental/Chaos/Public/Chaos/HeightField.h
@@ -26,6 +26,11 @@ namespace Chaos::Private
 	class FMeshContactGenerator;
 }
 
+namespace UE::Physics
+{
+	class FHeightFieldGeometry;
+}
+
 namespace Chaos
 {
 	class FHeightField final : public FImplicitObject
@@ -1123,6 +1128,7 @@ namespace Chaos
 		// Needed for serialization
 		FHeightField() : FImplicitObject(EImplicitObject::HasBoundingBox, ImplicitObjectType::HeightField) {}
 		friend FImplicitObject;
+		friend UE::Physics::FHeightFieldGeometry;
 
 		template <typename QueryGeomType>
 		bool OverlapGeomImp(const QueryGeomType& QueryGeom, const FRigidTransform3& QueryTM, const FReal Thickness, FMTDInfo* OutMTD = nullptr) const;
diff --git a/Engine/Source/Runtime/Experimental/RigidPhysics/Private/RigidPhysics/Geometry/HeightFieldGeometry.cpp b/Engine/Source/Runtime/Experimental/RigidPhysics/Private/RigidPhysics/Geometry/HeightFieldGeometry.cpp
new file mode 100644
--- /dev/null
+++ b/Engine/Source/Runtime/Experimental/RigidPhysics/Private/RigidPhysics/Geometry/HeightFieldGeometry.cpp
@@ -0,0 +1,81 @@
+// Copyright Epic Games, Inc. All Rights Reserved.
+#include "RigidPhysics/Geometry/HeightFieldGeometry.h"
+
+#if UE_RIGIDPHYSICS_API_ENABLED
+
+#include "Chaos/HeightField.h"
+
+namespace UE::Physics
+{
+	FHeightFieldGeometry::~FHeightFieldGeometry() = default;
+	FHeightFieldGeometry::FHeightFieldGeometry(const FHeightFieldGeometry&) = default;
+	FHeightFieldGeometry::FHeightFieldGeometry(FHeightFieldGeometry&&) = default;
+	FHeightFieldGeometry& FHeightFieldGeometry::operator=(const FHeightFieldGeometry&) = default;
+	FHeightFieldGeometry& FHeightFieldGeometry::operator=(FHeightFieldGeometry&&) = default;
+
+	FHeightFieldGeometry::FHeightFieldGeometry(FHeightFieldGeometrySetup&& InSetup)
+	{
+		if (!ensure(InSetup.NumCols * InSetup.NumRows == InSetup.Heights.Num()))
+		{
+			InSetup.Heights.SetNumZeroed(InSetup.NumCols * InSetup.NumRows);
+		}
+		const int32 NumQuads = FMath::Max(0, (InSetup.NumCols - 1) * (InSetup.NumRows - 1));
+		if (!ensure(NumQuads == InSetup.MaterialIndices.Num()))
+		{
+			InSetup.MaterialIndices.SetNumZeroed(NumQuads);
+		}
+
+		TArray<Chaos::FReal> Heights;
+		Heights.SetNum(InSetup.Heights.Num());
+		for (int32 I = 0; I < InSetup.Heights.Num(); ++I)
+		{
+			Heights[I] = (Chaos::FReal)InSetup.Heights[I];
+		}
+		const Chaos::FVec3 Scale(InSetup.Scale);
+
+		Implicit = new Chaos::FHeightField(
+			MoveTemp(Heights),
+			MoveTemp(InSetup.MaterialIndices),
+			InSetup.NumRows,
+			InSetup.NumCols,
+			Scale
+		);
+	}
+
+	FHeightFieldGeometry::FHeightFieldGeometry(Chaos::FHeightField* InImplicit)
+		: Implicit(InImplicit)
+	{
+	}
+
+	int32 FHeightFieldGeometry::GetNumRows() const
+	{
+		return Implicit->GetNumRows();
+	}
+
+	int32 FHeightFieldGeometry::GetNumCols() const
+	{
+		return Implicit->GetNumCols();
+	}
+
+	FVector3f FHeightFieldGeometry::GetScale() const
+	{
+		return FVector3f(Implicit->GeomData.Scale);
+	}
+
+	float FHeightFieldGeometry::GetHeight(const int32 InIndex) const
+	{
+		return (float)Implicit->GetHeight(InIndex);
+	}
+
+	uint8 FHeightFieldGeometry::GetMaterialIndex(const int32 InIndex) const
+	{
+		return Implicit->GetMaterialIndex(InIndex);
+	}
+
+	FString FHeightFieldGeometry::ToString() const
+	{
+		return TEXT("HeightField");
+	}
+} // namespace UE::Physics
+
+#endif // UE_RIGIDPHYSICS_API_ENABLED
diff --git a/Engine/Source/Runtime/Experimental/RigidPhysics/Public/RigidPhysics/Geometry/HeightFieldGeometry.h b/Engine/Source/Runtime/Experimental/RigidPhysics/Public/RigidPhysics/Geometry/HeightFieldGeometry.h
new file mode 100644
--- /dev/null
+++ b/Engine/Source/Runtime/Experimental/RigidPhysics/Public/RigidPhysics/Geometry/HeightFieldGeometry.h
@@ -0,0 +1,55 @@
+// Copyright Epic Games, Inc. All Rights Reserved.
+#pragma once
+
+#include "RigidPhysics/RigidFwd.h"
+
+#if UE_RIGIDPHYSICS_API_ENABLED
+
+namespace Chaos
+{
+	class FHeightField;
+} // namespace Chaos
+
+UE_EXPERIMENTAL(5.8, "The new Chaos API is experimental")
+namespace UE::Physics
+{
+	class FHeightFieldGeometrySetup
+	{
+	public:
+		TArray<float> Heights;
+		TArray<uint8> MaterialIndices;
+		int32 NumRows = 0;
+		int32 NumCols = 0;
+		FVector3f Scale = FVector3f::OneVector;
+	};
+
+	class FHeightFieldGeometry
+	{
+	public:
+		FHeightFieldGeometry() = delete;
+		RIGIDPHYSICS_API ~FHeightFieldGeometry();
+		RIGIDPHYSICS_API FHeightFieldGeometry(const FHeightFieldGeometry&);
+		RIGIDPHYSICS_API FHeightFieldGeometry(FHeightFieldGeometry&&);
+		RIGIDPHYSICS_API FHeightFieldGeometry& operator=(const FHeightFieldGeometry&);
+		RIGIDPHYSICS_API FHeightFieldGeometry& operator=(FHeightFieldGeometry&&);
+		RIGIDPHYSICS_API FHeightFieldGeometry(FHeightFieldGeometrySetup&& InSetup);
+		UE_INTERNAL RIGIDPHYSICS_API FHeightFieldGeometry(Chaos::FHeightField* InImplicit);
+
+		bool operator==(const FHeightFieldGeometry&) const = default;
+		bool operator!=(const FHeightFieldGeometry&) const = default;
+
+		RIGIDPHYSICS_API int32 GetNumRows() const;
+		RIGIDPHYSICS_API int32 GetNumCols() const;
+		RIGIDPHYSICS_API FVector3f GetScale() const;
+
+		RIGIDPHYSICS_API float GetHeight(const int32 InIndex) const;
+		RIGIDPHYSICS_API uint8 GetMaterialIndex(const int32 InIndex) const;
+
+		RIGIDPHYSICS_API FString ToString() const;
+
+	private:
+		TRefCountPtr<Chaos::FHeightField> Implicit;
+	};
+} // namespace UE::Physics
+
+#endif // UE_RIGIDPHYSICS_API_ENABLED
