Fixing a bughawk issue with FHeightFieldGeometry.

#Synced-CL 52100670

diff --git a/Engine/Source/Runtime/Experimental/RigidPhysics/Private/RigidPhysics/Geometry/HeightFieldGeometry.cpp b/Engine/Source/Runtime/Experimental/RigidPhysics/Private/RigidPhysics/Geometry/HeightFieldGeometry.cpp
--- a/Engine/Source/Runtime/Experimental/RigidPhysics/Private/RigidPhysics/Geometry/HeightFieldGeometry.cpp
+++ b/Engine/Source/Runtime/Experimental/RigidPhysics/Private/RigidPhysics/Geometry/HeightFieldGeometry.cpp
@@ -19,7 +19,9 @@ namespace UE::Physics
 		{
 			InSetup.Heights.SetNumZeroed(InSetup.NumCols * InSetup.NumRows);
 		}
-		const int32 NumQuads = FMath::Max(0, (InSetup.NumCols - 1) * (InSetup.NumRows - 1));
+		const int32 NumCellCols = FMath::Max(0, InSetup.NumCols - 1);
+		const int32 NumCellRows = FMath::Max(0, InSetup.NumRows - 1);
+		const int32 NumQuads = NumCellCols * NumCellRows;
 		if (!ensure(NumQuads == InSetup.MaterialIndices.Num()))
 		{
 			InSetup.MaterialIndices.SetNumZeroed(NumQuads);
