Dataflow : Add auto convert from asset to string (path)

#Synced-CL 52092384

diff --git a/Engine/Source/Runtime/Experimental/Dataflow/Engine/Private/Dataflow/DataflowAssetPathNodes.cpp b/Engine/Source/Runtime/Experimental/Dataflow/Engine/Private/Dataflow/DataflowAssetPathNodes.cpp
--- a/Engine/Source/Runtime/Experimental/Dataflow/Engine/Private/Dataflow/DataflowAssetPathNodes.cpp
+++ b/Engine/Source/Runtime/Experimental/Dataflow/Engine/Private/Dataflow/DataflowAssetPathNodes.cpp
@@ -2,17 +2,32 @@
 
 #include "Dataflow/DataflowAssetPathNodes.h"
 #include "Dataflow/DataflowNodeFactory.h"
+#include "Dataflow/DataflowAnyTypeRegistry.h"
 
 #include UE_INLINE_GENERATED_CPP_BY_NAME(DataflowAssetPathNodes)
 
 #define LOCTEXT_NAMESPACE "DataflowAssetPathNodes"
 
+class UStaticMesh;
+class USkeletalMesh;
+class USkeleton;
+
 namespace UE::Dataflow
 {
+	namespace Private
+	{
+		static bool AutoConvertUObjectConvertibleToString(FName FromType, FName ToType)
+		{
+			return FDataflowUObjectConvertibleTypePolicy::SupportsTypeStatic(FromType);
+		}
+	}
+
 	void RegisterAssetPathNodes()
 	{
 		DATAFLOW_NODE_REGISTER_CREATION_FACTORY(FDataflowGetAssetPathNode);
 		DATAFLOW_NODE_REGISTER_CREATION_FACTORY(FDataflowPathReplaceNode);
+
+		UE_DATAFLOW_REGISTER_AUTOCONVERT_WITH_FILTER(FDataflowUObjectConvertibleTypes, FString, FDataflowGetAssetPathNode, &Private::AutoConvertUObjectConvertibleToString);
 	}
 }
 
