[RemoteSessionsManager] Allow SendPong in same-process scenarios without a network bridge (e.g., PIE) where all messages are delivered locally on
the same bus and no bridge is needed.

#Synced-CL 52126719

diff --git a/Engine/Source/Runtime/Experimental/TraceBasedDebuggers/Private/RemoteSessionsManager.cpp b/Engine/Source/Runtime/Experimental/TraceBasedDebuggers/Private/RemoteSessionsManager.cpp
--- a/Engine/Source/Runtime/Experimental/TraceBasedDebuggers/Private/RemoteSessionsManager.cpp
+++ b/Engine/Source/Runtime/Experimental/TraceBasedDebuggers/Private/RemoteSessionsManager.cpp
@@ -543,7 +543,10 @@ void FRemoteSessionsManager::HandleSessionPingMessage(const FSessionPing& InMess
 	// Only respond if we can actually publish messages back to the controller.
 	// Without a network bridge, status messages (FRecordingStatusMessage, etc.) won't
 	// reach the controller, resulting in a non-functional ghost session.
-	if (CanPublishToNetwork<FSessionPong>())
+	// or in same-process scenarios without a network bridge (e.g., PIE) where all messages are delivered locally on
+	// the same bus and no bridge is needed
+	if (CanPublishToNetwork<FSessionPong>()
+		|| InMessage.ControllerInstanceId == FApp::GetInstanceId())
 	{
 		SendPong(InMessage, InContext);
 	}
