IAD Cache - Implement OpenMapped for files in MMap CAS.
Disabled by default.

#Synced-CL 52124292

diff --git a/Engine/Source/Runtime/Core/Private/IO/IoDispatcher.cpp b/Engine/Source/Runtime/Core/Private/IO/IoDispatcher.cpp
--- a/Engine/Source/Runtime/Core/Private/IO/IoDispatcher.cpp
+++ b/Engine/Source/Runtime/Core/Private/IO/IoDispatcher.cpp
@@ -79,7 +79,9 @@ static const TCHAR* const GetIoErrorText_ErrorCodeTextArray[] =
 	TEXT("FileMove Failed"),
 	TEXT("FileClose Failed"),
 	TEXT("File Corrupt"),
-	TEXT("Http Read Error")
+	TEXT("Http Read Error"),
+	TEXT("File Handle Open"),
+	TEXT("Pending Delete")
 };
 static_assert(std::size(GetIoErrorText_ErrorCodeTextArray) == static_cast<uint32>(EIoErrorCode::Last), "GetIoErrorText_ErrorCodeTextArray is missing entries to match EIoErrorCode. Please update the array.");
 
@@ -604,23 +606,37 @@ public:
 
 	TIoStatusOr<FIoMappedRegion> OpenMapped(const FIoChunkId& ChunkId, const FIoReadOptions& Options)
 	{
+		TIoStatusOr<FIoMappedRegion> OutResult{FIoStatus(EIoErrorCode::NotFound)};
+
 		if (ChunkId.IsValid())
 		{
 			FReadScopeLock _(BackendsLock);
+
 			for (const FBackendAndPriority& Backend : Backends)
 			{
 				TIoStatusOr<FIoMappedRegion> Result = Backend.Value->OpenMapped(ChunkId, Options);
 				if (Result.IsOk())
 				{
-					return Result;
+					OutResult = MoveTemp(Result);
+					break;
+				}
+
+				if (Result.Status().GetErrorCode() != EIoErrorCode::NotFound)
+				{
+					if (OutResult.Status().GetErrorCode() == EIoErrorCode::NotFound ||
+						OutResult.Status().GetErrorCode() == EIoErrorCode::UnknownChunkID)
+					{
+						OutResult = MoveTemp(Result);
+					}
 				}
 			}
-			return FIoStatus(EIoErrorCode::NotFound);
 		}
 		else
 		{
-			return FIoStatus(EIoErrorCode::InvalidParameter, TEXT("FIoChunkId is not valid"));
+			OutResult = FIoStatus(EIoErrorCode::InvalidParameter, TEXT("FIoChunkId is not valid"));
 		}
+
+		return OutResult;
 	}
 
 	void Mount(TSharedRef<IIoDispatcherBackend> Backend, int32 Priority)
diff --git a/Engine/Source/Runtime/Core/Public/IO/IoStatus.h b/Engine/Source/Runtime/Core/Public/IO/IoStatus.h
--- a/Engine/Source/Runtime/Core/Public/IO/IoStatus.h
+++ b/Engine/Source/Runtime/Core/Public/IO/IoStatus.h
@@ -46,6 +46,8 @@ enum class EIoErrorCode
 	FileCloseFailed,
 	FileCorrupt,
 	HttpReadError,
+	FileHandleOpen,
+	PendingDelete,
 	Last
 };
 
diff --git a/Engine/Source/Runtime/Experimental/IoStore/OnDemand/Private/OnDemandContentInstaller.cpp b/Engine/Source/Runtime/Experimental/IoStore/OnDemand/Private/OnDemandContentInstaller.cpp
--- a/Engine/Source/Runtime/Experimental/IoStore/OnDemand/Private/OnDemandContentInstaller.cpp
+++ b/Engine/Source/Runtime/Experimental/IoStore/OnDemand/Private/OnDemandContentInstaller.cpp
@@ -816,6 +816,9 @@ void FOnDemandContentInstaller::PinCachedChunks(FOnDemandContentInstaller::FRequ
 			if (bCached)
 			{
 				CachedEntryIndices.Add(EntryIndex);
+
+				const FIoChunkId& ChunkId = ResolvedChunks.Container->ChunkIds[EntryIndex];
+				InstallRequest.bUsesMMapCas = InstallRequest.bUsesMMapCas || GetCasType(ChunkId.GetChunkType()) == EOnDemandInstallCasType::MMap;
 			}
 			
 			OnChunkFound(ContainerIndex, EntryIndex, bCached);
@@ -1544,13 +1547,6 @@ void FOnDemandContentInstaller::OnProcessDownloadedChunkNotifyRequest(
 	const bool bCompleted = ++InstallRequest.DownloadedChunkCount >= InstallRequest.HttpRequestHandles.Num();
 	if (bCompleted)
 	{
-		// Cannot read from pending buffer for MMap chunks, must flush to disk
-		if (Request.IsOk() && InstallRequest.bUsesMMapCas)
-		{
-			// TODO: this could be a partial flush if pending chunks was sorted by request
-			Request.Result = IoStore.InstallCache->Flush(EOnDemandInstallCasType::MMap);
-		}
-
 		CompleteInstallRequest(Request);
 	}
 	else
@@ -1674,10 +1670,18 @@ void FOnDemandContentInstaller::CompleteInstallRequest(FRequest& Request)
 {
 	TRACE_CPUPROFILER_EVENT_SCOPE(FOnDemandContentInstaller::CompleteInstallRequest);
 
-	using namespace UE::IoStore::Private;
-
 	FRequest::FInstall& InstallRequest = Request.AsInstall();
 
+	// Cannot read from pending buffer for MMap chunks, must flush to disk
+	// This must be called in CompleteInstallRequest so it gets called 
+	// even if a reqeust doesn't queue up any downloads because some chunks may still be
+	// in the pending buffer.
+	if (Request.IsOk() && InstallRequest.bUsesMMapCas)
+	{
+		// TODO: this could be a partial flush if pending chunks was sorted by request
+		Request.Result = IoStore.InstallCache->Flush(EOnDemandInstallCasType::MMap);
+	}
+
 #if !UE_BUILD_SHIPPING
 	TArray<FString> DebugInfo;
 #endif
@@ -1698,7 +1702,7 @@ void FOnDemandContentInstaller::CompleteInstallRequest(FRequest& Request)
 			FOnDemandContentHandle& ContentHandle = InstallRequest.Args.ContentHandle;
 
 			TUniqueLock Lock(IoStore.ContainerMutex);
-			for (FResolvedContainerChunks& ResolvedChunks : InstallRequest.ResolvedChunks)
+			for (Private::FResolvedContainerChunks& ResolvedChunks : InstallRequest.ResolvedChunks)
 			{
 				const FSharedOnDemandContainer& Container = ResolvedChunks.Container;
 				FOnDemandChunkEntryReferences& References = Container->FindOrAddChunkEntryReferences(*ContentHandle.Handle);
diff --git a/Engine/Source/Runtime/Experimental/IoStore/OnDemand/Private/OnDemandContentInstaller.h b/Engine/Source/Runtime/Experimental/IoStore/OnDemand/Private/OnDemandContentInstaller.h
--- a/Engine/Source/Runtime/Experimental/IoStore/OnDemand/Private/OnDemandContentInstaller.h
+++ b/Engine/Source/Runtime/Experimental/IoStore/OnDemand/Private/OnDemandContentInstaller.h
@@ -104,7 +104,7 @@ class FOnDemandContentInstaller
 			uint64										DownloadedChunkCount = 0;
 			uint64										LastProgressCycles = 0;
 			bool										bHttpRequestsIssued = false;
-			bool										bUsesMMapCas = false;
+			bool										bUsesMMapCas = false; // Fine for now, but this will need to be bit flags if we have more than one CAS that needs flushing
 			std::atomic_bool							bNotifyingProgressOnGameThread{false};
 		};
 
diff --git a/Engine/Source/Runtime/Experimental/IoStore/OnDemand/Private/OnDemandInstallCache.cpp b/Engine/Source/Runtime/Experimental/IoStore/OnDemand/Private/OnDemandInstallCache.cpp
--- a/Engine/Source/Runtime/Experimental/IoStore/OnDemand/Private/OnDemandInstallCache.cpp
+++ b/Engine/Source/Runtime/Experimental/IoStore/OnDemand/Private/OnDemandInstallCache.cpp
@@ -11,6 +11,7 @@
 #include "Algo/Find.h"
 #include "Algo/Transform.h"
 #include "Async/Mutex.h"
+#include "Async/MappedFileHandle.h"
 #include "Async/SharedMutex.h"
 #include "Async/UniqueLock.h"
 #include "Async/SharedLock.h"
@@ -102,6 +103,47 @@ using FSharedFileHandle				= TSharedPtr<IFileHandle>;
 using FSharedAsyncFileHandle		= TSharedPtr<IAsyncReadFileHandle>;
 using FWeakAsyncFileHandle			= TWeakPtr<IAsyncReadFileHandle>;
 
+using FSharedMappedFileHandle		= TSharedPtr<IMappedFileHandle>;
+using FSharedMappedFileHandleRef	= TSharedRef<IMappedFileHandle>;
+using FWeakMappedFileHandle			= TWeakPtr<IMappedFileHandle>;
+
+///////////////////////////////////////////////////////////////////////////////
+class FSharedMappedFileHandleProxy final : public IMappedFileHandle
+{
+public:
+	FSharedMappedFileHandleProxy(FSharedMappedFileHandleRef&& InSharedMappedFileHandle)
+		: IMappedFileHandle(InSharedMappedFileHandle->GetFileSize())
+		, SharedMappedFileHandle(MoveTemp(InSharedMappedFileHandle))
+	{}
+
+	virtual IMappedFileRegion* MapRegion(int64 Offset = 0, int64 BytesToMap = MAX_int64, FFileMappingFlags Flags = EMappedFileFlags::ENone) override
+	{
+		// GetFileSize() is the size when the file was opened. Some platforms use current file size instead, so enforce BytesToMap < GetFileSize() here.
+		check(Offset < GetFileSize()); // don't map zero bytes and don't map off the end of the file
+		BytesToMap = FMath::Min<int64>(BytesToMap, GetFileSize() - Offset);
+		check(BytesToMap > 0); // don't map zero bytes
+
+		if (EnumHasAnyFlags(Flags.Flags, EMappedFileFlags::EFileWritable))
+		{
+			// Only allow read access through the mapped handle.
+			checkf(false, TEXT("Writes to the OnDeamnd cache are only allowed from the OnDemandContentInstaller!"));
+			return nullptr;
+		}
+
+		return SharedMappedFileHandle->MapRegion(Offset, BytesToMap, Flags);
+	}
+
+	virtual void Flush() override
+	{
+		SharedMappedFileHandle->Flush();
+	}
+
+private:
+	FSharedMappedFileHandleRef SharedMappedFileHandle;
+};
+
+
+///////////////////////////////////////////////////////////////////////////////
 using FCasAddr						= FHash96;
 static const FCasAddr& AsCasAddr(const FIoHash& IoHash)
 {
@@ -885,9 +927,11 @@ private:
 struct FCasBlockState
 {
 	uint32 Size = 0;
-	uint32 ReadHandleCount = 0;
+	uint32 HandleCount = 0;
 	int64 LastAccessUtcTicks = 0;
 	int64 LastModificationUtcTicks = 0;
+	TVariant<FWeakAsyncFileHandle, FWeakMappedFileHandle> WeakHandle; // TODO: Possibly make two flavors of FCasBlockState to avoid TVariant overhead
+	TOptional<FSharedEventRef> PendingDefragEvent;	// Currently only used for MMap CAS
 };
 
 struct FCas
@@ -901,26 +945,38 @@ struct FCas
 	using FBlockStates			= TMap<FCasBlockId, FCasBlockState>;
 	using FBlockFileSizes		= TMap<FCasBlockId, uint64>;
 
-	FCas();
+	struct FMaybeLock
+	{
+		TDynamicUniqueLock<const FCas> Lock;
+
+		FMaybeLock(const FCas* Cas, const UE::FDeferLock* SkipLock)
+			: Lock(SkipLock ? TDynamicUniqueLock(*Cas, *SkipLock) : TDynamicUniqueLock(*Cas))
+		{}
+	};
 
-	void						Lock()		{ Mutex.Lock(); }
-	void						Unlock()	{ Mutex.Unlock(); }
+	FCas(EOnDemandInstallCasType InType);
+
+	void						Lock() const		{ Mutex.Lock(); }
+	void						Unlock() const		{ Mutex.Unlock(); }
 
 	FResult						Configure(const FOnDemandInstallCasConfig& Config);
 	FResult						Initialize(FStringView Directory, bool bDeleteExisting = false);
-	FCasLocation				FindChunk(const FIoHash& Hash) const;
+	FCasLocation				FindChunk(const FIoHash& Hash, const UE::FDeferLock* SkipLock = nullptr) const;
 	FCasBlockId					CreateBlock();
 	bool						ContainsBlock(FCasBlockId BlockId) const;
 	FResult						DeleteBlock(FCasBlockId BlockId, TFunctionRef<void(const FCasAddr&)> OutAddrs);
 	FString						GetBlockFilename(FCasBlockId BlockId) const;
 	TResult<FSharedFileHandle>	OpenRead(FCasBlockId BlockId);
 	TResult<FSharedAsyncFileHandle>	OpenAsyncRead(FCasBlockId BlockId);
+	TResult<FSharedMappedFileHandle> OpenMapped(FCasBlockId BlockId, int64 MaxSize, TOptional<FSharedEventRef>* OutBlockDeletedEvent = nullptr, const UE::FDeferLock* SkipLock = nullptr);
 	void						OnFileHandleDeleted(FCasBlockId BlockId);
+	bool						SetPendingDefrag(FCasBlockId BlockId, const UE::FDeferLock* SkipLock = nullptr);
+	void						ResetAllPendingDefrag(const UE::FDeferLock* SkipLock = nullptr);
 	TResult<FUniqueFileHandle>	OpenWrite(FCasBlockId BlockId, bool bAppend) const;
-	bool						TrackAccessIf(ECasTrackAccessType Type, FCasBlockId BlockId, int64 UtcTicks, bool bDirty);
-	bool						TrackAccessIf(ECasTrackAccessType Type, FCasBlockId BlockId, bool bDirty)
+	bool						TrackAccessIf(ECasTrackAccessType Type, FCasBlockId BlockId, int64 UtcTicks, bool bDirty, const UE::FDeferLock* SkipLock = nullptr);
+	bool						TrackAccessIf(ECasTrackAccessType Type, FCasBlockId BlockId, bool bDirty, const UE::FDeferLock* SkipLock = nullptr)
 	{
-		return TrackAccessIf(Type, BlockId, FDateTime::UtcNow().GetTicks(), bDirty);
+		return TrackAccessIf(Type, BlockId, FDateTime::UtcNow().GetTicks(), bDirty, SkipLock);
 	};
 	bool						UnlockedTrackAccessIf(ECasTrackAccessType Type, uint32 BlockIdHash, FCasBlockId BlockId, int64 UtcTicks, bool bDirty);
 	bool						TrackModificationIf(ECasTrackModificationType Type, FCasBlockId BlockId, int64 UtcTicks);
@@ -963,17 +1019,19 @@ struct FCas
 
 	const FString&				GetRootDirectory() const { return RootDirectory; }
 
+	EOnDemandInstallCasType		GetType() const { return CasType; }
+
 private:
 	FString					RootDirectory;
 public: // TODO: FIXME: encapsulate these members
 	FLookup					Lookup;
 private:
 	FBlockStates			BlockIds;
-	FReadHandles			ReadHandles;
 	FEventRef				BlockReadsDoneEvent;
 	int64					LastAccessGranularityTicks = 0;
 	uint32					MinBlockSize = 32 << 19;
 	uint32					MaxBlockSize = 32 << 20;
+	EOnDemandInstallCasType	CasType{ EOnDemandInstallCasType::None };
 	mutable UE::FMutex		Mutex;
 };
 
@@ -1367,9 +1425,9 @@ bool FCasSnapshot::TempFilesExist(const FString& SnapshotFile, const FString& Jo
 }
 
 ///////////////////////////////////////////////////////////////////////////////
-FCas::FCas()
-{
-}
+FCas::FCas(EOnDemandInstallCasType InType)
+	: CasType(InType)
+{}
 
 FResult FCas::Configure(const FOnDemandInstallCasConfig& Config)
 {
@@ -1439,12 +1497,12 @@ FResult FCas::Initialize(FStringView Directory, bool bDeleteExisting)
 	return MakeValue();
 };
 
-FCasLocation FCas::FindChunk(const FIoHash& Hash) const
+FCasLocation FCas::FindChunk(const FIoHash& Hash, const UE::FDeferLock* SkipLock) const
 {
 	const FCasAddr& Addr	= AsCasAddr(Hash);
 	const uint32 TypeHash	= GetTypeHash(Addr);
 	{
-		UE::TUniqueLock Lock(Mutex);
+		FMaybeLock Lock(this, SkipLock);
 		if (const FCasLocation* Loc = Lookup.FindByHash(TypeHash, Addr))
 		{
 			return *Loc;
@@ -1503,11 +1561,28 @@ FResult FCas::DeleteBlock(FCasBlockId BlockId, TFunctionRef<void(const FCasAddr&
 	uint32			StartTimeCycles	= FPlatformTime::Cycles();
 	const uint32	WaitTimeMs		= 1000;
 
+	FCasBlockState* BlockState = nullptr;
+
 	for (;;)
 	{
 		Lock.Lock();
 
-		const int32 RequestCount = BlockIds.FindRef(BlockId).ReadHandleCount;
+		BlockState = BlockIds.Find(BlockId);
+
+		if (CasType == EOnDemandInstallCasType::MMap)
+		{
+			if (BlockState && BlockState->HandleCount > 0)
+			{
+				return MakeCasError<void>(
+					ECasErrorCode::DeleteBlockFailed, EIoErrorCode::FileHandleOpen,
+					FString::Printf(TEXT("Cannot delete CAS block %u while it has open mapped handles"), BlockId.Id));
+			}
+
+			// Leave mutex locked until it goes out of scope
+			break;
+		}
+
+		const int32 RequestCount = BlockState ? BlockState->HandleCount : 0;
 		if (RequestCount)
 		{
 			Lock.Unlock();
@@ -1534,8 +1609,12 @@ FResult FCas::DeleteBlock(FCasBlockId BlockId, TFunctionRef<void(const FCasAddr&
 			ECasErrorCode::DeleteBlockFailed, EIoErrorCode::DeleteError, FString::Printf(TEXT("Failed to delete CAS block %u"), BlockId.Id));
 	}
 
+	if (BlockState && BlockState->PendingDefragEvent)
+	{
+		BlockState->PendingDefragEvent.GetValue()->Trigger();
+	}
+
 	BlockIds.Remove(BlockId);
-	ReadHandles.Remove(BlockId);
 	for (auto It = Lookup.CreateIterator(); It; ++It)
 	{
 		if (It->Value.BlockId == BlockId)
@@ -1568,6 +1647,13 @@ TResult<FSharedFileHandle> FCas::OpenRead(FCasBlockId BlockId)
 	const FString	Filename = GetBlockFilename(BlockId);
 	IPlatformFile&	Ipf = FPlatformFileManager::Get().GetPlatformFile();
 
+	if (CasType == EOnDemandInstallCasType::MMap)
+	{
+		return MakeCasError<FSharedFileHandle>(
+			ECasErrorCode::ReadBlockFailed, EIoErrorCode::InvalidCode,
+			FString::Printf(TEXT("MMap CAS Cannot open CAS block '%s' for read"), *Filename));
+	}
+
 	UE::TUniqueLock Lock(Mutex);
 
 	FCasBlockState* BlockState = BlockIds.Find(BlockId);
@@ -1580,7 +1666,7 @@ TResult<FSharedFileHandle> FCas::OpenRead(FCasBlockId BlockId)
 	FFileOpenResult Result = Ipf.OpenRead(*Filename, IPlatformFile::EOpenReadFlags::AllowWrite);
 	if (Result.HasValue())
 	{
-		BlockState->ReadHandleCount++;
+		BlockState->HandleCount++;
 
 		FSharedFileHandle NewHandle(
 			Result.GetValue().Release(),
@@ -1602,33 +1688,40 @@ TResult<FSharedFileHandle> FCas::OpenRead(FCasBlockId BlockId)
 
 TResult<FSharedAsyncFileHandle> FCas::OpenAsyncRead(FCasBlockId BlockId)
 {
-	UE::TUniqueLock Lock(Mutex);
+	using namespace UE::UnifiedError::IoStoreOnDemand;
 
-	if (FWeakAsyncFileHandle* MaybeHandle = ReadHandles.Find(BlockId))
+	if (CasType == EOnDemandInstallCasType::MMap)
 	{
-		if (FSharedAsyncFileHandle Handle = MaybeHandle->Pin(); Handle.IsValid())
-		{
-			return MakeValue(MoveTemp(Handle));
-		}
+		return MakeCasError<FSharedAsyncFileHandle>(
+			ECasErrorCode::ReadBlockFailed, EIoErrorCode::InvalidCode,
+			FString::Printf(TEXT("MMap CAS Cannot open CAS block '%s' for async read"), *GetBlockFilename(BlockId)));
 	}
 
-	using namespace UE::UnifiedError::IoStoreOnDemand;
-
-	const FString			Filename = GetBlockFilename(BlockId);
+	UE::TUniqueLock Lock(Mutex);
 
 	FCasBlockState* BlockState = BlockIds.Find(BlockId);
 	if (BlockState == nullptr)
 	{
 		return MakeCasError<FSharedAsyncFileHandle>(
-			ECasErrorCode::ReadBlockFailed, EIoErrorCode::NotFound, FString::Printf(TEXT("Could not find block state for '%s'"), *Filename));
+			ECasErrorCode::ReadBlockFailed, EIoErrorCode::NotFound, FString::Printf(TEXT("Could not find block state for '%s'"), *GetBlockFilename(BlockId)));
 	}
 
+	if (FWeakAsyncFileHandle* MaybeHandle = BlockState->WeakHandle.TryGet<FWeakAsyncFileHandle>())
+	{
+		if (FSharedAsyncFileHandle Handle = MaybeHandle->Pin(); Handle.IsValid())
+		{
+			return MakeValue(MoveTemp(Handle));
+		}
+	}
+
+	const FString			Filename = GetBlockFilename(BlockId);
+
 	IPlatformFile&			Ipf = FPlatformFileManager::Get().GetPlatformFile();
 	FFileOpenAsyncResult	HandleResult(Ipf.OpenAsyncRead(*Filename, IPlatformFile::EOpenReadFlags::AllowWrite));
 
 	if (HandleResult.HasValue())
 	{
-		BlockState->ReadHandleCount++;
+		BlockState->HandleCount++;
 
 		FSharedAsyncFileHandle NewHandle(
 			HandleResult.GetValue().Release(),
@@ -1638,7 +1731,7 @@ TResult<FSharedAsyncFileHandle> FCas::OpenAsyncRead(FCasBlockId BlockId)
 				OnFileHandleDeleted(BlockId);
 			}
 		);
-		ReadHandles.FindOrAdd(BlockId, NewHandle);
+		BlockState->WeakHandle.Set<FWeakAsyncFileHandle>(NewHandle);
 		
 		return MakeValue(MoveTemp(NewHandle));
 	}
@@ -1649,10 +1742,95 @@ TResult<FSharedAsyncFileHandle> FCas::OpenAsyncRead(FCasBlockId BlockId)
 		HandleResult.StealError());
 }
 
+TResult<FSharedMappedFileHandle> FCas::OpenMapped(FCasBlockId BlockId, int64 MaxSize, TOptional<FSharedEventRef>* OutBlockDeletedEvent, const UE::FDeferLock* SkipLock)
+{
+	using namespace UE::UnifiedError::IoStoreOnDemand;
+
+	if (CasType != EOnDemandInstallCasType::MMap)
+	{
+		return MakeCasError<FSharedMappedFileHandle>(
+			ECasErrorCode::OpenMappedFailed, EIoErrorCode::InvalidCode,
+			FString::Printf(TEXT("%s CAS Cannot open CAS block '%s' for mapped read"), LexToString(CasType), *GetBlockFilename(BlockId)));
+	}
+
+	FMaybeLock Lock(this, SkipLock);
+
+	FCasBlockState* BlockState = BlockIds.Find(BlockId);
+	if (BlockState == nullptr)
+	{
+		return MakeCasError<FSharedMappedFileHandle>(
+			ECasErrorCode::OpenMappedFailed, EIoErrorCode::NotFound, FString::Printf(TEXT("Could not find block state for '%s'"), *GetBlockFilename(BlockId)));
+	}
+
+	// We can't open a mapped block if it is being defragged
+	if (BlockState->PendingDefragEvent)
+	{
+		if (OutBlockDeletedEvent)
+		{
+			*OutBlockDeletedEvent = BlockState->PendingDefragEvent;
+		}
+
+		return MakeCasError<FSharedMappedFileHandle>(
+			ECasErrorCode::OpenMappedFailed, EIoErrorCode::PendingDelete,
+			FString::Printf(TEXT("Cannot delete CAS block %u while it is being defragged"), BlockId.Id));
+	}
+
+	if (FWeakMappedFileHandle* MaybeHandle = BlockState->WeakHandle.TryGet<FWeakMappedFileHandle>())
+	{
+		if (FSharedMappedFileHandle Handle = MaybeHandle->Pin(); Handle.IsValid())
+		{
+			// If the block has grown since this handle was cached, open a new one
+			// 
+			// On windows this could be more optimal by reusing the underlying file handle and 
+			// just making a new mapping, but we'd need to add platform level support for that.
+			// On Nix platforms, no mapping mapping really happens until MapRegion and MaxSize is just an
+			// an assert - so we could also re-use the file handle.
+			//
+			// We should consider adding some kind of Remap() API to reuse the same platform file handle but with
+			// a new IMappedFileHandle instead of opening a new platform handle every time.
+			if (Handle->GetFileSize() >= MaxSize)
+			{
+				return MakeValue(MoveTemp(Handle));
+			}
+		}
+	}
+
+	const FString			Filename = GetBlockFilename(BlockId);
+
+	IPlatformFile& Ipf = FPlatformFileManager::Get().GetPlatformFile();
+	FOpenMappedResult HandleResult(Ipf.OpenMappedEx2(*Filename, 
+		IPlatformFile::EOpenMappedFlags::OpenExisting | IPlatformFile::EOpenMappedFlags::AllowWrite));
+
+	if (HandleResult.HasValue())
+	{
+		BlockState->HandleCount++;
+
+		FSharedMappedFileHandle NewHandle(
+			HandleResult.GetValue().Release(),
+			[this, BlockId](IMappedFileHandle* RawHandle)
+			{
+				delete RawHandle;
+				OnFileHandleDeleted(BlockId);
+			}
+		);
+
+		check(NewHandle->GetFileSize() >= MaxSize);
+
+		BlockState->WeakHandle.Set<FWeakMappedFileHandle>(NewHandle);
+
+		return MakeValue(MoveTemp(NewHandle));
+	}
+
+	return MakeCasError<FSharedMappedFileHandle>(
+		ECasErrorCode::OpenMappedFailed, EIoErrorCode::FileOpenFailed,
+		FString::Printf(TEXT("Failed to open mapped CAS block '%s'"), *Filename),
+		HandleResult.StealError());
+}
+
 void FCas::OnFileHandleDeleted(FCasBlockId BlockId)
 {
 	UE::TUniqueLock Lock(Mutex);
-	const int32 Count = --BlockIds.FindChecked(BlockId).ReadHandleCount;
+	const int32 Count = --BlockIds.FindChecked(BlockId).HandleCount;
 	check(Count >= 0);
 	if (Count == 0)
 	{
@@ -1660,6 +1838,51 @@ void FCas::OnFileHandleDeleted(FCasBlockId BlockId)
 	}
 }
 
+bool FCas::SetPendingDefrag(FCasBlockId BlockId, const UE::FDeferLock* SkipLock)
+{
+	if (CasType != EOnDemandInstallCasType::MMap)
+	{
+		return true;
+	}
+
+	FMaybeLock Lock(this, SkipLock);
+
+	FCasBlockState* BlockState = BlockIds.Find(BlockId);
+	if (!ensureMsgf(BlockState, TEXT("Expected block to exist during defrag!")))
+	{
+		return false;
+	}
+
+	if (BlockState->HandleCount > 0)
+	{
+		return false;
+	}
+
+	BlockState->PendingDefragEvent.Emplace(EEventMode::ManualReset);
+	return true;
+}
+
+void FCas::ResetAllPendingDefrag(const UE::FDeferLock* SkipLock)
+{
+	if (CasType != EOnDemandInstallCasType::MMap)
+	{
+		return;
+	}
+
+	FMaybeLock Lock(this, SkipLock);
+
+	for (TPair<FCasBlockId, FCasBlockState>& Kv : BlockIds)
+	{
+		FCasBlockState& BlockState = Kv.Value;
+		if (BlockState.PendingDefragEvent)
+		{
+			// Make sure no thread are still waiting on this
+			BlockState.PendingDefragEvent.GetValue()->Trigger();
+			BlockState.PendingDefragEvent.Reset();
+		}
+	}
+}
+
 TResult<FUniqueFileHandle> FCas::OpenWrite(FCasBlockId BlockId, const bool bAppend) const
 {
 	using namespace UE::UnifiedError::IoStoreOnDemand;
@@ -1685,11 +1908,11 @@ TResult<FUniqueFileHandle> FCas::OpenWrite(FCasBlockId BlockId, const bool bAppe
 		OpenWriteResult.StealError());
 }
 
-bool FCas::TrackAccessIf(const ECasTrackAccessType Type, const FCasBlockId BlockId, const int64 UtcTicks, const bool bDirty)
+bool FCas::TrackAccessIf(const ECasTrackAccessType Type, const FCasBlockId BlockId, const int64 UtcTicks, const bool bDirty, const UE::FDeferLock* SkipLock)
 {
 	const uint32 BlockIdHash = GetTypeHash(BlockId);
 
-	UE::TUniqueLock Lock(Mutex);
+	FMaybeLock Lock(this, SkipLock);
 
 	return UnlockedTrackAccessIf(Type, BlockIdHash, BlockId, UtcTicks, bDirty);
 }
@@ -1920,7 +2143,6 @@ void FCas::Compact()
 	UE::TUniqueLock Lock(Mutex);
 	Lookup.Compact();
 	BlockIds.Compact();
-	ReadHandles.Compact();
 }
 
 FResult FCas::Verify(TArray<FCasAddr>& OutRemovedChunksAddrs, TMap<FCasBlockId, uint64>& OutOverBudgetBlocks)
@@ -2313,7 +2535,8 @@ class FOnDemandInstallCache final
 
 	struct FCacheState
 	{
-		EOnDemandInstallCasType		Type{ EOnDemandInstallCasType::None };
+		FCacheState(EOnDemandInstallCasType InType) : Cas(InType) {}
+
 		FCas						Cas;
 		FPendingChunks				PendingChunks;
 		std::atomic<FCasBlockId>	CurrentBlock{ FCasBlockId::Invalid };
@@ -2422,8 +2645,7 @@ FOnDemandInstallCache::FOnDemandInstallCache(const FOnDemandInstallCacheConfig&
 			continue;
 		}
 
-		FCacheState& CacheState = CacheStates.Emplace_GetRef();
-		CacheState.Type = CasConfig.Type;
+		FCacheState& CacheState = CacheStates.Emplace_GetRef(CasConfig.Type);
 
 		CacheState.MaxCacheSize = CasConfig.DiskQuota;
 
@@ -2512,7 +2734,7 @@ FOnDemandInstallCache::FOnDemandInstallCache(const FOnDemandInstallCacheConfig&
 
 			for (FCacheState& CacheState : CacheStates)
 			{
-				FCasSnapshot::FCasState* SnapshotCasState = Algo::FindBy(Snapshot.CasState, CacheState.Type, &FCasSnapshot::FCasState::CasType);
+				FCasSnapshot::FCasState* SnapshotCasState = Algo::FindBy(Snapshot.CasState, CacheState.Cas.GetType(), &FCasSnapshot::FCasState::CasType);
 				if (SnapshotCasState)
 				{
 					UE_LOG(LogIoStoreOnDemand, Log, TEXT("Loaded CAS state from snapshot with %d blocks and %d chunk locations"),
@@ -2538,7 +2760,7 @@ FOnDemandInstallCache::FOnDemandInstallCache(const FOnDemandInstallCacheConfig&
 		InitResult = FCasJournal::Replay(JournalFilename, [this](const FCasJournal::FEntry& JournalEntry)
 		{
 			FResult Result = MakeValue();
-			FCacheState* CacheState = Algo::FindBy(CacheStates, JournalEntry.CasType(), &FCacheState::Type);
+			FCacheState* CacheState = Algo::FindBy(CacheStates, JournalEntry.CasType(), Projection(&FCacheState::Cas, &FCas::GetType));
 			if (CacheState == nullptr)
 			{
 				// This cache state is not currently configured and there's no way to partially cleanup the cache right now.
@@ -2651,14 +2873,14 @@ void FOnDemandInstallCache::Shutdown()
 		if (FResult Result = FlushPendingChunksImpl(CacheState, CacheState.PendingChunks); Result.HasError())
 		{
 			UE_LOG(LogIoStoreOnDemand, Error, TEXT("Failed to flush pending chunks on shutdown for %s, reason '%s'"), 
-				LexToString(CacheState.Type), *Result.GetError().GetErrorMessage(true).ToString());
+				LexToString(CacheState.Cas.GetType()), *Result.GetError().GetErrorMessage(true).ToString());
 		}
 	}
 
 	const FString JournalFile = GetJournalFilename();
 	for (FCacheState& CacheState : CacheStates)
 	{
-		FCasJournal::FTransaction Transaction = FCasJournal::Begin(CacheState.Type, JournalFile);
+		FCasJournal::FTransaction Transaction = FCasJournal::Begin(CacheState.Cas.GetType(), JournalFile);
 
 		FCas::FLastAccess LastAccess = CacheState.Cas.GetAndClearDirtyLastAccess();
 		JournalLastAccess(Transaction, LastAccess);
@@ -2729,7 +2951,8 @@ bool FOnDemandInstallCache::DoesChunkExist(const FIoChunkId& ChunkId) const
 	{
 		for (const FCacheState& CacheState : CacheStates)
 		{
-
+			const bool bAllowReadFromPendingChunks = CacheState.Cas.GetType() != EOnDemandInstallCasType::MMap;
+			if (bAllowReadFromPendingChunks)
 			{
 				TSharedLock SharedLock(CacheState.PendingChunks.SharedMutex);
 
@@ -2740,7 +2963,10 @@ bool FOnDemandInstallCache::DoesChunkExist(const FIoChunkId& ChunkId) const
 			}
 
 			const FCasLocation CasLoc = CacheState.Cas.FindChunk(ChunkInfo.Hash());
-			return CasLoc.IsValid();
+			if (CasLoc.IsValid())
+			{
+				return true;
+			}
 		}
 	}
 
@@ -2760,7 +2986,163 @@ TIoStatusOr<uint64> FOnDemandInstallCache::GetSizeForChunk(const FIoChunkId& Chu
 
 TIoStatusOr<FIoMappedRegion> FOnDemandInstallCache::OpenMapped(const FIoChunkId& ChunkId, const FIoReadOptions& Options)
 {
-	return FIoStatus(EIoErrorCode::FileOpenFailed, TEXT("On-demand install cache does not support memory mapped files."));
+	if (!FPlatformProperties::SupportsMemoryMappedFiles())
+	{
+		return FIoStatus(EIoErrorCode::Unknown, TEXT("Platform does not support memory mapped files"));
+	}
+
+#if UE_ONDEMANDINSTALLCACHE_EXCLUSIVE_WRITE
+	return FIoStatus(EIoErrorCode::Unknown, TEXT("OpenMapped is not supported with EXCLUSIVE_WRITE"));
+#endif
+
+	if (Options.GetTargetVa() != nullptr)
+	{
+		return FIoStatus(EIoErrorCode::InvalidParameter, TEXT("Invalid read options"));
+	}
+
+	EIoErrorCode ErrorCode = EIoErrorCode::UnknownChunkID;
+	FOnDemandChunkInfo ChunkInfo = IoStore.GetInstalledChunkInfo(ChunkId, ErrorCode);
+	if (ChunkInfo.IsValid() == false)
+	{
+		if (ErrorCode == EIoErrorCode::NotInstalled)
+		{
+#if !UE_BUILD_SHIPPING
+			if (IOnDemandDevelopmentExtension* Ext = IoStore.TryGetDevelopmentExtension())
+			{
+				TUtf8StringBuilder<128> Filename;
+				TStringBuilder<128> Sb;
+				FStringView PackageName = TryConvertChunkIdToPackageName(ChunkId, Filename, Sb);
+
+				if (Ext->GetUnreferencedIoChunkStatus(
+					ChunkId,
+					Filename.ToView(),
+					PackageName,
+					IsDevModeEnabled()) != EIoErrorCode::NotInstalled)
+				{
+					return FIoStatus(EIoErrorCode::NotFound);
+				}
+			}
+#endif
+			return FIoStatus(EIoErrorCode::NotInstalled);
+		}
+		return FIoStatus(EIoErrorCode::NotFound);
+	}
+
+	const uint64 RequestSize = FMath::Min<uint64>(
+		Options.GetSize(),
+		ChunkInfo.RawSize() - Options.GetOffset());
+
+	TIoStatusOr<FIoOffsetAndLength> MaybeChunkRange = FIoChunkEncoding::GetChunkRange(
+		ChunkInfo.RawSize(),
+		ChunkInfo.BlockSize(),
+		ChunkInfo.Blocks(),
+		Options.GetOffset(),
+		RequestSize);
+
+	if (MaybeChunkRange.IsOk() == false)
+	{
+		UE_LOG(LogIoStoreOnDemand, Error, TEXT("Failed to get chunk range"));
+		return MaybeChunkRange.Status();
+	}
+
+	const FIoOffsetAndLength ChunkRange = MaybeChunkRange.ConsumeValueOrDie();
+
+	// OpenMapped is not async and may be called from any thread so it needs to hold the Cas lock from this point forward to avoid
+	// racing with a possible purge or defrag being called from an installer task. Normally, DeleteBlock waits for pending reads to
+	// flush but a mapped file handle may be held indefinitely so any block with open mmap handles cannot be deleted.
+	TDynamicUniqueLock<FCas> CasUniqueLock;
+
+	FCasLocation CasLoc;
+	FCacheState *pCacheState = nullptr;
+	for (FCacheState& CacheState : CacheStates)
+	{
+		TDynamicUniqueLock LocalLock(CacheState.Cas);
+		CasLoc = CacheState.Cas.FindChunk(ChunkInfo.Hash(), &UE::DeferLock);
+		if (CasLoc.IsValid())
+		{
+			pCacheState = &CacheState;
+			CasUniqueLock = MoveTemp(LocalLock);
+			break;
+		}
+	}
+
+	TResult<FSharedMappedFileHandle> FileOpenResult = MakeError(UE::UnifiedError::IoStoreOnDemand::CasError::MakeError());
+
+	for(;;)
+	{
+		if (CasLoc.IsValid() == false)
+		{
+			// Chunk is referenced but not in any cache. This should probably never happen currently.
+			// If pinning of unreferenced chunks was allowed, then this would be expected...
+			return FIoStatus(EIoErrorCode::NotInstalled, TEXT("IoChunk should be installed but was not found!"));
+		}
+
+		// Cas is now locked by CasUniqueLock so skip locking on remaining API calls
+
+		check(pCacheState);
+		FCacheState& CacheState = *pCacheState;
+
+		constexpr const bool bDirty = true;
+		if (CacheState.Cas.TrackAccessIf(ECasTrackAccessType::Granular, CasLoc.BlockId, bDirty, &UE::DeferLock))
+		{
+			IoStore.FlushLastAccess({}, nullptr);
+		}
+
+		const int64 MaxSize = CasLoc.BlockOffset + ChunkRange.GetOffset() + ChunkRange.GetLength();
+
+		TOptional<FSharedEventRef> BlockDeletedEvent;
+		FileOpenResult = CacheState.Cas.OpenMapped(CasLoc.BlockId, MaxSize, &BlockDeletedEvent, &UE::DeferLock);
+		if (FileOpenResult.HasError())
+		{
+			using namespace UE::UnifiedError::IoStoreOnDemand;
+
+			if (BlockDeletedEvent)
+			{
+				CasUniqueLock.Unlock();
+
+				// Wait for delete event from defrag
+				BlockDeletedEvent.GetValue()->Wait();
+
+				// Retry
+				CasUniqueLock.Lock();
+				CasLoc = CacheState.Cas.FindChunk(ChunkInfo.Hash(), &UE::DeferLock);
+				continue;
+			}
+
+			const FString Filename = CacheState.Cas.GetBlockFilename(CasLoc.BlockId);
+			UE_LOG(LogIoStoreOnDemand, Error, TEXT("Failed to open CAS block '%s' in CAS %s for mapped reading, reason '%s'"),
+				*Filename, LexToString(CacheState.Cas.GetType()), *FileOpenResult.GetError().GetErrorMessage(true).ToString());
+
+			if (const FCasErrorContext* Ctx = FileOpenResult.GetError().GetErrorContext<FCasErrorContext>())
+			{
+				return FIoStatus(Ctx->IoErrorCode, Ctx->SystemErrorCode, *LexToString(FileOpenResult.GetError()));
+			}
+
+			return FIoStatus(EIoErrorCode::FileOpenFailed, *LexToString(FileOpenResult.GetError()));
+		}
+
+		break;
+	}
+
+	// The file is open the and the Cas block has a non-0 handle count, so the lock can released
+	CasUniqueLock.Unlock();
+
+	const FCacheState& CacheState = *pCacheState;
+
+	FSharedMappedFileHandle MappedHandle = FileOpenResult.StealValue();
+	IMappedFileRegion* MappedRegion = MappedHandle->MapRegion(CasLoc.BlockOffset + ChunkRange.GetOffset(), ChunkRange.GetLength());
+	if (MappedRegion == nullptr)
+	{
+		const FString Filename = CacheState.Cas.GetBlockFilename(CasLoc.BlockId);
+		UE_LOG(LogIoStoreOnDemand, Error, TEXT("Failed to map region for chunk in %s"), *Filename);
+		return FIoStatus(EIoErrorCode::ReadError, TEXT("Failed to map region"));
+	}
+
+	return FIoMappedRegion
+	{
+		.MappedFileHandle = new FSharedMappedFileHandleProxy(MappedHandle.ToSharedRef()),
+		.MappedFileRegion = MappedRegion
+	};
 }
 
 const TCHAR* FOnDemandInstallCache::GetName() const
@@ -2817,6 +3199,8 @@ bool FOnDemandInstallCache::Resolve(FIoRequestImpl* Request)
 		return true;
 	}
 
+	// TODO: Should this support GetTargetVa() ?
+
 	// The internal request parameters are attached/owned by the I/O request via
 	// the backend data parameter. The chunk request is deleted in GetCompletedRequests
 	FChunkRequest::Attach(*Request, new FChunkRequest(
@@ -2838,7 +3222,7 @@ bool FOnDemandInstallCache::Resolve(FIoRequestImpl* Request)
 	if (!bFoundInCas)
 	{
 		// Chunk is referenced but not in any cache. This should probably never happen currently.
-		// If pinning of unrefernced chunks was allowed, then this would be expected...
+		// If pinning of unreferenced chunks was allowed, then this would be expected...
 		TUniquePtr<FChunkRequest> Detached = FChunkRequest::Detach(*Request);
 		CompleteRequest(Request, EIoErrorCode::NotInstalled);
 	}
@@ -2850,6 +3234,8 @@ bool FOnDemandInstallCache::ResolveFromCas(FCacheState& CacheState, FIoRequestIm
 {
 	FChunkRequest& ChunkRequest = FChunkRequest::GetRef(*Request);
 
+	const bool bAllowResolve = CacheState.Cas.GetType() != EOnDemandInstallCasType::MMap;
+
 	// Must check pending chunks before CAS because otherwise
 	// a chunk could be flushed after checking the CAS and before checking pending.
 	// Otherise we'd need take both locks.
@@ -2859,6 +3245,14 @@ bool FOnDemandInstallCache::ResolveFromCas(FCacheState& CacheState, FIoRequestIm
 		const FIoBuffer* PendingChunk = CacheState.PendingChunks.FindChunk(ChunkRequest.ChunkInfo.Hash());
 		if (PendingChunk != nullptr)
 		{
+			if (bAllowResolve == false)
+			{
+				// Can't use the Resolve API to read this data
+				TUniquePtr<FChunkRequest> Detached = FChunkRequest::Detach(*Request);
+				CompleteRequest(Request, EIoErrorCode::FileOpenFailed);
+				return true;
+			}
+
 			FMemoryView BytesToRead = PendingChunk->GetView().Mid(ChunkRequest.ChunkRange.GetOffset(), ChunkRequest.EncodedChunk.GetSize());
 			check(BytesToRead.GetSize() == ChunkRequest.EncodedChunk.GetSize());
 			FMemory::Memcpy(ChunkRequest.EncodedChunk.GetData(), BytesToRead.GetData(), ChunkRequest.EncodedChunk.GetSize());
@@ -2880,6 +3274,14 @@ bool FOnDemandInstallCache::ResolveFromCas(FCacheState& CacheState, FIoRequestIm
 		return false;
 	}
 
+	if (bAllowResolve == false)
+	{
+		// Can't use the Resolve API to read this data
+		TUniquePtr<FChunkRequest> Detached = FChunkRequest::Detach(*Request);
+		CompleteRequest(Request, EIoErrorCode::FileOpenFailed);
+		return true;
+	}
+
 	TRACE_IOSTORE_BACKEND_REQUEST_STARTED(Request, this);
 	constexpr const bool bDirty = true;
 	if (CacheState.Cas.TrackAccessIf(ECasTrackAccessType::Granular, CasLoc.BlockId, bDirty))
@@ -2910,7 +3312,7 @@ bool FOnDemandInstallCache::ResolveFromCas(FCacheState& CacheState, FIoRequestIm
 				if (Request->IsCancelled())
 				{
 					UE_LOG(LogIoStoreOnDemand, Verbose, TEXT("Cancelled request - skipped seek to offset %" INT64_FMT " in CAS block '%s' in CAS %s"), 
-						CasBlockOffset, *Filename, LexToString(CacheState.Type));
+						CasBlockOffset, *Filename, LexToString(CacheState.Cas.GetType()));
 				}
 				else if (FileHandle->Seek(CasBlockOffset))
 				{
@@ -2922,19 +3324,19 @@ bool FOnDemandInstallCache::ResolveFromCas(FCacheState& CacheState, FIoRequestIm
 					else
 					{
 						UE_LOG(LogIoStoreOnDemand, Error, TEXT("Failed to read %" UINT64_FMT " bytes at offset %" INT64_FMT " in CAS block '%s' in CAS %s"),
-							ChunkRequest.EncodedChunk.GetSize(), CasBlockOffset, *Filename, LexToString(CacheState.Type));
+							ChunkRequest.EncodedChunk.GetSize(), CasBlockOffset, *Filename, LexToString(CacheState.Cas.GetType()));
 					}
 				}
 				else
 				{
 					UE_LOG(LogIoStoreOnDemand, Error, TEXT("Failed to seek to offset %" INT64_FMT " in CAS block '%s' in CAS %s"), 
-						CasBlockOffset, *Filename, LexToString(CacheState.Type));
+						CasBlockOffset, *Filename, LexToString(CacheState.Cas.GetType()));
 				}
 			}
 			else
 			{
 				UE_LOGFMT(LogIoStoreOnDemand, Error, "Failed to open CAS block '{Filename}' in CAS {CasType} for reading ({Error})", 
-					*Filename, LexToString(CacheState.Type), FileOpenResult.GetError().GetErrorMessage(true).ToString());
+					*Filename, LexToString(CacheState.Cas.GetType()), FileOpenResult.GetError().GetErrorMessage(true).ToString());
 			}
 
 			UE::Tasks::Launch(UE_SOURCE_LOCATION, [this, Request, Status]
@@ -2953,7 +3355,7 @@ bool FOnDemandInstallCache::ResolveFromCas(FCacheState& CacheState, FIoRequestIm
 	{
 		const FString Filename = CacheState.Cas.GetBlockFilename(CasLoc.BlockId);
 		UE_LOG(LogIoStoreOnDemand, Error, TEXT("Failed to open CAS block '%s' in CAS %s for async reading, reason '%s'"),
-			*Filename, LexToString(CacheState.Type), *FileOpenResult.GetError().GetErrorMessage(true).ToString());
+			*Filename, LexToString(CacheState.Cas.GetType()), *FileOpenResult.GetError().GetErrorMessage(true).ToString());
 		TUniquePtr<FChunkRequest> Detached = FChunkRequest::Detach(*Request);
 		CompleteRequest(Request, EIoErrorCode::FileOpenFailed);
 		return true;
@@ -2992,6 +3394,9 @@ bool FOnDemandInstallCache::IsChunkCached(const FIoHash& ChunkHash)
 {
 	for (FCacheState& CacheState : CacheStates)
 	{
+		// Do not check for EOnDemandInstallCasFlags::AllowReadFromPendingChunks here because if the chunk
+		// is pending we don't want to make another request for it regardless of the CAS its in.
+
 		{
 			TSharedLock SharedLock(CacheState.PendingChunks.SharedMutex);
 
@@ -3080,7 +3485,7 @@ FResult FOnDemandInstallCache::Purge(EOnDemandInstallCasType CasType, const uint
 		if (Info.FileSize > CacheState->Cas.GetMaxBlockSize())
 		{
 			UE_LOGFMT(LogIoStoreOnDemand, Error, "CAS {CasType} Block {BlockId} has total size {BlockSize} which is greater than max block size {MaxBlockSize}", 
-				LexToString(CacheState->Type), Kv.Key.Id, Info.FileSize, CacheState->Cas.GetMaxBlockSize());
+				LexToString(CacheState->Cas.GetType()), Kv.Key.Id, Info.FileSize, CacheState->Cas.GetMaxBlockSize());
 			ensure(false);
 		}
 
@@ -3091,7 +3496,7 @@ FResult FOnDemandInstallCache::Purge(EOnDemandInstallCasType CasType, const uint
 		else if (Info.RefSize > Info.FileSize)
 		{
 			UE_LOGFMT(LogIoStoreOnDemand, Error, "CAS {CasType} Block {BlockId} has RefSize {RefSize} which is greater than total size {BlockSize}",
-				LexToString(CacheState->Type), Kv.Key.Id, Info.RefSize, Info.FileSize);
+				LexToString(CacheState->Cas.GetType()), Kv.Key.Id, Info.RefSize, Info.FileSize);
 			ensure(false);
 		}
 
@@ -3115,12 +3520,12 @@ FResult FOnDemandInstallCache::Purge(EOnDemandInstallCasType CasType, const uint
 	if (TotalRequiredBytes <= CacheState->MaxCacheSize)
 	{
 		UE_LOG(LogIoStoreOnDemand, Verbose, TEXT("Skipping cache purge for %s, MaxCacheSize=%.2lf MiB, CacheSize=%.2lf MiB, ReferencedBlockSize=%.2lf MiB, ReferencedSize=%.2lf MiB, FragmentedBytes=%.2lf MiB, UncachedSize=%.2lf MiB"),
-			LexToString(CacheState->Type), ToMiB(CacheState->MaxCacheSize), ToMiB(TotalCachedBytes), ToMiB(TotalReferencedBlockBytes), ToMiB(ReferencedBytes), ToMiB(FragmentedBytes), ToMiB(TotalUncachedBytes));
+			LexToString(CacheState->Cas.GetType()), ToMiB(CacheState->MaxCacheSize), ToMiB(TotalCachedBytes), ToMiB(TotalReferencedBlockBytes), ToMiB(ReferencedBytes), ToMiB(FragmentedBytes), ToMiB(TotalUncachedBytes));
 		return MakeValue();
 	}
 
 	UE_LOG(LogIoStoreOnDemand, Log, TEXT("Purging install cache %s, MaxCacheSize=%.2lf MiB, CacheSize=%.2lf MiB, ReferencedBlockSize=%.2lf MiB, ReferencedSize=%.2lf MiB, FragmentedBytes=%.2lf MiB, UncachedSize=%.2lf MiB"),
-		LexToString(CacheState->Type), ToMiB(CacheState->MaxCacheSize), ToMiB(TotalCachedBytes), ToMiB(TotalReferencedBlockBytes), ToMiB(ReferencedBytes), ToMiB(FragmentedBytes), ToMiB(TotalUncachedBytes));
+		LexToString(CacheState->Cas.GetType()), ToMiB(CacheState->MaxCacheSize), ToMiB(TotalCachedBytes), ToMiB(TotalReferencedBlockBytes), ToMiB(ReferencedBytes), ToMiB(FragmentedBytes), ToMiB(TotalUncachedBytes));
 
 	const uint64	TotalBytesToPurge	= TotalRequiredBytes - CacheState->MaxCacheSize;
 	uint64			TotalPurgedBytes	= 0;
@@ -3129,13 +3534,13 @@ FResult FOnDemandInstallCache::Purge(EOnDemandInstallCasType CasType, const uint
 	if (TotalPurgedBytes > 0)
 	{
 		UE_LOG(LogIoStoreOnDemand, Log, TEXT("Purged %.2lf MiB (%.2lf%%) from install cache %s"),
-			ToMiB(TotalPurgedBytes), 100.0 * (double(TotalPurgedBytes) / double(TotalCachedBytes)), LexToString(CacheState->Type));
+			ToMiB(TotalPurgedBytes), 100.0 * (double(TotalPurgedBytes) / double(TotalCachedBytes)), LexToString(CacheState->Cas.GetType()));
 	}
 
 	const uint64 NewCachedBytes = TotalCachedBytes - TotalPurgedBytes;
 	UE_CLOG(NewCachedBytes > CacheState->MaxCacheSize,
 		LogIoStoreOnDemand, Warning, TEXT("Max install cache size exceeded by %.2lf MiB (%.2lf%%) for cache %s"),
-			ToMiB(NewCachedBytes - CacheState->MaxCacheSize), 100.0 * (double(NewCachedBytes - CacheState->MaxCacheSize) / double(CacheState->MaxCacheSize)), LexToString(CacheState->Type));
+			ToMiB(NewCachedBytes - CacheState->MaxCacheSize), 100.0 * (double(NewCachedBytes - CacheState->MaxCacheSize) / double(CacheState->MaxCacheSize)), LexToString(CacheState->Cas.GetType()));
 
 	uint64 DefragPurgedBytes = 0;
 	if (Result.HasError() == false && TotalPurgedBytes < TotalBytesToPurge)
@@ -3188,7 +3593,7 @@ FResult FOnDemandInstallCache::PurgeAllUnreferenced(EOnDemandInstallCasType CasT
 		uint64(0));
 
 	UE_LOG(LogIoStoreOnDemand, Log, TEXT("Purging install cache %s, MaxCacheSize=%.2lf MiB, CacheSize=%.2lf MiB, ReferencedBytes=%.2lf MiB"),
-		LexToString(CacheState->Type), ToMiB(CacheState->MaxCacheSize), ToMiB(TotalCachedBytes), ToMiB(TotalReferencedBytes));
+		LexToString(CacheState->Cas.GetType()), ToMiB(CacheState->MaxCacheSize), ToMiB(TotalCachedBytes), ToMiB(TotalReferencedBytes));
 
 	// TODO: should probably default to TotalCachedBytes and not MaxCacheSize
 	const uint64	TotalBytesToPurge	= BytesToPurge ? *BytesToPurge : CacheState->MaxCacheSize;
@@ -3198,13 +3603,13 @@ FResult FOnDemandInstallCache::PurgeAllUnreferenced(EOnDemandInstallCasType CasT
 	if (TotalPurgedBytes > 0)
 	{
 		UE_LOG(LogIoStoreOnDemand, Log, TEXT("Purged %.2lf MiB (%.2lf%%) from install cache %s"),
-			ToMiB(TotalPurgedBytes), 100.0 * (double(TotalPurgedBytes) / double(TotalCachedBytes)), LexToString(CacheState->Type));
+			ToMiB(TotalPurgedBytes), 100.0 * (double(TotalPurgedBytes) / double(TotalCachedBytes)), LexToString(CacheState->Cas.GetType()));
 	}
 
 	const uint64 NewCachedBytes = TotalCachedBytes - TotalPurgedBytes;
 	UE_CLOG(NewCachedBytes > CacheState->MaxCacheSize,
 		LogIoStoreOnDemand, Warning, TEXT("Max install cache size exceeded by %.2lf MiB (%.2lf%%) for cache %s"),
-			ToMiB(NewCachedBytes - CacheState->MaxCacheSize), 100.0 * (double(NewCachedBytes - CacheState->MaxCacheSize) / double(CacheState->MaxCacheSize)), LexToString(CacheState->Type));
+			ToMiB(NewCachedBytes - CacheState->MaxCacheSize), 100.0 * (double(NewCachedBytes - CacheState->MaxCacheSize) / double(CacheState->MaxCacheSize)), LexToString(CacheState->Cas.GetType()));
 
 	if (BytesToPurge)
 	{
@@ -3259,7 +3664,7 @@ FResult FOnDemandInstallCache::DefragAll(EOnDemandInstallCasType CasType, const
 		uint64(0));
 
 	UE_LOG(LogIoStoreOnDemand, Log, TEXT("Defragmenting install cache %s, MaxCacheSize=%.2lf MiB, CacheSize=%.2lf MiB, ReferencedBlockSize=%.2lf MiB, ReferencedSize=%.2lf MiB"),
-		LexToString(CacheState->Type), ToMiB(CacheState->MaxCacheSize), ToMiB(TotalCachedBytes), ToMiB(TotalReferencedBlockBytes), ToMiB(ReferencedBytes));
+		LexToString(CacheState->Cas.GetType()), ToMiB(CacheState->MaxCacheSize), ToMiB(TotalCachedBytes), ToMiB(TotalReferencedBlockBytes), ToMiB(ReferencedBytes));
 
 	return Defrag(*CacheState, Containers, ChunkEntryIndices, BlockInfo, BytesToFree);
 }
@@ -3348,13 +3753,13 @@ FResult FOnDemandInstallCache::Verify()
 
 		if (ChunkCount == 0)
 		{
-			UE_LOG(LogIoStoreOnDemand, Log, TEXT("Verify skipped CAS %s, install cache is empty"), LexToString(CacheState.Type));
+			UE_LOG(LogIoStoreOnDemand, Log, TEXT("Verify skipped CAS %s, install cache is empty"), LexToString(CacheState.Cas.GetType()));
 			continue;
 		}
 
 		FIoBuffer	Chunk(1 << 20);
 
-		UE_LOG(LogIoStoreOnDemand, Log, TEXT("Verifying %d installed chunks in CAS %s..."), ChunkCount, LexToString(CacheState.Type));
+		UE_LOG(LogIoStoreOnDemand, Log, TEXT("Verifying %d installed chunks in CAS %s..."), ChunkCount, LexToString(CacheState.Cas.GetType()));
 		for (int32 ChunkIndex = 0; const FCasAddrLocation& ChunkLocation : ChunkLocations)
 		{
 			TResult<FSharedFileHandle> OpenResult = CacheState.Cas.OpenRead(ChunkLocation.Location.BlockId);
@@ -3554,7 +3959,7 @@ FResult FOnDemandInstallCache::InitialVerify()
 			if (Verify.HasError())
 			{
 				// Remove all entries that doesn't have a valid cache block 
-				FCasJournal::FTransaction Transaction = FCasJournal::Begin(CacheState.Type, GetJournalFilename());
+				FCasJournal::FTransaction Transaction = FCasJournal::Begin(CacheState.Cas.GetType(), GetJournalFilename());
 				for (const FCasAddr& Addr : RemovedChunks)
 				{
 					Transaction.ChunkLocation(FCasLocation::Invalid, Addr);
@@ -3578,12 +3983,12 @@ FResult FOnDemandInstallCache::InitialVerify()
 				uint64			TotalPurgedBytes = 0;
 
 				UE_LOG(LogIoStoreOnDemand, Warning, TEXT("Cache size is greater than disk quota for CAS %s - Purging install cache, MaxCacheSize=%.2lf MiB, TotalSize=%.2lf MiB, TotalBytesToPurge=%.2lf MiB"),
-					LexToString(CacheState.Type), ToMiB(CacheState.MaxCacheSize), ToMiB(CacheSize), ToMiB(TotalBytesToPurge));
+					LexToString(CacheState.Cas.GetType()), ToMiB(CacheState.MaxCacheSize), ToMiB(CacheSize), ToMiB(TotalBytesToPurge));
 
 				if (FResult Result = PurgeInternal(CacheState, BlockInfo, TotalBytesToPurge, TotalPurgedBytes); Result.HasError())
 				{
 					const FString ErrorMessage = FString::Printf(TEXT("Failed to purge %.2lf MiB from CAS %s reason '%s'"),
-						ToMiB(TotalBytesToPurge), LexToString(CacheState.Type), *LexToString(Result.GetError()));
+						ToMiB(TotalBytesToPurge), LexToString(CacheState.Cas.GetType()), *LexToString(Result.GetError()));
 
 					UE_LOG(LogIoStoreOnDemand, Error, TEXT("%s"), *ErrorMessage);
 					return Result;
@@ -3593,13 +3998,13 @@ FResult FOnDemandInstallCache::InitialVerify()
 				{
 					// This should never happen since we don't have any referenced cache blocks at startup
 					const FString ErrorMessage = FString::Printf(TEXT("Failed to purge %.2lf MiB from CAS %s. Actually purged %.2lf MiB"),
-						ToMiB(TotalBytesToPurge), LexToString(CacheState.Type), ToMiB(TotalPurgedBytes));
+						ToMiB(TotalBytesToPurge), LexToString(CacheState.Cas.GetType()), ToMiB(TotalPurgedBytes));
 
 					UE_LOG(LogIoStoreOnDemand, Error, TEXT("%s"), *ErrorMessage);
 					return MakeError(UE::UnifiedError::IoStoreOnDemand::InstallCachePurgeError::MakeError());
 				}
 
-				UE_LOG(LogIoStoreOnDemand, Log, TEXT("Successfully purged %.2lf MiB from CAS %s"), ToMiB(TotalPurgedBytes), LexToString(CacheState.Type));
+				UE_LOG(LogIoStoreOnDemand, Log, TEXT("Successfully purged %.2lf MiB from CAS %s"), ToMiB(TotalPurgedBytes), LexToString(CacheState.Cas.GetType()));
 			}
 		}
 
@@ -3608,7 +4013,7 @@ FResult FOnDemandInstallCache::InitialVerify()
 			const FCasBlockId OldCurrentBlock = CacheState.CurrentBlock.exchange(FCasBlockId::Invalid);
 			if (OldCurrentBlock != FCasBlockId::Invalid)
 			{
-				UE_LOG(LogIoStoreOnDemand, Warning, TEXT("Current Block %u was not found for CAS %s, resetting"), OldCurrentBlock.Id, LexToString(CacheState.Type));
+				UE_LOG(LogIoStoreOnDemand, Warning, TEXT("Current Block %u was not found for CAS %s, resetting"), OldCurrentBlock.Id, LexToString(CacheState.Cas.GetType()));
 			}
 		}
 	}
@@ -3667,28 +4072,28 @@ void FOnDemandInstallCache::AddReferencesToBlocks(
 				if (!BlockInfo)
 				{
 					UE_LOG(LogIoStoreOnDemand, Error, TEXT("Failed to find CAS block info for referenced chunk, CAS=%s, ChunkId='%s', Container='%s'"),
-						LexToString(CacheState.Type), *LexToString(Container->ChunkIds[EntryIndex]), *Container->Name());
+						LexToString(CacheState.Cas.GetType()), *LexToString(Container->ChunkIds[EntryIndex]), *Container->Name());
 				}
 				else
 				{
 					BlockInfo->RefSize += ChunkDiskSize;
 				}
 			}
-			else
-			{
-				// Check pending list
-				bool bChunkIsPending = false;
-				{
-					TSharedLock SharedLock(CacheState.PendingChunks.SharedMutex);
-					bChunkIsPending = CacheState.PendingChunks.ContainsChunk(Entry.Hash);
-				}
-
-				// TODO: IadMultiCAS: this probably is no longer valid with multiple CAS.  Its likely that the referenced chunk is just in another CAS
-				UE_CLOG(!bChunkIsPending, LogIoStoreOnDemand, Error, TEXT("Failed to find CAS location or pending chunk for chunk reference, CAS=%s, ChunkId='%s', Container='%s'"),
-					LexToString(CacheState.Type), *LexToString(Container->ChunkIds[EntryIndex]), *Container->Name());
-
-				ensure(bChunkIsPending);
-			}
+			// This chek is no longer valid with multiple CAS because referenced chunk is likely just in a different CAS
+			//else
+			//{
+			//	// Check pending list
+			//	bool bChunkIsPending = false;
+			//	{
+			//		TSharedLock SharedLock(CacheState.PendingChunks.SharedMutex);
+			//		bChunkIsPending = CacheState.PendingChunks.ContainsChunk(Entry.Hash);
+			//	}
+
+			//	UE_CLOG(!bChunkIsPending, LogIoStoreOnDemand, Error, TEXT("Failed to find CAS location or pending chunk for chunk reference, CAS=%s, ChunkId='%s', Container='%s'"),
+			//		LexToString(CacheState.Cas.GetType()), *LexToString(Container->ChunkIds[EntryIndex]), *Container->Name());
+
+			//	ensure(bChunkIsPending);
+			//}
 
 			EntryIndex++;
 		}
@@ -3719,7 +4124,7 @@ FResult FOnDemandInstallCache::PurgeInternal(
 			continue;
 		}
 
-		FCasJournal::FTransaction	Transaction = FCasJournal::Begin(CacheState.Type, GetJournalFilename());
+		FCasJournal::FTransaction	Transaction = FCasJournal::Begin(CacheState.Cas.GetType(), GetJournalFilename());
 
 		auto HandleChunkRemoved = [&Transaction](const FCasAddr& Addr)
 		{
@@ -3728,6 +4133,18 @@ FResult FOnDemandInstallCache::PurgeInternal(
 
 		if (FResult Result = CacheState.Cas.DeleteBlock(BlockId, HandleChunkRemoved); Result.HasError())
 		{
+			// Skip any blocks with open MMap handles. They can't be deleted currently.
+			if (CacheState.Cas.GetType() == EOnDemandInstallCasType::MMap)
+			{
+				using namespace UE::UnifiedError::IoStoreOnDemand;
+
+				if (const FCasErrorContext* Ctx = Result.GetError().GetErrorContext<FCasErrorContext>();
+					Ctx && Ctx->ErrorCode == ECasErrorCode::DeleteBlockFailed && Ctx->IoErrorCode == EIoErrorCode::FileHandleOpen)
+				{
+					continue;
+				}
+			}
+
 			return Result;
 		}
 
@@ -3813,6 +4230,22 @@ FResult FOnDemandInstallCache::Defrag(
 		return LHS.RefSize < RHS.RefSize;
 	});
 
+	// If an error occurs, cleanup any blocks marked as pending defrag
+	ON_SCOPE_EXIT
+	{
+		CacheState.Cas.ResetAllPendingDefrag();
+	};
+
+	TDynamicUniqueLock<FCas> CasUniqueLock;
+	const UE::FDeferLock* SkipLock = nullptr;
+	if (CacheState.Cas.GetType() == EOnDemandInstallCasType::MMap)
+	{
+		// Any blocks added to BlocksToDefrag must be marked in the CAS as pending defrag,
+		// OpenMapped will wait until they have been defragged.
+		CasUniqueLock = TDynamicUniqueLock(CacheState.Cas);
+		SkipLock = &UE::DeferLock;
+	}
+
 	if (TotalBytesToFree)
 	{
 		// Partial defrag
@@ -3833,30 +4266,36 @@ FResult FOnDemandInstallCache::Defrag(
 			if (!bPossibleToFreeBytes && Info.RefSize < Info.FileSize)
 			{
 				// Block is fragmented
-				BlocksToDefragFragmentedSize += (Info.FileSize - Info.RefSize);
-				BlocksToDefragTotalSize += Info.FileSize;
+				if (CacheState.Cas.SetPendingDefrag(BlockId, SkipLock))
+				{
+					BlocksToDefragFragmentedSize += (Info.FileSize - Info.RefSize);
+					BlocksToDefragTotalSize += Info.FileSize;
 
-				FreedBlockBytes += Info.FileSize;
-				NewBlockBytes += Info.RefSize; // For now, assume that nothing will be moved to the current block
+					FreedBlockBytes += Info.FileSize;
+					NewBlockBytes += Info.RefSize; // For now, assume that nothing will be moved to the current block
 
-				BlocksToDefrag.Emplace(BlockId, Info);
+					BlocksToDefrag.Emplace(BlockId, Info);
 
-				if (FreedBlockBytes >= NewBlockBytes && FreedBlockBytes - NewBlockBytes >= *TotalBytesToFree)
-				{
-					bPossibleToFreeBytes = true;
+					if (FreedBlockBytes >= NewBlockBytes && FreedBlockBytes - NewBlockBytes >= *TotalBytesToFree)
+					{
+						bPossibleToFreeBytes = true;
+					}
 				}
 			}
 			else if (Info.FileSize < CacheState.Cas.GetMinBlockSize())
 			{
 				// Block is too small whether or not its fragmented
-				if (ensure(Info.RefSize <= Info.FileSize))
+				if (CacheState.Cas.SetPendingDefrag(BlockId, SkipLock))
 				{
-					BlocksToDefragFragmentedSize += (Info.FileSize - Info.RefSize);
-				}
+					if (ensure(Info.RefSize <= Info.FileSize))
+					{
+						BlocksToDefragFragmentedSize += (Info.FileSize - Info.RefSize);
+					}
 
-				BlocksToDefragTotalSize += Info.FileSize;
+					BlocksToDefragTotalSize += Info.FileSize;
 
-				BlocksToDefrag.Emplace(BlockId, Info);
+					BlocksToDefrag.Emplace(BlockId, Info);
+				}
 			}
 		}
 
@@ -3873,6 +4312,8 @@ FResult FOnDemandInstallCache::Defrag(
 	}
 	else
 	{
+		bool bSkippedFragmentedBlock = false;
+
 		// Full defrag
 		for (const TPair<FCasBlockId, FCasBlockInfo>& Kv : BlockInfo)
 		{
@@ -3886,34 +4327,62 @@ FResult FOnDemandInstallCache::Defrag(
 			if (Info.RefSize < Info.FileSize)
 			{
 				// Block is fragmented
-				BlocksToDefragFragmentedSize += (Info.FileSize - Info.RefSize);
-				BlocksToDefragTotalSize += Info.FileSize;
+				if (CacheState.Cas.SetPendingDefrag(BlockId, SkipLock))
+				{
+					BlocksToDefragFragmentedSize += (Info.FileSize - Info.RefSize);
+					BlocksToDefragTotalSize += Info.FileSize;
 
-				BlocksToDefrag.Emplace(BlockId, Info);
+					BlocksToDefrag.Emplace(BlockId, Info);
+				}
+				else
+				{
+					bSkippedFragmentedBlock = true;
+				}
 			}
 			else if (Info.FileSize < CacheState.Cas.GetMinBlockSize())
 			{
 				// Block is too small whether or not its fragmented
-				if (ensure(Info.RefSize <= Info.FileSize))
+				if (CacheState.Cas.SetPendingDefrag(BlockId, SkipLock))
 				{
-					BlocksToDefragFragmentedSize += (Info.FileSize - Info.RefSize);
-				}
+					if (ensure(Info.RefSize <= Info.FileSize))
+					{
+						BlocksToDefragFragmentedSize += (Info.FileSize - Info.RefSize);
+					}
 
-				BlocksToDefragTotalSize += Info.FileSize;
+					BlocksToDefragTotalSize += Info.FileSize;
 
-				BlocksToDefrag.Emplace(BlockId, Info);
+					BlocksToDefrag.Emplace(BlockId, Info);
+				}
+				else
+				{
+					bSkippedFragmentedBlock = true;
+				}
 			}
 		}
 
 		if (BlocksToDefrag.IsEmpty())
 		{
-			// Already defragged
-			UE_LOG(LogIoStoreOnDemand, Display, TEXT("Cache %s not fragmented."), LexToString(CacheState.Type));
+			if (bSkippedFragmentedBlock)
+			{
+				UE_LOG(LogIoStoreOnDemand, Display, TEXT("No %s cache blocks are eligable for defrag."), LexToString(CacheState.Cas.GetType()));
+			}
+			else
+			{
+				// Already defragged
+				UE_LOG(LogIoStoreOnDemand, Display, TEXT("%s cache not fragmented."), LexToString(CacheState.Cas.GetType()));
+			}
+			
 			FResult Result = MakeValue();
 			FOnDemandInstallCacheStats::OnDefrag(Result, TotalFragmentedBytes, TotalCachedBytes, BlocksToDefragFragmentedSize, BlocksToDefragTotalSize);
 			return Result;
 		}
 	}
+	
+	if (CasUniqueLock)
+	{
+		CasUniqueLock.Unlock();
+		SkipLock = nullptr;
+	}	
 
 	UE_LOG(LogIoStoreOnDemand, Display, TEXT("Defrag found %" UINT64_FMT " fragmented bytes of %" UINT64_FMT " total bytes in %i blocks."), 
 		BlocksToDefragFragmentedSize, BlocksToDefragTotalSize, BlocksToDefrag.Num()); 
@@ -3922,7 +4391,7 @@ FResult FOnDemandInstallCache::Defrag(
 	{
 		// Ruh-Roh! There's not enough of the disk quota left to run a defrag!
 		const FString ErrorMsg = FString::Printf(TEXT("Cache size is greater than disk quota for %s - Cannot Defragment!, MaxCacheSize=%.2lf MiB"),
-			LexToString(CacheState.Type), ToMiB(CacheState.MaxCacheSize));
+			LexToString(CacheState.Cas.GetType()), ToMiB(CacheState.MaxCacheSize));
 		UE_LOG(LogIoStoreOnDemand, Error, TEXT("%s"), *ErrorMsg);
 
 		FInstallCacheErrorContext Ctx = MakeInstallCacheErrorContext(CacheState, TotalCachedBytes, TotalBytesToFree ? *TotalBytesToFree : TotalFragmentedBytes);
@@ -3930,7 +4399,7 @@ FResult FOnDemandInstallCache::Defrag(
 		FResult Result = MakeError(InstallCacheDefragError::MakeError(MoveTemp(Ctx), UE::UnifiedError::EDetailFilter::All));
 
 		// Append a critical error entry to clear the cache at next startup
-		FCasJournal::FTransaction Transaction = FCasJournal::Begin(CacheState.Type, GetJournalFilename());
+		FCasJournal::FTransaction Transaction = FCasJournal::Begin(CacheState.Cas.GetType(), GetJournalFilename());
 		Transaction.CriticalError(FCasJournal::EErrorCode::DefragOutOfDiskSpace);
 		if (FResult CommitResult = FCasJournal::Commit(MoveTemp(Transaction)); CommitResult.HasError())
 		{
@@ -4054,7 +4523,7 @@ FResult FOnDemandInstallCache::Defrag(
 						*LexToString(ReffedChunk.ChunkId), DefragBlock.BlockId.Id, ReffedChunk.BlockOffset);
 
 					// Append a critical error entry to clear the cache at next startup
-					FCasJournal::FTransaction Transaction = FCasJournal::Begin(CacheState.Type, GetJournalFilename());
+					FCasJournal::FTransaction Transaction = FCasJournal::Begin(CacheState.Cas.GetType(), GetJournalFilename());
 					Transaction.CriticalError(FCasJournal::EErrorCode::DefragHashMismatch);
 					if (FResult CommitResult = FCasJournal::Commit(MoveTemp(Transaction)); CommitResult.HasError())
 					{
@@ -4102,7 +4571,7 @@ FResult FOnDemandInstallCache::Defrag(
 			check(DefragPendingChunks.IsEmpty());
 		}
 
-		FCasJournal::FTransaction Transaction = FCasJournal::Begin(CacheState.Type, GetJournalFilename());
+		FCasJournal::FTransaction Transaction = FCasJournal::Begin(CacheState.Cas.GetType(), GetJournalFilename());
 
 		auto HandleChunkRemoved = [&Transaction](const FCasAddr& Addr)
 		{
@@ -4130,7 +4599,7 @@ FResult FOnDemandInstallCache::Defrag(
 	}
 
 	UE_LOG(LogIoStoreOnDemand, Display, TEXT("Defrag removed %" UINT64_FMT " fragmented bytes of %" UINT64_FMT " total bytes in %i blocks from CAS %s."),
-		BlocksToDefragTotalSize, BlocksToDefragTotalSize, BlocksToDefrag.Num(), LexToString(CacheState.Type));
+		BlocksToDefragTotalSize, BlocksToDefragTotalSize, BlocksToDefrag.Num(), LexToString(CacheState.Cas.GetType()));
 
 	if (OutTotalFreedBytes)
 	{
@@ -4212,7 +4681,7 @@ void FOnDemandInstallCache::UpdateLastAccess(TConstArrayView<FIoHash> ChunkHashe
 
 	for (FCacheState& CacheState : CacheStates)
 	{
-		const EOnDemandInstallCasType CasType = CacheState.Type;
+		const EOnDemandInstallCasType CasType = CacheState.Cas.GetType();
 		TUniqueLock Lock(CacheState.Cas);
 
 		for (const FIoHash& ChunkHash : ChunkHashes)
@@ -4348,7 +4817,7 @@ FResult FOnDemandInstallCache::FlushPendingChunksImpl(FCacheState& CacheState, c
 
 	while (ChunkIdx < Chunks.Chunks.Num())
 	{
-		FCasJournal::FTransaction Transaction = FCasJournal::Begin(CacheState.Type, GetJournalFilename());
+		FCasJournal::FTransaction Transaction = FCasJournal::Begin(CacheState.Cas.GetType(), GetJournalFilename());
 		
 		// Only open for append if continuing a block.
 		const bool bAppendToBlock = CurrentBlockId.IsValid();
@@ -4549,12 +5018,12 @@ void FOnDemandInstallCache::CompleteRequest(FIoRequestImpl* Request, EIoErrorCod
 
 FOnDemandInstallCache::FCacheState* FOnDemandInstallCache::GetCacheState(EOnDemandInstallCasType CacheType)
 {
-	return Algo::FindBy(CacheStates, CacheType, &FCacheState::Type);
+	return Algo::FindBy(CacheStates, CacheType, Projection(&FCacheState::Cas, &FCas::GetType));
 }
 
 const FOnDemandInstallCache::FCacheState* FOnDemandInstallCache::GetCacheState(EOnDemandInstallCasType CacheType) const
 {
-	return Algo::FindBy(CacheStates, CacheType, &FCacheState::Type);
+	return Algo::FindBy(CacheStates, CacheType, Projection(&FCacheState::Cas, &FCas::GetType));
 }
 
 UE::UnifiedError::IoStoreOnDemand::FInstallCacheErrorContext FOnDemandInstallCache::MakeInstallCacheErrorContext(FCacheState& CacheState, uint64 TotalCachedBytes, uint64 RequestedBytes, uint32 LineNo)
@@ -4562,7 +5031,7 @@ UE::UnifiedError::IoStoreOnDemand::FInstallCacheErrorContext FOnDemandInstallCac
 	UE::UnifiedError::IoStoreOnDemand::FInstallCacheErrorContext OutCtx;
 
 	const FString CachRootDirectory = CacheState.Cas.GetRootDirectory();
-	OutCtx.CasType					= LexToString(CacheState.Type);
+	OutCtx.CasType					= LexToString(CacheState.Cas.GetType());
 	OutCtx.bDiskQuerySucceeded		= FPlatformMisc::GetDiskTotalAndFreeSpace(CachRootDirectory, OutCtx.DiskTotalBytes, OutCtx.DiskFreeBytes);
 	OutCtx.MaxCacheSize				= CacheState.MaxCacheSize;
 	OutCtx.CacheSize				= TotalCachedBytes;
diff --git a/Engine/Source/Runtime/Experimental/IoStore/OnDemandCore/Public/IO/OnDemandError.h b/Engine/Source/Runtime/Experimental/IoStore/OnDemandCore/Public/IO/OnDemandError.h
--- a/Engine/Source/Runtime/Experimental/IoStore/OnDemandCore/Public/IO/OnDemandError.h
+++ b/Engine/Source/Runtime/Experimental/IoStore/OnDemandCore/Public/IO/OnDemandError.h
@@ -58,7 +58,8 @@ enum class ECasErrorCode : uint32
 	CommitJournalFailed,
 	CreateSnapshotFailed,
 	LoadSnapshotFailed,
-	SaveSnapshotFailed
+	SaveSnapshotFailed,
+	OpenMappedFailed,
 };
 
 inline const TCHAR* ToString(ECasErrorCode Code)
@@ -89,6 +90,8 @@ inline const TCHAR* ToString(ECasErrorCode Code)
 			return TEXT("LoadSnapshotFailed");
 		case ECasErrorCode::SaveSnapshotFailed:
 			return TEXT("SaveSnapshotFailed");
+		case ECasErrorCode::OpenMappedFailed:
+			return TEXT("OpenMappedFailed");
 		default:
 			return TEXT("<InvalidErrorCode>");
 	};
diff --git a/Engine/Source/Runtime/PakFile/Private/FileIoDispatcherBackend.cpp b/Engine/Source/Runtime/PakFile/Private/FileIoDispatcherBackend.cpp
--- a/Engine/Source/Runtime/PakFile/Private/FileIoDispatcherBackend.cpp
+++ b/Engine/Source/Runtime/PakFile/Private/FileIoDispatcherBackend.cpp
@@ -1035,7 +1035,7 @@ TIoStatusOr<FIoMappedRegion> FFileIoDispatcherBackend::OpenMapped(const FIoChunk
 	if (Partition.MappedFileHandle.IsValid() == false)
 	{
 		IPlatformFile& Ipf = FPlatformFileManager::Get().GetPlatformFile();
-		if (FOpenMappedResult Result = Ipf.OpenMappedEx(*Partition.Filename); !Result.HasError())
+		if (FOpenMappedResult Result = Ipf.OpenMappedEx2(*Partition.Filename); !Result.HasError())
 		{
 			Partition.MappedFileHandle = Result.StealValue();
 		}
diff --git a/Engine/Source/Runtime/PakFile/Private/IoDispatcherFileBackend.cpp b/Engine/Source/Runtime/PakFile/Private/IoDispatcherFileBackend.cpp
--- a/Engine/Source/Runtime/PakFile/Private/IoDispatcherFileBackend.cpp
+++ b/Engine/Source/Runtime/PakFile/Private/IoDispatcherFileBackend.cpp
@@ -668,7 +668,7 @@ IMappedFileHandle* FFileIoStoreContainerFilePartition::BorrowMappedFileHandle()
 		if (UNLIKELY(Handle == nullptr))
 		{
 			IPlatformFile& Ipf = FPlatformFileManager::Get().GetPlatformFile();
-			FOpenMappedResult Result = Ipf.OpenMappedEx(*FilePath);
+			FOpenMappedResult Result = Ipf.OpenMappedEx2(*FilePath);
 			if (!Result.HasError())
 			{
 				Handle = Result.StealValue().Release();
