Add new console command  'iax.HttpErrorChance' which allows the user to provide the chance that any given IAX http request should fail.

#rnx

- The request takes the chance (0-100%) that a http request should be 'failed' and an optional filter (IAD/IAX/All) so that we can selectively apply it.
- When a request chain is selected to be sent to our http client we check to see if there is an active percent chance that we should simulate an error and roll a random number against that chance to see if it should fail.
- If the request chain should fail then we immediately fail it without sending it to our http client, this way we don't trigger a hostgroup failure as the error is not associated with the hostgroup.
- We already have code that early fails a http request if the hostgroup is not connected, so we can make use of that. The only difference is if we are failing on a simulated error we want to record an error, not a request cancellation.
- Ended up reworking the existing code for 'iax.InvalidUrlChance' so that it could be used for the new code.

#Synced-CL 52050999

diff --git a/Engine/Source/Runtime/Experimental/IoStore/OnDemand/Private/OnDemandHttpIoDispacher.cpp b/Engine/Source/Runtime/Experimental/IoStore/OnDemand/Private/OnDemandHttpIoDispacher.cpp
--- a/Engine/Source/Runtime/Experimental/IoStore/OnDemand/Private/OnDemandHttpIoDispacher.cpp
+++ b/Engine/Source/Runtime/Experimental/IoStore/OnDemand/Private/OnDemandHttpIoDispacher.cpp
@@ -409,102 +409,132 @@ class FRequestErrorSimulator
 public:
 	FRequestErrorSimulator()
 	{
-		InvalidUrlCmdHandle = FConsoleCommandPtr(IConsoleManager::Get().RegisterConsoleCommand(
-			TEXT("iax.InvalidUrlChance"),
-			TEXT("Chance that a url for a GET request will be invalid (0-100%) and optional request type filter [IAD/IAS/All]. By default the request filter will be set to All."),
-			FConsoleCommandWithArgsAndOutputDeviceDelegate::CreateRaw(this, &FRequestErrorSimulator::OnSetInvalidUrlChance),
-			ECVF_Default));
+		SimulateError.Bind(TEXT("iax.HttpErrorChance"), TEXT("Chance that a request will fail with an error (0-100%) and optional request type filter [IAD/IAS/All]. By default the request filter will be set to All."));
+		SimulateInvalidUrl.Bind(TEXT("iax.InvalidUrlChance"), TEXT("Chance that a url for a GET request will be invalid (0-100%) and optional request type filter [IAD/IAS/All]. By default the request filter will be set to All."));
 	}
 
 	~FRequestErrorSimulator() = default;
 
-	bool SimulateUrlError(EHttpRequestType RequestType)
+	bool ShouldSimulateError(EHttpRequestType RequestType)
 	{
-		if (InvalidUrlChance <= 0.0f)
-		{
-			return false;
-		}
-
-		if (!EnumHasAnyFlags(InvalidUrlFilter, ToFilter(RequestType)))
-		{
-			return false;
-		}
-
-		return (FMath::FRand() * 100.0f) < InvalidUrlChance;
+		return SimulateError.TriggerForRequestType(RequestType);
 	}
 
-private:
-	enum class ERequestFilter : uint8
+	bool ShouldSimulateInvalidUrl(EHttpRequestType RequestType)
 	{
-		Installed = 1 << 0,
-		Streaming = 1 << 1,
+		return SimulateInvalidUrl.TriggerForRequestType(RequestType);
+	}
 
-		All = Installed | Streaming
-	};
-	FRIEND_ENUM_CLASS_FLAGS(ERequestFilter);
+private:
 
-	ERequestFilter ToFilter(EHttpRequestType Type)
+	struct FChance
 	{
-		switch (Type)
+		void Bind(const TCHAR* Name, const TCHAR* HelpTip)
 		{
-			case EHttpRequestType::Installed:
-				return ERequestFilter::Installed;
-			case EHttpRequestType::Streaming:
-				return ERequestFilter::Streaming;
-			default:
-				checkNoEntry();
-				return ERequestFilter::All;
+			ConsoleCommandHandle = FConsoleCommandPtr(IConsoleManager::Get().RegisterConsoleCommand(
+				Name, HelpTip,
+				FConsoleCommandWithArgsAndOutputDeviceDelegate::CreateRaw(this, &FChance::OnSet),
+				ECVF_Default));
 		}
-	}
 
-	void OnSetInvalidUrlChance(const TArray<FString>& Args, FOutputDevice& Ar)
-	{
-		if (Args.Num() < 1 || Args.Num() > 2)
+		bool TriggerForRequestType(EHttpRequestType RequestType)
 		{
-			UE_CVAR_ERROR_LOG(TEXT("Invalid args! Expected iax.InvalidUrlChance [chance (float)] [IAD/IAS/All (string)]"));
-			return;
+			if (Chance <= 0.0f)
+			{
+				return false;
+			}
+
+			if (!EnumHasAnyFlags(Filter, ToFilter(RequestType)))
+			{
+				return false;
+			}
+
+			return (FMath::FRand() * 100.0f) < Chance;
 		}
 
-		float ParsedChance = 0.0f;
-		if (!LexTryParseString(ParsedChance, *Args[0]) || ParsedChance < 0.0f || ParsedChance > 100.0f)
+		enum class ERequestFilter : uint8
 		{
-			UE_CVAR_ERROR_LOG(TEXT("Invalid arg!, first arg should be a number between 0.0 and 100.0"));
-			return;
-		}
+			Invalid = 0, // We must have at least one request type, so this is our error state
+			Installed = 1 << 0,
+			Streaming = 1 << 1,
 
-		ERequestFilter ParsedFilter = ERequestFilter::All;
+			All = Installed | Streaming
+		};
 
-		if (Args.Num() == 2)
+	private:
+		ERequestFilter ToFilter(EHttpRequestType Type)
 		{
-			if (Args[1] == TEXT("IAD"))
+			switch (Type)
 			{
-				ParsedFilter = ERequestFilter::Installed;
+				case EHttpRequestType::Installed:
+					return ERequestFilter::Installed;
+				case EHttpRequestType::Streaming:
+					return ERequestFilter::Streaming;
+				default:
+					checkNoEntry();
+					return ERequestFilter::All;
 			}
-			else if (Args[1] == TEXT("IAS"))
+		}
+
+		void OnSet(const TArray<FString>& Args, FOutputDevice& Ar)
+		{
+			if (Args.Num() < 1 || Args.Num() > 2)
 			{
-				ParsedFilter = ERequestFilter::Streaming;
+				UE_CVAR_ERROR_LOG(TEXT("Invalid args! Expected '<command> [chance (float)] [IAD/IAS/All (string)]'"));
+				return;
 			}
-			else if (Args[1] == TEXT("All"))
+
+			float ParsedChance = 0.0f;
+			if (!LexTryParseString(ParsedChance, *Args[0]) || ParsedChance < 0.0f || ParsedChance > 100.0f)
 			{
-				ParsedFilter = ERequestFilter::All;
+				UE_CVAR_ERROR_LOG(TEXT("Invalid arg!, first arg should be a number between 0.0 and 100.0"));
+				return;
 			}
-			else
+
+			ERequestFilter ParsedFilter = ParseFilter(Args, Ar);
+			if (ParsedFilter == ERequestFilter::Invalid)
 			{
-				UE_CVAR_ERROR_LOG(TEXT("Unrecognized arg!, second arg should be [IAD/IAS/All]"));
 				return;
 			}
+
+			Chance = ParsedChance;
+			Filter = ParsedFilter;
 		}
 
-		InvalidUrlChance = ParsedChance;
-		InvalidUrlFilter = ParsedFilter;
-	}
-private:
+		ERequestFilter ParseFilter(const TArray<FString>& Args, FOutputDevice& Ar)
+		{
+			if (Args.Num() == 2)
+			{
+				if (Args[1] == TEXT("IAD"))
+				{
+					return ERequestFilter::Installed;
+				}
+				else if (Args[1] == TEXT("IAS"))
+				{
+					return ERequestFilter::Streaming;
+				}
+				else if (Args[1] == TEXT("All"))
+				{
+					return ERequestFilter::All;
+				}
+				else
+				{
+					UE_CVAR_ERROR_LOG(TEXT("Unrecognized arg!, filter arg should be [IAD/IAS/All]"));
+					return ERequestFilter::Invalid;
+				}
+			}
+
+			return ERequestFilter::All;
+		}
+
+		FConsoleCommandPtr	ConsoleCommandHandle;
+		float				Chance = 0.0f;
+		ERequestFilter		Filter = ERequestFilter::All;
+	};
 
-	FConsoleCommandPtr InvalidUrlCmdHandle;
-	float InvalidUrlChance = 0.0f;
-	ERequestFilter InvalidUrlFilter = ERequestFilter::All;
+	FChance SimulateError;
+	FChance SimulateInvalidUrl;
 };
-ENUM_CLASS_FLAGS(FRequestErrorSimulator::ERequestFilter);
 
 #endif // UE_ALLOW_ERROR_SIMULATION
 
@@ -1350,8 +1380,10 @@ private:
 
 	void							UpdateHostGroup(FIASHostGroup& Hostgroup, const FMultiEndpointHttpClientResponse& Response) const;
 
+	bool							ShouldSimulateError(EHttpRequestType RequestType);
+
 #if UE_ALLOW_HTTP_PAUSE
-	void OnPauseCommand(bool bShouldPause, const TArray<FString>& Args, FOutputDevice& Ar);
+	void							OnPauseCommand(bool bShouldPause, const TArray<FString>& Args, FOutputDevice& Ar);
 #endif //UE_ALLOW_HTTP_PAUSE
 
 	TUniquePtr<IIasCache>					Cache;
@@ -1981,13 +2013,11 @@ bool FHttpDispatcher::ProcessHttpRequests()
 			FOnDemandIoBackendStats::Get()->OnHttpStarted(HttpRequest->GetHttpRequestType());
 
 			FIASHostGroup HostGroup = FHostGroupManager::Get().Find(HttpRequest->HostGroupName);
-			if (HostGroup.IsConnected() == false)
+			if (!HostGroup.IsConnected() || ShouldSimulateError(HttpRequest->GetHttpRequestType()))
 			{
-				// Technically this request is being skipped because of a pre-existing error. It is not
-				// an error itself and it is not being canceled by higher level code. However we do not
-				// currently have a statistic for that and we have to call one of the existing types in
-				// order to correctly reduce the pending count.
-				const bool bWasCanceled = true;
+				// If the hostgroup is disconnected then we are canceling the request as we cannot process it. If the hostgroup is connected
+				// then we are simulating an error.
+				const bool bWasCanceled = !HostGroup.IsConnected();
 				FOnDemandIoBackendStats::Get()->OnHttpCompleted(HttpRequest->GetHttpRequestType(), 0, 0, 0, 0, 0, bWasCanceled);
 				CompleteRequest(*HttpRequest, HttpRequest->Buffer, FIoHttpRange(), 0, false);
 
@@ -2019,7 +2049,7 @@ bool FHttpDispatcher::ProcessHttpRequests()
 			}
 
 #if UE_ALLOW_ERROR_SIMULATION
-			if (RequestErrorSimulator.SimulateUrlError(HttpRequest->GetHttpRequestType()))
+			if (RequestErrorSimulator.ShouldSimulateInvalidUrl(HttpRequest->GetHttpRequestType()))
 			{
 				Url << "-DebugInvalidUrl";
 			}
@@ -2237,6 +2267,14 @@ void FHttpDispatcher::OnPauseCommand(bool bShouldPause, const TArray<FString>& A
 
 #endif // UE_ALLOW_HTTP_PAUSE
 
+bool FHttpDispatcher::ShouldSimulateError(EHttpRequestType RequestType)
+{
+#if UE_ALLOW_ERROR_SIMULATION
+	return RequestErrorSimulator.ShouldSimulateError(RequestType);
+#else
+	return false;
+#endif
+}
 void FHttpDispatcher::ApplyCancellationToken(FHttpRequest* RequestChain, FMultiEndpointHttpClient::FHttpTicketId TicketId)
 {
 	TUniqueLock _(CancellationMutex);
