Fix remval event collision bug when Pull Data is multithreaded
Enable multithreading for Pull Data
Geoemtry Collection particle are marked dirty in single thread

#Synced-CL 52136471

diff --git a/Engine/Source/Runtime/Experimental/Chaos/Private/PBDRigidsSolver.cpp b/Engine/Source/Runtime/Experimental/Chaos/Private/PBDRigidsSolver.cpp
--- a/Engine/Source/Runtime/Experimental/Chaos/Private/PBDRigidsSolver.cpp
+++ b/Engine/Source/Runtime/Experimental/Chaos/Private/PBDRigidsSolver.cpp
@@ -428,17 +428,9 @@ namespace Chaos
 		bool bRemoveParticleFromMovingKinematicsOnDisable = true;
 		FAutoConsoleVariableRef CVarChaosRemoveParticleFromMovingKinematicsOnDisable(TEXT("p.Chaos.RemoveParticleFromMovingKinematicsOnDisable"), bRemoveParticleFromMovingKinematicsOnDisable, TEXT(""));
 
-		bool bSingleThreadPullData = true;
+		bool bSingleThreadPullData = false;
 		FAutoConsoleVariableRef CVarChaosSingleThreadPushData(TEXT("p.Chaos.SingleThreadPullData"), bSingleThreadPullData, TEXT("Run pull Data in Single Thread."));
 	}
-
-	namespace Private
-	{
-		bool IsUsingSingleThreadPullData()
-		{
-			return CVars::bSingleThreadPullData;
-		}
-	}
 }
 
 namespace PhysicsReplicationCVars
@@ -2820,6 +2812,7 @@ namespace Chaos
 				{
 					PullData->DirtyGeometryCollections.AddDefaulted();
 					ActiveGC[Idx]->BufferPhysicsResults_Internal(this, PullData->DirtyGeometryCollections.Last());
+					ActiveGC[Idx]->FlushDirtyParticles(*this);
 				}
 			}
 			{
@@ -2865,6 +2858,10 @@ namespace Chaos
 		if (bRunMultithreading)
 		{
 			TaskDispatcher->WaitTaskEndBufferResults();
+			for (int32 Idx = 0; Idx < TaskDispatcher->GetActiveGC().Num(); ++Idx)
+			{
+				TaskDispatcher->GetActiveGC()[Idx]->FlushDirtyParticles(*this);
+			}
 		}
 
 		// Now that results have been buffered we have completed a solve step so we can broadcast that event
diff --git a/Engine/Source/Runtime/Experimental/Chaos/Private/PhysicsProxy/GeometryCollectionPhysicsProxy.cpp b/Engine/Source/Runtime/Experimental/Chaos/Private/PhysicsProxy/GeometryCollectionPhysicsProxy.cpp
--- a/Engine/Source/Runtime/Experimental/Chaos/Private/PhysicsProxy/GeometryCollectionPhysicsProxy.cpp
+++ b/Engine/Source/Runtime/Experimental/Chaos/Private/PhysicsProxy/GeometryCollectionPhysicsProxy.cpp
@@ -707,6 +707,7 @@ void FGeometryCollectionPhysicsProxy::Initialize(Chaos::FPBDRigidsEvolutionBase
 	
 	NumEffectiveParticles = CalculateEffectiveParticles(DynamicCollection, NumTransforms, Parameters.MaxSimulatedLevel, Parameters.EnableClustering, GetOwner(), EffectiveParticles);
 
+	ParticlesToDirty.Init(false, NumEffectiveParticles);
 	SolverClusterID.Init(nullptr, NumEffectiveParticles);
 	SolverClusterHandles.Init(nullptr, NumEffectiveParticles);
 	SolverParticleHandles.Init(nullptr, NumEffectiveParticles);
@@ -4059,23 +4060,28 @@ void FGeometryCollectionPhysicsProxy::BufferPhysicsResults_External(Chaos::FDirt
 	PrepareBufferData(BufferData, PhysicsThreadCollection);
 }
 
-static void UpdateParticleHandleTransform(Chaos::FPBDRigidsSolver& CurrentSolver, Chaos::FPBDRigidClusteredParticleHandle& Handle, const Chaos::FRigidTransform3& NewTransform)
+void FGeometryCollectionPhysicsProxy::FlushDirtyParticles(Chaos::FPBDRigidsSolver& CurrentSolver)
+{
+	for (TConstSetBitIterator<> It(ParticlesToDirty); It; ++It)
+	{
+		if (FClusterHandle* Handle = SolverParticleHandles[It.GetIndex()])
+		{
+			CurrentSolver.GetEvolution()->DirtyParticle(*Handle);
+		}
+	}
+	ParticlesToDirty.Init(false, ParticlesToDirty.Num());
+}
+
+static void UpdateParticleHandleTransform(Chaos::FPBDRigidsSolver& CurrentSolver, Chaos::FPBDRigidClusteredParticleHandle& Handle, const Chaos::FRigidTransform3& NewTransform, TBitArray<>& ParticlesToDirty, int32 ParticleIndex)
 {
 	Handle.SetX(NewTransform.GetTranslation());
 	Handle.SetR(NewTransform.GetRotation());
 	Handle.UpdateWorldSpaceState(NewTransform, Chaos::FVec3{0});
 
-	// Temporary fix for the bug where player looses collision with partial clusters after using remove on sleep/break 
-	// Issue was introduced at CL 45952481
-	// Current temporary fix is to re-enable Single Thread Pull Data on the solver and when doing so we can safely update the dirty particle from this proxy
-	// Proper fix will require marking the dirty particle in a thread safe way via a potentially new API and having the actual dirtying happening as apost process to pull data when we use multi threaded pull data
-	if (Chaos::Private::IsUsingSingleThreadPullData())
-	{
-		CurrentSolver.GetEvolution()->DirtyParticle(Handle);
-	}
+	ParticlesToDirty[ParticleIndex] = true;
 }
 
-static void UpdateParticleHandleTransformIfNeeded(Chaos::FPBDRigidsSolver& CurrentSolver, Chaos::FPBDRigidClusteredParticleHandle& Handle, const Chaos::FRigidTransform3& NewTransform)
+static void UpdateParticleHandleTransformIfNeeded(Chaos::FPBDRigidsSolver& CurrentSolver, Chaos::FPBDRigidClusteredParticleHandle& Handle, const Chaos::FRigidTransform3& NewTransform, TBitArray<>& ParticlesToDirty, int32 ParticleIndex)
 {
 	const Chaos::FVec3 NewX = NewTransform.GetTranslation();
 	const Chaos::FVec3 OldX = Handle.GetX();
@@ -4085,7 +4091,7 @@ static void UpdateParticleHandleTransformIfNeeded(Chaos::FPBDRigidsSolver& Curre
 	
 	if (OldX != NewX || OldR != NewR)
 	{
-		UpdateParticleHandleTransform(CurrentSolver, Handle, NewTransform);
+		UpdateParticleHandleTransform(CurrentSolver, Handle, NewTransform, ParticlesToDirty, ParticleIndex);
 	}
 }
 
@@ -4248,7 +4254,7 @@ void FGeometryCollectionPhysicsProxy::BufferPhysicsResults_Internal(Chaos::FPBDR
 				bHasChanged = true;
 
 				// as we just transitioned from disabled to non disabled the update is unconditional 
-				UpdateParticleHandleTransform(*CurrentSolver, *Handle, ParticleToWorld);
+				UpdateParticleHandleTransform(*CurrentSolver, *Handle, ParticleToWorld, ParticlesToDirty, ParticleIndex);
 			}
 
 			if (bGeometryCollectionEnabledNestedChildTransformUpdates)
@@ -4256,7 +4262,7 @@ void FGeometryCollectionPhysicsProxy::BufferPhysicsResults_Internal(Chaos::FPBDR
 				if (ClusterParent && !ClusterParent->Disabled())
 				{
 					const FRigidTransform3 ChildToWorld = Handle->ChildToParent() * FRigidTransform3(ClusterParent->GetX(), ClusterParent->GetR());
-					UpdateParticleHandleTransformIfNeeded(*CurrentSolver, *Handle, ChildToWorld);
+					UpdateParticleHandleTransformIfNeeded(*CurrentSolver, *Handle, ChildToWorld, ParticlesToDirty, ParticleIndex);
 					// fields may have applied velocities, we need to make sure to clear that up, so that we don't accumulate
 					VelocityData.ParticleV = FVec3::ZeroVector;
 					VelocityData.ParticleW = FVec3::ZeroVector;
diff --git a/Engine/Source/Runtime/Experimental/Chaos/Private/TaskDispatcherSolver.h b/Engine/Source/Runtime/Experimental/Chaos/Private/TaskDispatcherSolver.h
--- a/Engine/Source/Runtime/Experimental/Chaos/Private/TaskDispatcherSolver.h
+++ b/Engine/Source/Runtime/Experimental/Chaos/Private/TaskDispatcherSolver.h
@@ -60,6 +60,11 @@ public:
 
 	void Initialize(TParticleView<FPBDRigidParticles>& DirtyParticles);
 
+	TArray<FGeometryCollectionPhysicsProxy*>& GetActiveGC()
+	{
+		return ActiveGC;
+	}
+
 	template<typename LambdaRigid, typename LambdaGCAndCluster>
 	int32 CollectPhysicsResults(LambdaRigid CollectRigidResults, LambdaGCAndCluster CollectGCAndClustersResults, int32 NumDirty, int32 DispatchBatchIndex)
 	{
diff --git a/Engine/Source/Runtime/Experimental/Chaos/Public/PhysicsProxy/GeometryCollectionPhysicsProxy.h b/Engine/Source/Runtime/Experimental/Chaos/Public/PhysicsProxy/GeometryCollectionPhysicsProxy.h
--- a/Engine/Source/Runtime/Experimental/Chaos/Public/PhysicsProxy/GeometryCollectionPhysicsProxy.h
+++ b/Engine/Source/Runtime/Experimental/Chaos/Public/PhysicsProxy/GeometryCollectionPhysicsProxy.h
@@ -516,6 +516,8 @@ public:
 
 	CHAOS_API float ComputeMaterialBasedDamageThreshold_Internal(Chaos::FPBDRigidClusteredParticleHandle& ClusteredParticle) const;
 
+	UE_INTERNAL CHAOS_API void FlushDirtyParticles(Chaos::FPBDRigidsSolver& Solver);
+
 	FProxyInterpolationBase* GetInterpolationData() { return InterpolationData.Get(); }
 	const FProxyInterpolationBase* GetInterpolationData() const { return InterpolationData.Get(); }
 
@@ -769,6 +771,7 @@ private:
 	TArray<int32> FromParticleToTransformIndex;
 	TArray<int32> FromTransformToParticleIndex;
 	TBitArray<> EffectiveParticles;
+	TBitArray<> ParticlesToDirty;
 
 	// Game thread particles 
 	TArray<TUniquePtr<FParticle>> GTParticles;
