Implement weight maps for ClothCollisionThickness for all colliders instead of just skinned meshes.

#Synced-CL 52021766

diff --git a/Engine/Plugins/ChaosCloth/Source/ChaosCloth/Private/ChaosCloth/ChaosClothConstraints.cpp b/Engine/Plugins/ChaosCloth/Source/ChaosCloth/Private/ChaosCloth/ChaosClothConstraints.cpp
--- a/Engine/Plugins/ChaosCloth/Source/ChaosCloth/Private/ChaosCloth/ChaosClothConstraints.cpp
+++ b/Engine/Plugins/ChaosCloth/Source/ChaosCloth/Private/ChaosCloth/ChaosClothConstraints.cpp
@@ -2107,7 +2107,7 @@ void FClothConstraints::Update(
 	{
 		static IConsoleVariable* const WriteCCDContacts = IConsoleManager::Get().FindConsoleVariable(TEXT("p.Chaos.PBDEvolution.WriteCCDContacts"));
 		const bool bWriteCCDContacts = WriteCCDContacts ? WriteCCDContacts->GetBool() : false;
-		CollisionConstraint->SetProperties(ConfigProperties);
+		CollisionConstraint->SetProperties(ConfigProperties, WeightMaps);
 		CollisionConstraint->SetWriteDebugContacts(bWriteCCDContacts);
 	}
 	if (SkinnedTriangleCollisionsConstraint)
diff --git a/Engine/Plugins/ChaosCloth/Source/ChaosCloth/Private/ChaosCloth/ChaosClothVisualization.cpp b/Engine/Plugins/ChaosCloth/Source/ChaosCloth/Private/ChaosCloth/ChaosClothVisualization.cpp
--- a/Engine/Plugins/ChaosCloth/Source/ChaosCloth/Private/ChaosCloth/ChaosClothVisualization.cpp
+++ b/Engine/Plugins/ChaosCloth/Source/ChaosCloth/Private/ChaosCloth/ChaosClothVisualization.cpp
@@ -761,28 +761,34 @@ namespace Chaos
 		}
 	}
 
-	static void DrawSphere(FPrimitiveDrawInterface* PDI, const ::Chaos::TSphere<::Chaos::FReal, 3>& Sphere, const FQuat& Rotation, const FVector& Position, const FReal Scale, const FLinearColor& Color, const FMaterialRenderProxy* SolidMaterial = nullptr)
+
+	static void DrawSphere(FPrimitiveDrawInterface* PDI, const FVector& Center, const FQuat& Rotation, const FReal Radius, const FLinearColor& Color, const FMaterialRenderProxy* SolidMaterial = nullptr)
 	{
-		const FVec3 Center = Position + Rotation.RotateVector(Scale * FVec3(Sphere.GetCenterf()));
 		if (!PDI)
 		{
-			const FReal Radius = Sphere.GetRadiusf() * Scale;
 			FDebugDrawQueue::GetInstance().DrawDebugSphere(Center, Radius, 12, Color.ToFColor(true), false, KINDA_SMALL_NUMBER, Chaos::Private::GetDepthPriority(), 0.f);
 			return;
 		}
 #if WITH_EDITOR
 		if (SolidMaterial)
 		{
-			::DrawSphere(PDI, Center, FRotator(Rotation), FVector(Sphere.GetRadiusf()), 32, 16, SolidMaterial, Chaos::Private::GetDepthPriority());
+			::DrawSphere(PDI, Center, FRotator(Rotation), FVector(Radius), 32, 16, SolidMaterial, Chaos::Private::GetDepthPriority());
 		}
 		else
 		{
-			const FTransform Transform(Rotation, Center, FVector(Scale));
-			DrawWireSphere(PDI, Transform, Color, Sphere.GetRadiusf(), 12, Chaos::Private::GetDepthPriority(), 0.0f, 0.001f, false);
+			const FTransform Transform(Rotation, Center, FVector(1.));
+			DrawWireSphere(PDI, Transform, Color, Radius, 12, Chaos::Private::GetDepthPriority(), 0.0f, 0.001f, false);
 		}
 #endif
 	}
 
+	static void DrawSphere(FPrimitiveDrawInterface* PDI, const ::Chaos::TSphere<::Chaos::FReal, 3>& Sphere, const FQuat& Rotation, const FVector& Position, const FReal Scale, const FLinearColor& Color, const FMaterialRenderProxy* SolidMaterial = nullptr)
+	{
+		const FVec3 Center = Position + Rotation.RotateVector(Scale * FVec3(Sphere.GetCenterf()));
+		const FReal Radius = Sphere.GetRadiusf() * Scale;
+		DrawSphere(PDI, Center, Rotation, Radius, Color, SolidMaterial);
+	}
+
 	static void DrawBox(FPrimitiveDrawInterface* PDI, const ::Chaos::FAABB3& Box, const FQuat& Rotation, const FVector& Position, FReal Scale,  const FLinearColor& Color)
 	{
 		if (!PDI)
@@ -3721,6 +3727,93 @@ namespace Chaos
 		}
 	}
 
+	void FClothVisualizationNoGC::DrawClothCollisionThickness(const FDrawContext& Context) const
+	{
+		if (!Solver)
+		{
+			return;
+		}
+		for (const FClothingSimulationCloth* const Cloth : Solver->GetCloths())
+		{
+			const int32 ParticleRangeId = Cloth->GetParticleRangeId(Solver);
+			if (ParticleRangeId == INDEX_NONE)
+			{
+				continue;
+			}
+
+			const FClothConstraints& ClothConstraints = Solver->GetClothConstraints(ParticleRangeId);
+
+			if (const Softs::FPBDSoftBodyCollisionConstraint* const CollisionConstraints = ClothConstraints.GetCollisionConstraint().Get())
+			{
+				const int32 Offset = Solver->IsLegacySolver() ? ParticleRangeId : 0;
+				const TSet<int32>* const DebugFilterSet = Context.Filter.VertexSet.IsEmpty() ? nullptr : Cloth->GetVertexSetByName(Solver, Context.Filter.VertexSet);
+				const int32 GlobalParticleOffset = Solver->GetGlobalParticleOffset(ParticleRangeId);
+				auto ShouldDrawParticle = [&Context, DebugFilterSet, GlobalParticleOffset](const int32 Index)
+					{
+						return Private::ShouldDrawParticle(Context.Filter, DebugFilterSet, Index, GlobalParticleOffset);
+					};
+				const TConstArrayView<Softs::FSolverVec3> Positions = Cloth->GetParticlePositions(Solver);
+				const TConstArrayView<Softs::FSolverReal> InvMasses = Cloth->GetParticleInvMasses(Solver);
+				for (int32 VertexIndex = 0; VertexIndex < Positions.Num(); ++VertexIndex)
+				{
+					const bool bIsKinematic0 = (InvMasses[VertexIndex] == (Softs::FSolverReal)0.);
+					if (bIsKinematic0 || !ShouldDrawParticle(VertexIndex))
+					{
+						continue;
+					}
+
+					const FVec3 Pos0 = GetWorldPosition(*Solver, Positions[VertexIndex]);
+					const FReal Radius1 = (FReal)CollisionConstraints->GetParticleClothThickness(VertexIndex + Offset);
+					DrawSphere(Context.PDI, Pos0, FQuat::Identity, Radius1, FLinearColor::White, nullptr);
+				}
+			}
+		}		
+	}
+
+	void FClothVisualizationNoGC::DrawInnerCollisionThickness(const FDrawContext& Context) const
+	{
+		if (!Solver)
+		{
+			return;
+		}
+
+		for (const FClothingSimulationCloth* const Cloth : Solver->GetCloths())
+		{
+			const int32 ParticleRangeId = Cloth->GetParticleRangeId(Solver);
+			if (ParticleRangeId == INDEX_NONE)
+			{
+				continue;
+			}
+
+			const FClothConstraints& ClothConstraints = Solver->GetClothConstraints(ParticleRangeId);
+
+			if (const Softs::FPBDSoftBodyCollisionConstraint* const CollisionConstraints = ClothConstraints.GetCollisionConstraint().Get())
+			{
+				const int32 Offset = Solver->IsLegacySolver() ? ParticleRangeId : 0;
+				const TSet<int32>* const DebugFilterSet = Context.Filter.VertexSet.IsEmpty() ? nullptr : Cloth->GetVertexSetByName(Solver, Context.Filter.VertexSet);
+				const int32 GlobalParticleOffset = Solver->GetGlobalParticleOffset(ParticleRangeId);
+				auto ShouldDrawParticle = [&Context, DebugFilterSet, GlobalParticleOffset](const int32 Index)
+					{
+						return Private::ShouldDrawParticle(Context.Filter, DebugFilterSet, Index, GlobalParticleOffset);
+					};
+				const TConstArrayView<Softs::FSolverVec3> Positions = Cloth->GetParticlePositions(Solver);
+				const TConstArrayView<Softs::FSolverReal> InvMasses = Cloth->GetParticleInvMasses(Solver);
+				for (int32 VertexIndex = 0; VertexIndex < Positions.Num(); ++VertexIndex)
+				{
+					const bool bIsKinematic0 = (InvMasses[VertexIndex] == (Softs::FSolverReal)0.);
+					if (bIsKinematic0 || !ShouldDrawParticle(VertexIndex))
+					{
+						continue;
+					}
+
+					const FVec3 Pos0 = GetWorldPosition(*Solver, Positions[VertexIndex]);
+					const FReal Radius1 = (FReal)CollisionConstraints->GetParticleInnerThickness(VertexIndex + Offset);
+					DrawSphere(Context.PDI, Pos0, FQuat::Identity, Radius1, FLinearColor::Gray, nullptr);
+				}
+			}
+		}
+	}
+
 	void FClothVisualizationNoGC::DrawKinematicColliderWired(const FDrawContext& Context) const
 	{
 		if (!Solver)
diff --git a/Engine/Plugins/ChaosCloth/Source/ChaosCloth/Private/ChaosCloth/ChaosClothingSimulation.cpp b/Engine/Plugins/ChaosCloth/Source/ChaosCloth/Private/ChaosCloth/ChaosClothingSimulation.cpp
--- a/Engine/Plugins/ChaosCloth/Source/ChaosCloth/Private/ChaosCloth/ChaosClothingSimulation.cpp
+++ b/Engine/Plugins/ChaosCloth/Source/ChaosCloth/Private/ChaosCloth/ChaosClothingSimulation.cpp
@@ -98,6 +98,9 @@ namespace ClothingSimulationCVar
 	TAutoConsoleVariable<bool> DebugDrawWindVelocity        (TEXT("p.ChaosCloth.DebugDrawWindVelocity"        ), false, TEXT("Whether to debug draw the Chaos Cloth wind velocity"), ECVF_Cheat);
 	TAutoConsoleVariable<bool> DebugDrawSelfCollision       (TEXT("p.ChaosCloth.DebugDrawSelfCollision"       ), false, TEXT("Whether to debug draw the Chaos Cloth self collision information"), ECVF_Cheat);
 	TAutoConsoleVariable<bool> DebugDrawSelfIntersection    (TEXT("p.ChaosCloth.DebugDrawSelfIntersection"    ), false, TEXT("Whether to debug draw the Chaos Cloth self intersection information"), ECVF_Cheat);
+	TAutoConsoleVariable<bool> DebugDrawSelfCollisionThickness(TEXT("p.ChaosCloth.DebugDrawSelfCollisionThickness"), false, TEXT("Whether to debug draw the Chaos Cloth self collision thickness"), ECVF_Cheat);
+	TAutoConsoleVariable<bool> DebugDrawClothCollisionThickness(TEXT("p.ChaosCloth.DebugDrawClothCollisionThickness"), false, TEXT("Whether to debug draw the Chaos Cloth cloth collision thickness"), ECVF_Cheat);
+	TAutoConsoleVariable<bool> DebugDrawInnerCollisionThickness(TEXT("p.ChaosCloth.DebugDrawInnerCollisionThickness"), false, TEXT("Whether to debug draw the Chaos Cloth inner collision thickness"), ECVF_Cheat);
 	TAutoConsoleVariable<bool> DebugDrawParticleIndices     (TEXT("p.ChaosCloth.DebugDrawParticleIndices"     ), false, TEXT("Whether to debug draw the Chaos Cloth particle indices"), ECVF_Cheat);
 	TAutoConsoleVariable<bool> DebugDrawElementIndices      (TEXT("p.ChaosCloth.DebugDrawElementIndices"      ), false, TEXT("Whether to debug draw the Chaos Cloth element indices"), ECVF_Cheat);
 	TAutoConsoleVariable<bool> DebugDrawClothClothConstraints    (TEXT("p.ChaosCloth.DebugDrawClothClothConstraints"    ), false, TEXT("Whether to debug draw the Chaos Cloth-Cloth constraints"), ECVF_Cheat);
diff --git a/Engine/Plugins/ChaosCloth/Source/ChaosCloth/Public/ChaosCloth/ChaosClothVisualization.h b/Engine/Plugins/ChaosCloth/Source/ChaosCloth/Public/ChaosCloth/ChaosClothVisualization.h
--- a/Engine/Plugins/ChaosCloth/Source/ChaosCloth/Public/ChaosCloth/ChaosClothVisualization.h
+++ b/Engine/Plugins/ChaosCloth/Source/ChaosCloth/Public/ChaosCloth/ChaosClothVisualization.h
@@ -84,6 +84,8 @@ namespace Chaos
 		CHAOSCLOTH_API void DrawSelfCollision(const FDrawContext& Context = FDrawContext()) const;
 		CHAOSCLOTH_API void DrawSelfIntersection(const FDrawContext& Context = FDrawContext()) const;
 		CHAOSCLOTH_API void DrawSelfCollisionThickness(const FDrawContext& Context = FDrawContext()) const;
+		CHAOSCLOTH_API void DrawClothCollisionThickness(const FDrawContext& Context = FDrawContext()) const;
+		CHAOSCLOTH_API void DrawInnerCollisionThickness(const FDrawContext& Context = FDrawContext()) const;
 		CHAOSCLOTH_API void DrawKinematicColliderWired(const FDrawContext& Context = FDrawContext()) const;
 		CHAOSCLOTH_API void DrawBounds(const FDrawContext& Context = FDrawContext()) const;
 		CHAOSCLOTH_API void DrawGravity(const FDrawContext& Context = FDrawContext()) const;
@@ -302,6 +304,8 @@ namespace Chaos
 		void DrawSelfCollision(const FDrawContext& /*Context*/ = FDrawContext()) const {}
 		void DrawSelfIntersection(const FDrawContext& /*Context*/ = FDrawContext()) const {}
 		void DrawSelfCollisionThickness(const FDrawContext& /*Context*/ = FDrawContext()) const {}
+		void DrawClothCollisionThickness(const FDrawContext & /*Context*/ = FDrawContext()) const {}
+		void DrawInnerCollisionThickness(const FDrawContext & /*Context*/ = FDrawContext()) const {}
 		void DrawKinematicColliderWired(const FDrawContext& /*Context*/ = FDrawContext()) const {}
 		void DrawBounds(const FDrawContext& /*Context*/ = FDrawContext()) const {}
 		void DrawGravity(const FDrawContext& /*Context*/ = FDrawContext()) const {}
diff --git a/Engine/Plugins/ChaosClothAsset/Source/ChaosClothAssetEngine/Private/ChaosClothAsset/ClothSimulationProxy.cpp b/Engine/Plugins/ChaosClothAsset/Source/ChaosClothAssetEngine/Private/ChaosClothAsset/ClothSimulationProxy.cpp
--- a/Engine/Plugins/ChaosClothAsset/Source/ChaosClothAssetEngine/Private/ChaosClothAsset/ClothSimulationProxy.cpp
+++ b/Engine/Plugins/ChaosClothAsset/Source/ChaosClothAssetEngine/Private/ChaosClothAsset/ClothSimulationProxy.cpp
@@ -695,6 +695,9 @@ namespace UE::Chaos::ClothAsset
 			static const TConsoleVariableData<bool>* const DebugDrawWindVelocityCVar = IConsoleManager::Get().FindTConsoleVariableDataBool(TEXT("p.ChaosCloth.DebugDrawWindVelocity"));
 			static const TConsoleVariableData<bool>* const DebugDrawSelfCollisionCVar = IConsoleManager::Get().FindTConsoleVariableDataBool(TEXT("p.ChaosCloth.DebugDrawSelfCollision"));
 			static const TConsoleVariableData<bool>* const DebugDrawSelfIntersectionCVar = IConsoleManager::Get().FindTConsoleVariableDataBool(TEXT("p.ChaosCloth.DebugDrawSelfIntersection"));
+			static const TConsoleVariableData<bool>* const DebugDrawSelfCollisionThicknessCVar = IConsoleManager::Get().FindTConsoleVariableDataBool(TEXT("p.ChaosCloth.DebugDrawSelfCollisionThickness"));
+			static const TConsoleVariableData<bool>* const DebugDrawClothCollisionThicknessCVar = IConsoleManager::Get().FindTConsoleVariableDataBool(TEXT("p.ChaosCloth.DebugDrawClothCollisionThickness"));
+			static const TConsoleVariableData<bool>* const DebugDrawInnerCollisionThicknessCVar = IConsoleManager::Get().FindTConsoleVariableDataBool(TEXT("p.ChaosCloth.DebugDrawInnerCollisionThickness"));
 			static const TConsoleVariableData<bool>* const DebugDrawParticleIndicesCVar = IConsoleManager::Get().FindTConsoleVariableDataBool(TEXT("p.ChaosCloth.DebugDrawParticleIndices"));
 			static const TConsoleVariableData<bool>* const DebugDrawElementIndicesCVar = IConsoleManager::Get().FindTConsoleVariableDataBool(TEXT("p.ChaosCloth.DebugDrawElementIndices"));
 			static const TConsoleVariableData<bool>* const DebugDrawClothClothConstraintsCVar = IConsoleManager::Get().FindTConsoleVariableDataBool(TEXT("p.ChaosCloth.DebugDrawClothClothConstraints"));
@@ -731,6 +734,9 @@ namespace UE::Chaos::ClothAsset
 			if (DebugDrawWindVelocityCVar && DebugDrawWindVelocityCVar->GetValueOnAnyThread()) { Visualization->DrawWindVelocity(DrawContext); }
 			if (DebugDrawSelfCollisionCVar && DebugDrawSelfCollisionCVar->GetValueOnAnyThread()) { Visualization->DrawSelfCollision(DrawContext); }
 			if (DebugDrawSelfIntersectionCVar && DebugDrawSelfIntersectionCVar->GetValueOnAnyThread()) { Visualization->DrawSelfIntersection(DrawContext); }
+			if (DebugDrawSelfCollisionThicknessCVar && DebugDrawSelfCollisionThicknessCVar->GetValueOnAnyThread()) { Visualization->DrawSelfCollisionThickness(DrawContext); }
+			if (DebugDrawClothCollisionThicknessCVar && DebugDrawClothCollisionThicknessCVar->GetValueOnAnyThread()) { Visualization->DrawClothCollisionThickness(DrawContext); }
+			if (DebugDrawInnerCollisionThicknessCVar && DebugDrawInnerCollisionThicknessCVar->GetValueOnAnyThread()) { Visualization->DrawInnerCollisionThickness(DrawContext); }
 			if (DebugDrawParticleIndicesCVar && DebugDrawParticleIndicesCVar->GetValueOnAnyThread()) { Visualization->DrawParticleIndices(DrawTextsContext); }
 			if (DebugDrawElementIndicesCVar && DebugDrawElementIndicesCVar->GetValueOnAnyThread()) { Visualization->DrawElementIndices(DrawTextsContext); }
 			if (DebugDrawClothClothConstraintsCVar && DebugDrawClothClothConstraintsCVar->GetValueOnAnyThread()) { Visualization->DrawClothClothConstraints(DrawContext); }
diff --git a/Engine/Plugins/ChaosClothAssetEditorCore/Source/ChaosClothAssetEditor/Private/ChaosClothAsset/ClothEditorSimulationVisualization.cpp b/Engine/Plugins/ChaosClothAssetEditorCore/Source/ChaosClothAssetEditor/Private/ChaosClothAsset/ClothEditorSimulationVisualization.cpp
--- a/Engine/Plugins/ChaosClothAssetEditorCore/Source/ChaosClothAssetEditor/Private/ChaosClothAsset/ClothEditorSimulationVisualization.cpp
+++ b/Engine/Plugins/ChaosClothAssetEditorCore/Source/ChaosClothAssetEditor/Private/ChaosClothAsset/ClothEditorSimulationVisualization.cpp
@@ -436,6 +436,20 @@ const FVisualizationOption FVisualizationOption::OptionData[] =
 		}),
 		LOCTEXT("ChaosVisName_SelfCollisionThickness", "Self Collision Thickness"), 
 		LOCTEXT("ChaosVisName_SelfCollisionThickness_Tooltip", "Draws the self collision Thickness")),
+	FVisualizationOption(
+		FClothVisualizationDebugDraw::CreateLambda([](const FClothEditorSimulationVisualization&, const ::Chaos::FClothVisualizationNoGC& Visualization, const ::Chaos::FClothVisualizationNoGC::FDrawContext& Context)
+		{
+			Visualization.DrawClothCollisionThickness(Context);
+		}),
+		LOCTEXT("ChaosVisName_ClothCollisionThickness", "Cloth Collision Thickness"), 
+		LOCTEXT("ChaosVisName_ClothCollisionThickness_Tooltip", "Draws the cloth collision Thickness")),
+	FVisualizationOption(
+		FClothVisualizationDebugDraw::CreateLambda([](const FClothEditorSimulationVisualization&, const ::Chaos::FClothVisualizationNoGC& Visualization, const ::Chaos::FClothVisualizationNoGC::FDrawContext& Context)
+		{
+			Visualization.DrawInnerCollisionThickness(Context);
+		}),
+		LOCTEXT("ChaosVisName_InnerCollisionThickness", "Inner Collision Thickness"), 
+		LOCTEXT("ChaosVisName_InnerCollisionThickness_Tooltip", "Draws the inner collision Thickness")),
 	FVisualizationOption(
 		FClothVisualizationDebugDraw::CreateLambda([](const FClothEditorSimulationVisualization&, const ::Chaos::FClothVisualizationNoGC& Visualization, const ::Chaos::FClothVisualizationNoGC::FDrawContext& Context)
 		{
diff --git a/Engine/Source/Runtime/Experimental/Chaos/Private/Chaos/PBDSoftBodyCollisionConstraint.cpp b/Engine/Source/Runtime/Experimental/Chaos/Private/Chaos/PBDSoftBodyCollisionConstraint.cpp
--- a/Engine/Source/Runtime/Experimental/Chaos/Private/Chaos/PBDSoftBodyCollisionConstraint.cpp
+++ b/Engine/Source/Runtime/Experimental/Chaos/Private/Chaos/PBDSoftBodyCollisionConstraint.cpp
@@ -131,7 +131,7 @@ extern "C" void ApplySimpleCCDNoFriction(
 	const int32 CollisionIndex,
 	const FSolverRigidTransform3& LastCollisionTransform,
 	const FSolverReal* XSoa,
-	const FSolverReal Thickness,
+	const FSolverReal* Thickness,
 	FSolverReal* PSoa,
 	const int32 ProgramCount,
 	const int32 Mask)
@@ -157,7 +157,7 @@ extern "C" void ApplySimpleCCDNoFriction(
 				CollisionIndex,
 				LastCollisionTransform,
 				X,
-				Thickness,
+				Thickness[Index],
 				(FSolverReal)0., // No friction, no FrictionCoefficient needed
 				(FSolverReal)0., // No friction, no Dt needed
 				P,
@@ -176,7 +176,7 @@ extern "C" void ApplySimpleCCDWithFriction(
 	const int32 CollisionIndex,
 	const FSolverRigidTransform3& LastCollisionTransform,
 	const FSolverReal* XSoa,
-	const FSolverReal Thickness,
+	const FSolverReal* Thickness,
 	const FSolverReal FrictionCoefficient,
 	const FSolverReal Dt,
 	FSolverReal* PSoa,
@@ -204,7 +204,7 @@ extern "C" void ApplySimpleCCDWithFriction(
 				CollisionIndex,
 				LastCollisionTransform,
 				X,
-				Thickness,
+				Thickness[Index],
 				FrictionCoefficient,
 				Dt,
 				P,
@@ -583,7 +583,7 @@ void FPBDSoftBodyCollisionConstraintBase::ApplySimpleInternal(FSolverParticlesRa
 					FVec3 ImplicitNormalReal;
 					const FSolverReal Phi = (FSolverReal)Geometry->PhiWithNormal(RigidSpacePosition, ImplicitNormalReal);
 					FSolverVec3 ImplicitNormal(ImplicitNormalReal);
-					const FSolverReal Thickness = SoftBodyCollisionThickness + CollisionThickness;
+					const FSolverReal Thickness = SoftBodyCollisionThicknessWeighted.GetValue(Index) + CollisionThickness;
 					FSolverReal Penetration = Thickness - Phi;
 					if (Penetration > (FSolverReal)0.)
 					{
@@ -708,7 +708,7 @@ void FPBDSoftBodyCollisionConstraintBase::ApplyComplexInternal(FSolverParticlesR
 				{
 					if (const TWeightedLatticeImplicitObject<FLevelSet>* const LevelSet = Geometry->GetObject<TWeightedLatticeImplicitObject<FLevelSet>>())
 					{
-						LevelSet->BatchPhiWithNormalAndGreatestInfluenceBone(Particles.GetPAndInvM(), Frame, SoftBodyCollisionThickness + CollisionThickness, BatchPhis, BatchNormals, BatchVelocityBones);
+						LevelSet->BatchPhiWithNormalAndGreatestInfluenceBone(Particles.GetPAndInvM(), Frame, FMath::Max(SoftBodyCollisionThicknessWeighted.GetLow(), SoftBodyCollisionThicknessWeighted.GetHigh()) + CollisionThickness, BatchPhis, BatchNormals, BatchVelocityBones);
 					}
 					else
 					{
@@ -723,7 +723,7 @@ void FPBDSoftBodyCollisionConstraintBase::ApplyComplexInternal(FSolverParticlesR
 								return;
 							}
 
-							const FSolverReal Penetration = SoftBodyCollisionThickness + CollisionThickness - BatchPhis[Index];
+							const FSolverReal Penetration = SoftBodyCollisionThicknessWeighted.GetValue(Index) + CollisionThickness - BatchPhis[Index];
 							if (Penetration > (FSolverReal)0.)
 							{
 								const FSolverVec3 NormalWorld = Frame.TransformVector(BatchNormals[Index]);
@@ -775,7 +775,7 @@ void FPBDSoftBodyCollisionConstraintBase::ApplyComplexInternal(FSolverParticlesR
 
 						// Batch Query
 						constexpr int32 MLLevelSetThreadNum = 0;
-						MLLevelSet->BatchPhiWithNormal(Particles.GetPAndInvM(), Frame, BatchPhis, BatchNormals, SoftBodyCollisionThickness + CollisionThickness, MLLevelSetThreadNum, 0, Particles.GetRangeSize());
+						MLLevelSet->BatchPhiWithNormal(Particles.GetPAndInvM(), Frame, BatchPhis, BatchNormals, FMath::Max(SoftBodyCollisionThicknessWeighted.GetLow(), SoftBodyCollisionThicknessWeighted.GetHigh()) + CollisionThickness, MLLevelSetThreadNum, 0, Particles.GetRangeSize());
 					}
 					else if (const TWeightedLatticeImplicitObject<FLevelSet>* const LevelSet = Geometry->GetObject<TWeightedLatticeImplicitObject<FLevelSet>>())
 					{
@@ -794,7 +794,7 @@ void FPBDSoftBodyCollisionConstraintBase::ApplyComplexInternal(FSolverParticlesR
 								return;
 							}
 
-							const FSolverReal Penetration = SoftBodyCollisionThickness + CollisionThickness - BatchPhis[Index];
+							const FSolverReal Penetration = SoftBodyCollisionThicknessWeighted.GetValue(Index) + CollisionThickness - BatchPhis[Index];
 							if (Penetration > (FSolverReal)0.)
 							{
 								const FSolverVec3 NormalWorld = Frame.TransformVector(BatchNormals[Index]);
@@ -838,7 +838,7 @@ void FPBDSoftBodyCollisionConstraintBase::ApplyComplexInternal(FSolverParticlesR
 
 					// Batch Query
 					constexpr int32 MLLevelSetThreadNum = 0;
-					MLLevelSet->BatchPhiWithNormal(Particles.GetPAndInvM(), Frame, BatchPhis, BatchNormals, SoftBodyCollisionThickness + CollisionThickness, MLLevelSetThreadNum, 0, Particles.GetRangeSize());
+					MLLevelSet->BatchPhiWithNormal(Particles.GetPAndInvM(), Frame, BatchPhis, BatchNormals, FMath::Max(SoftBodyCollisionThicknessWeighted.GetLow(), SoftBodyCollisionThicknessWeighted.GetHigh()) + CollisionThickness, MLLevelSetThreadNum, 0, Particles.GetRangeSize());
 				}
 				else if (const TWeightedLatticeImplicitObject<FLevelSet>* const LevelSet = Geometry->GetObject<TWeightedLatticeImplicitObject<FLevelSet>>())
 				{
@@ -857,7 +857,7 @@ void FPBDSoftBodyCollisionConstraintBase::ApplyComplexInternal(FSolverParticlesR
 							return;
 						}
 
-						const FSolverReal Penetration = SoftBodyCollisionThickness + CollisionThickness - BatchPhis[Index];
+						const FSolverReal Penetration = SoftBodyCollisionThicknessWeighted.GetValue(Index) + CollisionThickness - BatchPhis[Index];
 						if (Penetration > (FSolverReal)0.)
 						{
 							const FSolverVec3 NormalWorld = Frame.TransformVector(BatchNormals[Index]);
@@ -975,7 +975,7 @@ void FPBDSoftBodyCollisionConstraintBase::ApplyInternalCCD(FSolverParticlesRange
 
 				const Pair<FVec3, bool> PointPair = CollisionParticlesRange.GetGeometry(CollisionIndex)->FindClosestIntersection(  // Geometry operates in FReal
 					FVec3(CollisionTransforms[CollisionIndex].InverseTransformPositionNoScale(X[Index])),        // hence the back and forth
-					FVec3(Frame.InverseTransformPositionNoScale(PAndInvM[Index].P)), (FReal)(SoftBodyCollisionThickness + CollisionThickness));                   // FVec3/FReal conversions
+					FVec3(Frame.InverseTransformPositionNoScale(PAndInvM[Index].P)), (FReal)(SoftBodyCollisionThicknessWeighted.GetValue(Index) + CollisionThickness));  // FVec3/FReal conversions
 
 				if (PointPair.Second)
 				{
@@ -1084,7 +1084,9 @@ void FPBDSoftBodyCollisionConstraintBase::ApplySimpleInternalISPC(FSolverParticl
 							(const ispc::FVector3f*)CollisionParticlesRange.GetW().GetData(),
 							(const ispc::FVector4f*)CollisionParticlesRange.GetR().GetData(),
 							FrictionCoefficient,
-							SoftBodyCollisionThickness + CollisionThickness,
+							CollisionThickness,
+							reinterpret_cast<const ispc::FVector2f&>(SoftBodyCollisionThicknessWeighted.GetOffsetRange()),
+							SoftBodyCollisionThicknessWeighted.GetMapValues().GetData(),
 							(const uint8*)&CollisionParticlesRange,
 							(const uint8*)CollisionParticlesRange.GetAllGeometry().GetData(),
 							sizeof(FImplicitObject),
@@ -1105,7 +1107,9 @@ void FPBDSoftBodyCollisionConstraintBase::ApplySimpleInternalISPC(FSolverParticl
 							(const ispc::FVector3f*)Particles.XArray().GetData(),
 							(const ispc::FVector3f*)CollisionParticlesRange.XArray().GetData(),
 							(const ispc::FVector4f*)CollisionParticlesRange.GetR().GetData(),
-							SoftBodyCollisionThickness + CollisionThickness,
+							CollisionThickness,
+							reinterpret_cast<const ispc::FVector2f&>(SoftBodyCollisionThicknessWeighted.GetOffsetRange()),
+							SoftBodyCollisionThicknessWeighted.GetMapValues().GetData(),
 							(const uint8*)&CollisionParticlesRange,
 							(const uint8*)CollisionParticlesRange.GetAllGeometry().GetData(),
 							sizeof(FImplicitObject),
@@ -1133,7 +1137,9 @@ void FPBDSoftBodyCollisionConstraintBase::ApplySimpleInternalISPC(FSolverParticl
 							(const ispc::FVector4f*)CollisionParticlesRange.GetR().GetData(),
 							(const ispc::FTransform3f*)CollisionTransforms.GetData(),
 							FrictionCoefficient,
-							SoftBodyCollisionThickness + CollisionThickness,
+							CollisionThickness,
+							reinterpret_cast<const ispc::FVector2f&>(SoftBodyCollisionThicknessWeighted.GetOffsetRange()),
+							SoftBodyCollisionThicknessWeighted.GetMapValues().GetData(),
 							bHasInnerCollisionThickness ? InnerCollisionThickness.GetTable().GetData() : nullptr,
 							bHasInnerCollisionThicknessWeightMap ? InnerCollisionThickness.GetIndices().GetData() : nullptr,
 							(const uint8*)&CollisionParticlesRange,
@@ -1154,7 +1160,9 @@ void FPBDSoftBodyCollisionConstraintBase::ApplySimpleInternalISPC(FSolverParticl
 							(const ispc::FVector3f*)CollisionParticlesRange.XArray().GetData(),
 							(const ispc::FVector4f*)CollisionParticlesRange.GetR().GetData(),
 							(const ispc::FTransform3f*)CollisionTransforms.GetData(),
-							SoftBodyCollisionThickness + CollisionThickness,
+							CollisionThickness,
+							reinterpret_cast<const ispc::FVector2f&>(SoftBodyCollisionThicknessWeighted.GetOffsetRange()),
+							SoftBodyCollisionThicknessWeighted.GetMapValues().GetData(),
 							bHasInnerCollisionThickness ? InnerCollisionThickness.GetTable().GetData() : nullptr,
 							bHasInnerCollisionThicknessWeightMap ? InnerCollisionThickness.GetIndices().GetData() : nullptr,
 							(const uint8*)&CollisionParticlesRange,
@@ -1207,7 +1215,7 @@ void FPBDSoftBodyCollisionConstraintBase::ApplyComplexInternalISPC(FSolverPartic
 
 					// Batch Query
 					constexpr int32 MLLevelSetThreadNum = 0;
-					MLLevelSet->BatchPhiWithNormal(Particles.GetPAndInvM(), Frame, BatchPhis, BatchNormals, SoftBodyCollisionThickness + CollisionThickness, MLLevelSetThreadNum, 0, Particles.GetRangeSize());
+					MLLevelSet->BatchPhiWithNormal(Particles.GetPAndInvM(), Frame, BatchPhis, BatchNormals, FMath::Max(SoftBodyCollisionThicknessWeighted.GetLow(), SoftBodyCollisionThicknessWeighted.GetHigh()) + CollisionThickness, MLLevelSetThreadNum, 0, Particles.GetRangeSize());
 
 					// Apply
 					if (bGeneratePlanarConstraints)
@@ -1226,7 +1234,9 @@ void FPBDSoftBodyCollisionConstraintBase::ApplyComplexInternalISPC(FSolverPartic
 							reinterpret_cast<const ispc::FVector3f&>(CollisionW),
 							reinterpret_cast<const ispc::FVector4f&>(CollisionR),
 							FrictionCoefficient,
-							SoftBodyCollisionThickness + CollisionThickness,
+							CollisionThickness,
+							reinterpret_cast<const ispc::FVector2f&>(SoftBodyCollisionThicknessWeighted.GetOffsetRange()),
+							SoftBodyCollisionThicknessWeighted.GetMapValues().GetData(),
 							Dt,
 							0,
 							Particles.GetRangeSize());
@@ -1243,7 +1253,9 @@ void FPBDSoftBodyCollisionConstraintBase::ApplyComplexInternalISPC(FSolverPartic
 							reinterpret_cast<const ispc::FVector3f&>(CollisionW),
 							reinterpret_cast<const ispc::FVector4f&>(CollisionR),
 							FrictionCoefficient,
-							SoftBodyCollisionThickness + CollisionThickness,
+							CollisionThickness,
+							reinterpret_cast<const ispc::FVector2f&>(SoftBodyCollisionThicknessWeighted.GetOffsetRange()),
+							SoftBodyCollisionThicknessWeighted.GetMapValues().GetData(),
 							Dt,
 							0,
 							Particles.GetRangeSize());
@@ -1253,7 +1265,7 @@ void FPBDSoftBodyCollisionConstraintBase::ApplyComplexInternalISPC(FSolverPartic
 				{
 					if (const FPBDComplexColliderBoneData* const ColliderBoneData = ComplexBoneData.Find(FParticleRangeIndex(CollisionParticlesRange.GetRangeId(), CollisionIndex)))
 					{
-						LevelSet->BatchPhiWithNormalAndGreatestInfluenceBone(Particles.GetPAndInvM(), Frame, SoftBodyCollisionThickness + CollisionThickness, BatchPhis, BatchNormals, BatchVelocityBones);
+						LevelSet->BatchPhiWithNormalAndGreatestInfluenceBone(Particles.GetPAndInvM(), Frame, FMath::Max(SoftBodyCollisionThicknessWeighted.GetLow(), SoftBodyCollisionThicknessWeighted.GetHigh()) + CollisionThickness, BatchPhis, BatchNormals, BatchVelocityBones);
 
 						if (bGeneratePlanarConstraints)
 						{
@@ -1277,7 +1289,9 @@ void FPBDSoftBodyCollisionConstraintBase::ApplyComplexInternalISPC(FSolverPartic
 								(const ispc::FVector3f*)ColliderBoneData->X.GetData(),
 								(const ispc::FVector3f*)ColliderBoneData->W.GetData(),
 								FrictionCoefficient,
-								SoftBodyCollisionThickness + CollisionThickness,
+								CollisionThickness,
+								reinterpret_cast<const ispc::FVector2f&>(SoftBodyCollisionThicknessWeighted.GetOffsetRange()),
+								SoftBodyCollisionThicknessWeighted.GetMapValues().GetData(),
 								Dt,
 								0,
 								Particles.GetRangeSize());
@@ -1300,7 +1314,9 @@ void FPBDSoftBodyCollisionConstraintBase::ApplyComplexInternalISPC(FSolverPartic
 								(const ispc::FVector3f*)ColliderBoneData->X.GetData(),
 								(const ispc::FVector3f*)ColliderBoneData->W.GetData(),
 								FrictionCoefficient,
-								SoftBodyCollisionThickness + CollisionThickness,
+								CollisionThickness,
+								reinterpret_cast<const ispc::FVector2f&>(SoftBodyCollisionThicknessWeighted.GetOffsetRange()),
+								SoftBodyCollisionThicknessWeighted.GetMapValues().GetData(),
 								Dt,
 								0,
 								Particles.GetRangeSize());
@@ -1325,7 +1341,9 @@ void FPBDSoftBodyCollisionConstraintBase::ApplyComplexInternalISPC(FSolverPartic
 								reinterpret_cast<const ispc::FVector3f&>(CollisionW),
 								reinterpret_cast<const ispc::FVector4f&>(CollisionR),
 								FrictionCoefficient,
-								SoftBodyCollisionThickness + CollisionThickness,
+								CollisionThickness,
+								reinterpret_cast<const ispc::FVector2f&>(SoftBodyCollisionThicknessWeighted.GetOffsetRange()),
+								SoftBodyCollisionThicknessWeighted.GetMapValues().GetData(),
 								Dt,
 								0,
 								Particles.GetRangeSize());
@@ -1342,7 +1360,9 @@ void FPBDSoftBodyCollisionConstraintBase::ApplyComplexInternalISPC(FSolverPartic
 								reinterpret_cast<const ispc::FVector3f&>(CollisionW),
 								reinterpret_cast<const ispc::FVector4f&>(CollisionR),
 								FrictionCoefficient,
-								SoftBodyCollisionThickness + CollisionThickness,
+								CollisionThickness,
+								reinterpret_cast<const ispc::FVector2f&>(SoftBodyCollisionThicknessWeighted.GetOffsetRange()),
+								SoftBodyCollisionThicknessWeighted.GetMapValues().GetData(),
 								Dt,
 								0,
 								Particles.GetRangeSize());
@@ -1363,7 +1383,7 @@ void FPBDSoftBodyCollisionConstraintBase::ApplyComplexInternalISPC(FSolverPartic
 
 					// Batch Query
 					constexpr int32 MLLevelSetThreadNum = 0;
-					MLLevelSet->BatchPhiWithNormal(Particles.GetPAndInvM(), Frame, BatchPhis, BatchNormals, SoftBodyCollisionThickness + CollisionThickness, MLLevelSetThreadNum, 0, Particles.GetRangeSize());
+					MLLevelSet->BatchPhiWithNormal(Particles.GetPAndInvM(), Frame, BatchPhis, BatchNormals, FMath::Max(SoftBodyCollisionThicknessWeighted.GetLow(), SoftBodyCollisionThicknessWeighted.GetHigh()) + CollisionThickness, MLLevelSetThreadNum, 0, Particles.GetRangeSize());
 				}
 				else if (const TWeightedLatticeImplicitObject<FLevelSet>* const LevelSet = Geometry->GetObject<TWeightedLatticeImplicitObject<FLevelSet>>())
 				{
@@ -1386,7 +1406,9 @@ void FPBDSoftBodyCollisionConstraintBase::ApplyComplexInternalISPC(FSolverPartic
 						BatchPhis.GetData(),
 						(const ispc::FVector3f*)BatchNormals.GetData(),
 						reinterpret_cast<const ispc::FVector4f&>(CollisionR),
-						SoftBodyCollisionThickness + CollisionThickness,
+						CollisionThickness,
+						reinterpret_cast<const ispc::FVector2f&>(SoftBodyCollisionThicknessWeighted.GetOffsetRange()),
+						SoftBodyCollisionThicknessWeighted.GetMapValues().GetData(),
 						0,
 						Particles.GetRangeSize());
 				}
@@ -1398,7 +1420,9 @@ void FPBDSoftBodyCollisionConstraintBase::ApplyComplexInternalISPC(FSolverPartic
 						BatchPhis.GetData(),
 						(const ispc::FVector3f*)BatchNormals.GetData(),
 						reinterpret_cast<const ispc::FVector4f&>(CollisionR),
-						SoftBodyCollisionThickness + CollisionThickness,
+						CollisionThickness,
+						reinterpret_cast<const ispc::FVector2f&>(SoftBodyCollisionThicknessWeighted.GetOffsetRange()),
+						SoftBodyCollisionThicknessWeighted.GetMapValues().GetData(),
 						0,
 						Particles.GetRangeSize());
 				}
@@ -1441,7 +1465,7 @@ void FPBDSoftBodyCollisionConstraintBase::UpdateLinearSystem(const FSolverPartic
 				const FVec3 RigidSpacePosition(Frame.InverseTransformPosition(PAndInvM[Index].P));  // PhiWithNormal requires FReal based arguments
 				FVec3 ImplicitNormal;                                                                // since implicits don't use FSolverReal
 				const FSolverReal Phi = (FSolverReal)CollisionParticlesRange.GetGeometry(CollisionIndex)->PhiWithNormal(RigidSpacePosition, ImplicitNormal);
-				const FSolverReal Penetration = SoftBodyCollisionThickness + CollisionThickness - Phi; // This is related to the Normal impulse
+				const FSolverReal Penetration = SoftBodyCollisionThicknessWeighted.GetValue(Index) + CollisionThickness - Phi; // This is related to the Normal impulse
 				const FSolverVec3 Normal(ImplicitNormal);
 
 				if (Penetration > (FSolverReal)0.)
@@ -1486,7 +1510,8 @@ FPBDSoftBodyCollisionConstraint::FPBDSoftBodyCollisionConstraint(
 		GetUseCCD(PropertyCollection, false),
 		GetProximityStiffness(PropertyCollection, DefaultProximityStiffness),
 		InCollisionParticleCollided, InContacts, InNormals, InPhis,
-		GetSoftBodyCollisionThickness(PropertyCollection, DefaultSoftBodyCollisionThickness),
+		GetWeightedFloatClothCollisionThickness(PropertyCollection,
+			FSolverVec2(DefaultSoftBodyCollisionThickness, DefaultSoftBodyCollisionThickness)),
 		GetEnableSimpleColliders(PropertyCollection, true),
 		GetEnableComplexColliders(PropertyCollection, true),
 		GetUsePlanarConstraintForSimpleColliders(PropertyCollection, bDefaultUsePlanarConstraintForSimpleColliders),
@@ -1494,7 +1519,7 @@ FPBDSoftBodyCollisionConstraint::FPBDSoftBodyCollisionConstraint(
 		InComplexBoneData)
 	, MeshScale(InMeshScale)
 	, CollisionThicknessIndex(PropertyCollection)
-	, SoftBodyCollisionThicknessIndex(PropertyCollection)
+	, ClothCollisionThicknessIndex(PropertyCollection)
 	, FrictionCoefficientIndex(PropertyCollection)
 	, UseCCDIndex(PropertyCollection)
 	, ProximityStiffnessIndex(PropertyCollection)
@@ -1515,17 +1540,39 @@ FPBDSoftBodyCollisionConstraint::FPBDSoftBodyCollisionConstraint(
 			WeightMap.Num());
 		InnerCollisionThickness.ApplyValues();
 	}
+
+	// Initialize the base class SoftBodyCollisionThicknessWeighted map
+	{
+		const FSolverVec2 SoftBodyWeightedValue(GetWeightedFloatClothCollisionThickness(PropertyCollection,
+			FSolverVec2(DefaultSoftBodyCollisionThickness, DefaultSoftBodyCollisionThickness)));
+		const FString SoftBodyWeightMapString = GetClothCollisionThicknessString(PropertyCollection, ClothCollisionThicknessName.ToString());
+		const TConstArrayView<FRealSingle>& SoftBodyWeightMap = WeightMaps.FindRef(SoftBodyWeightMapString);
+		SoftBodyCollisionThicknessWeighted = FPBDFlatWeightMap(
+			SoftBodyWeightedValue, SoftBodyWeightMap, SoftBodyWeightMap.Num());
+	}
 }
 
-void FPBDSoftBodyCollisionConstraint::SetProperties(const FCollectionPropertyConstFacade& PropertyCollection)
+void FPBDSoftBodyCollisionConstraint::SetProperties(
+	const FCollectionPropertyConstFacade& PropertyCollection,
+	const TMap<FString, TConstArrayView<float>>& WeightMaps)
 {
 	if (IsCollisionThicknessMutable(PropertyCollection))
 	{
 		CollisionThickness = MeshScale * GetCollisionThickness(PropertyCollection);
 	}
-	if (IsSoftBodyCollisionThicknessMutable(PropertyCollection))
+	if (IsClothCollisionThicknessMutable(PropertyCollection))
 	{
-		SoftBodyCollisionThickness = GetSoftBodyCollisionThickness(PropertyCollection);
+		const FSolverVec2 WeightedValue(GetWeightedFloatClothCollisionThickness(PropertyCollection));
+		if (IsClothCollisionThicknessStringDirty(PropertyCollection))
+		{
+			const FString& WeightMapName = GetClothCollisionThicknessString(PropertyCollection);
+			const TConstArrayView<FRealSingle>& WeightMap = WeightMaps.FindRef(WeightMapName);
+			SoftBodyCollisionThicknessWeighted = FPBDFlatWeightMap(WeightedValue, WeightMap, WeightMap.Num());
+		}
+		else
+		{
+			SoftBodyCollisionThicknessWeighted.SetWeightedValue(WeightedValue);
+		}
 	}
 	if (IsFrictionCoefficientMutable(PropertyCollection))
 	{
diff --git a/Engine/Source/Runtime/Experimental/Chaos/Private/Chaos/PBDSoftBodyCollisionConstraint_Complex.ispc b/Engine/Source/Runtime/Experimental/Chaos/Private/Chaos/PBDSoftBodyCollisionConstraint_Complex.ispc
--- a/Engine/Source/Runtime/Experimental/Chaos/Private/Chaos/PBDSoftBodyCollisionConstraint_Complex.ispc
+++ b/Engine/Source/Runtime/Experimental/Chaos/Private/Chaos/PBDSoftBodyCollisionConstraint_Complex.ispc
@@ -11,7 +11,9 @@ export void ApplyPerParticleBatchCollisionFastFriction(uniform FVector4f Particl
 													const uniform FVector3f& CollisionW,
 													const uniform FVector4f& CollisionR,
 													const uniform float PerGroupFriction,
-													const uniform float PerGroupThickness,
+													const uniform float PerGroupCollisionThickness,
+													const uniform FVector2f& SoftBodyThicknessOffsetRange,
+													const uniform float SoftBodyThicknessMap[],
 													const uniform float Dt,
 													const uniform int OuterParticleBatchBegin,
 													const uniform int OuterParticleBatchEnd)
@@ -29,8 +31,14 @@ export void ApplyPerParticleBatchCollisionFastFriction(uniform FVector4f Particl
 		}
 
 		const float Phi = BatchPhis[Index];
+		
+		varying float CollisionThickness = PerGroupCollisionThickness + SoftBodyThicknessOffsetRange.V[0];
+		if (SoftBodyThicknessMap != NULL)
+		{
+			CollisionThickness += SoftBodyThicknessMap[Index] * SoftBodyThicknessOffsetRange.V[1];
+		}
 
-		const float Penetration = PerGroupThickness - Phi; // This is related to the Normal impulse
+		const float Penetration = CollisionThickness - Phi; // This is related to the Normal impulse
 		if(IM != 0 && Penetration > 0.0f)
 		{
 			const FVector3f Normal = VectorLoad(&BatchNormals[extract(Index,0)]);
@@ -62,7 +70,9 @@ export void ApplyPerParticleBatchCollisionFastFrictionWithVelocityBones(uniform
 													const uniform FVector3f ColliderBoneX[],
 													const uniform FVector3f ColliderBoneW[],
 													const uniform float PerGroupFriction,
-													const uniform float PerGroupThickness,
+													const uniform float PerGroupCollisionThickness,
+													const uniform FVector2f& SoftBodyThicknessOffsetRange,
+													const uniform float SoftBodyThicknessMap[],
 													const uniform float Dt,
 													const uniform int OuterParticleBatchBegin,
 													const uniform int OuterParticleBatchEnd)
@@ -80,8 +90,14 @@ export void ApplyPerParticleBatchCollisionFastFrictionWithVelocityBones(uniform
 		}
 
 		const float Phi = BatchPhis[Index];
+		
+		varying float CollisionThickness = PerGroupCollisionThickness + SoftBodyThicknessOffsetRange.V[0];
+		if (SoftBodyThicknessMap != NULL)
+		{
+			CollisionThickness += SoftBodyThicknessMap[Index] * SoftBodyThicknessOffsetRange.V[1];
+		}
 
-		const float Penetration = PerGroupThickness - Phi; // This is related to the Normal impulse
+		const float Penetration = CollisionThickness - Phi; // This is related to the Normal impulse
 		if(IM != 0 && Penetration > 0.0f)
 		{
 			const FVector3f Normal = VectorLoad(&BatchNormals[extract(Index,0)]);
@@ -126,7 +142,9 @@ export void ApplyPerParticleBatchCollisionFastFrictionAndGeneratePlanarConstrain
 													const uniform FVector3f& CollisionW,
 													const uniform FVector4f& CollisionR,
 													const uniform float PerGroupFriction,
-													const uniform float PerGroupThickness,
+													const uniform float PerGroupCollisionThickness,
+													const uniform FVector2f& SoftBodyThicknessOffsetRange,
+													const uniform float SoftBodyThicknessMap[],
 													const uniform float Dt,
 													const uniform int OuterParticleBatchBegin,
 													const uniform int OuterParticleBatchEnd)
@@ -144,8 +162,14 @@ export void ApplyPerParticleBatchCollisionFastFrictionAndGeneratePlanarConstrain
 		}
 
 		const float Phi = BatchPhis[Index];
+		
+		varying float CollisionThickness = PerGroupCollisionThickness + SoftBodyThicknessOffsetRange.V[0];
+		if (SoftBodyThicknessMap != NULL)
+		{
+			CollisionThickness += SoftBodyThicknessMap[Index] * SoftBodyThicknessOffsetRange.V[1];
+		}
 
-		const float Penetration = PerGroupThickness - Phi; // This is related to the Normal impulse
+		const float Penetration = CollisionThickness - Phi; // This is related to the Normal impulse
 		if(IM != 0 && Penetration > 0.0f)
 		{
 			const FVector3f Normal = VectorLoad(&BatchNormals[extract(Index,0)]);
@@ -188,7 +212,9 @@ export void ApplyPerParticleBatchCollisionFastFrictionWithVelocityBonesAndGenera
 													const uniform FVector3f ColliderBoneX[],
 													const uniform FVector3f ColliderBoneW[],
 													const uniform float PerGroupFriction,
-													const uniform float PerGroupThickness,
+													const uniform float PerGroupCollisionThickness,
+													const uniform FVector2f& SoftBodyThicknessOffsetRange,
+													const uniform float SoftBodyThicknessMap[],
 													const uniform float Dt,
 													const uniform int OuterParticleBatchBegin,
 													const uniform int OuterParticleBatchEnd)
@@ -206,8 +232,14 @@ export void ApplyPerParticleBatchCollisionFastFrictionWithVelocityBonesAndGenera
 		}
 
 		const float Phi = BatchPhis[Index];
+		
+		varying float CollisionThickness = PerGroupCollisionThickness + SoftBodyThicknessOffsetRange.V[0];
+		if (SoftBodyThicknessMap != NULL)
+		{
+			CollisionThickness += SoftBodyThicknessMap[Index] * SoftBodyThicknessOffsetRange.V[1];
+		}
 
-		const float Penetration = PerGroupThickness - Phi; // This is related to the Normal impulse
+		const float Penetration = CollisionThickness - Phi; // This is related to the Normal impulse
 		if(IM != 0 && Penetration > 0.0f)
 		{
 			const FVector3f Normal = VectorLoad(&BatchNormals[extract(Index,0)]);
@@ -251,7 +283,9 @@ export void ApplyPerParticleBatchCollisionNoFriction(uniform FVector4f Particles
 													const uniform float BatchPhis[],
 													const uniform FVector3f BatchNormals[],
 													const uniform FVector4f& CollisionR,
-													const uniform float PerGroupThickness,
+													const uniform float PerGroupCollisionThickness,
+													const uniform FVector2f& SoftBodyThicknessOffsetRange,
+													const uniform float SoftBodyThicknessMap[],
 													const uniform int OuterParticleBatchBegin,
 													const uniform int OuterParticleBatchEnd)
 {
@@ -268,8 +302,14 @@ export void ApplyPerParticleBatchCollisionNoFriction(uniform FVector4f Particles
 		}
 
 		const float Phi = BatchPhis[Index];
+		
+		varying float CollisionThickness = PerGroupCollisionThickness + SoftBodyThicknessOffsetRange.V[0];
+		if (SoftBodyThicknessMap != NULL)
+		{
+			CollisionThickness += SoftBodyThicknessMap[Index] * SoftBodyThicknessOffsetRange.V[1];
+		}
 
-		const float Penetration = PerGroupThickness - Phi; // This is related to the Normal impulse
+		const float Penetration = CollisionThickness - Phi; // This is related to the Normal impulse
 		if(IM != 0 && Penetration > 0.0f)
 		{
 			const FVector3f Normal = VectorLoad(&BatchNormals[extract(Index,0)]);
@@ -291,7 +331,9 @@ export void ApplyPerParticleBatchCollisionNoFrictionAndGeneratePlanarConstraints
 													const uniform float BatchPhis[],
 													const uniform FVector3f BatchNormals[],
 													const uniform FVector4f& CollisionR,
-													const uniform float PerGroupThickness,
+													const uniform float PerGroupCollisionThickness,
+													const uniform FVector2f& SoftBodyThicknessOffsetRange,
+													const uniform float SoftBodyThicknessMap[],
 													const uniform int OuterParticleBatchBegin,
 													const uniform int OuterParticleBatchEnd)
 {
@@ -308,8 +350,14 @@ export void ApplyPerParticleBatchCollisionNoFrictionAndGeneratePlanarConstraints
 		}
 
 		const float Phi = BatchPhis[Index];
+		
+		varying float CollisionThickness = PerGroupCollisionThickness + SoftBodyThicknessOffsetRange.V[0];
+		if (SoftBodyThicknessMap != NULL)
+		{
+			CollisionThickness += SoftBodyThicknessMap[Index] * SoftBodyThicknessOffsetRange.V[1];
+		}
 
-		const float Penetration = PerGroupThickness - Phi; // This is related to the Normal impulse
+		const float Penetration = CollisionThickness - Phi; // This is related to the Normal impulse
 		if(IM != 0 && Penetration > 0.0f)
 		{
 			const FVector3f Normal = VectorLoad(&BatchNormals[extract(Index,0)]);
diff --git a/Engine/Source/Runtime/Experimental/Chaos/Private/Chaos/PBDSoftBodyCollisionConstraint_Simple.ispc b/Engine/Source/Runtime/Experimental/Chaos/Private/Chaos/PBDSoftBodyCollisionConstraint_Simple.ispc
--- a/Engine/Source/Runtime/Experimental/Chaos/Private/Chaos/PBDSoftBodyCollisionConstraint_Simple.ispc
+++ b/Engine/Source/Runtime/Experimental/Chaos/Private/Chaos/PBDSoftBodyCollisionConstraint_Simple.ispc
@@ -8,7 +8,7 @@ extern "C" unmasked void ApplySimpleCCDNoFriction(
 	const uniform int CollisionIndex,
 	const uniform FTransform3f& LastCollisionTransform,
 	const uniform float* uniform X,
-	const uniform float uniform Thickness,
+	const uniform float* uniform Thickness,
 	uniform float* uniform P,
 	const uniform int ProgramCount, 
 	const uniform int Mask);
@@ -18,7 +18,7 @@ extern "C" unmasked void ApplySimpleCCDWithFriction(
 	const uniform int CollisionIndex,
 	const uniform FTransform3f& LastCollisionTransform,
 	const uniform float* uniform X,
-	const uniform float uniform Thickness,
+	const uniform float* uniform Thickness,
 	const uniform float uniform FrictionCoefficient,
 	const uniform float Dt,
 	uniform float* uniform P,
@@ -35,7 +35,9 @@ export void ApplyPerParticleSimpleCollisionFastFriction(uniform FVector4f Partic
 													const uniform FVector4f CollisionR[],
 													const uniform FTransform3f LastCollisionTransforms[],
 													const uniform float PerGroupFriction,
-													const uniform float PerGroupThickness,
+													const uniform float PerGroupCollisionThickness,
+													const uniform FVector2f& SoftBodyThicknessOffsetRange,
+													const uniform float SoftBodyThicknessMap[],
 													const uniform float *uniform InnerThicknessTable,
 													const uniform uint8 *uniform InnerThicknessIndices,
 													const uniform uint8 *uniform CollisionParticlesRange,
@@ -66,6 +68,12 @@ export void ApplyPerParticleSimpleCollisionFastFriction(uniform FVector4f Partic
 					continue;
 				}
 
+				varying float CollisionThickness = PerGroupCollisionThickness + SoftBodyThicknessOffsetRange.V[0];
+				if (SoftBodyThicknessMap != NULL)
+				{
+					CollisionThickness += SoftBodyThicknessMap[Index] * SoftBodyThicknessOffsetRange.V[1];
+				}
+
 				varying uint8 InnerThicknessIndex = InnerThicknessIndices[Index];
 				#pragma ignore warning(perf)
 				varying float InnerThickness = InnerThicknessTable[InnerThicknessIndex];
@@ -95,14 +103,14 @@ export void ApplyPerParticleSimpleCollisionFastFriction(uniform FVector4f Partic
 						IMMask = lanemask();
 					}
 
-					if(!GetSimplePenetrationWithNormal(PhiWithNormalInputFloat, Type, Ptr, TypedPtr, SizeofFImplicitObject, OffsetofGeometryType, OffsetOfMargin, PerGroupThickness, Penetration, Normal))
+					if(!GetSimplePenetrationWithNormal(PhiWithNormalInputFloat, Type, Ptr, TypedPtr, SizeofFImplicitObject, OffsetofGeometryType, OffsetOfMargin, CollisionThickness, Penetration, Normal))
 					{
 						const uniform int Mask = IMMask & lanemask();
 						if(Mask != 0)  // Do not call c++ when there is no work to do
 						{
 							float Phi = 0.f;
 							GetPhiWithNormalCollisionParticleRange(CollisionParticlesRange, (const uniform float *uniform)&PhiWithNormalInputFloat, (uniform float *uniform)&Normal, (uniform float *uniform)&Phi, i, programCount, Mask);
-							Penetration = PerGroupThickness - Phi;
+							Penetration = CollisionThickness - Phi;
 						}
 					}
 
@@ -118,7 +126,7 @@ export void ApplyPerParticleSimpleCollisionFastFriction(uniform FVector4f Partic
 								i,
 								LastCollisionTransforms[i],
 								(uniform float* uniform)&X,
-								PerGroupThickness,
+								(uniform float* uniform)&CollisionThickness,
 								PerGroupFriction,
 								Dt,
 								(uniform float *uniform)&P,
@@ -157,6 +165,12 @@ export void ApplyPerParticleSimpleCollisionFastFriction(uniform FVector4f Partic
 				{
 					continue;
 				}
+				
+				varying float CollisionThickness = PerGroupCollisionThickness + SoftBodyThicknessOffsetRange.V[0];
+				if (SoftBodyThicknessMap != NULL)
+				{
+					CollisionThickness += SoftBodyThicknessMap[Index] * SoftBodyThicknessOffsetRange.V[1];
+				}
 
 				for(uniform int i = 0; i < InnerCollisionCount; i++) // SequentialFor
 				{
@@ -182,15 +196,16 @@ export void ApplyPerParticleSimpleCollisionFastFriction(uniform FVector4f Partic
 					{
 						IMMask = lanemask();
 					}
+					
 
-					if(!GetSimplePenetrationWithNormal(PhiWithNormalInputFloat, Type, Ptr, TypedPtr, SizeofFImplicitObject, OffsetofGeometryType, OffsetOfMargin, PerGroupThickness, Penetration, Normal))
+					if(!GetSimplePenetrationWithNormal(PhiWithNormalInputFloat, Type, Ptr, TypedPtr, SizeofFImplicitObject, OffsetofGeometryType, OffsetOfMargin, CollisionThickness, Penetration, Normal))
 					{
 						const uniform int Mask = IMMask & lanemask();
 						if(Mask != 0)  // Do not call c++ when there is no work to do
 						{
 							float Phi = 0.f;
 							GetPhiWithNormalCollisionParticleRange(CollisionParticlesRange, (const uniform float *uniform)&PhiWithNormalInputFloat, (uniform float *uniform)&Normal, (uniform float *uniform)&Phi, i, programCount, Mask);
-							Penetration = PerGroupThickness - Phi;
+							Penetration = CollisionThickness - Phi;
 						}
 					}
 
@@ -206,7 +221,7 @@ export void ApplyPerParticleSimpleCollisionFastFriction(uniform FVector4f Partic
 								i,
 								LastCollisionTransforms[i],
 								(uniform float* uniform)&X,
-								PerGroupThickness,
+								(uniform float* uniform)&CollisionThickness,
 								PerGroupFriction,
 								Dt,
 								(uniform float *uniform)&P,
@@ -244,6 +259,12 @@ export void ApplyPerParticleSimpleCollisionFastFriction(uniform FVector4f Partic
 			{
 				continue;
 			}
+			
+			varying float CollisionThickness = PerGroupCollisionThickness + SoftBodyThicknessOffsetRange.V[0];
+			if (SoftBodyThicknessMap != NULL)
+			{
+				CollisionThickness += SoftBodyThicknessMap[Index] * SoftBodyThicknessOffsetRange.V[1];
+			}
 
 			for(uniform int i = 0; i < InnerCollisionCount; i++) // SequentialFor
 			{
@@ -264,7 +285,7 @@ export void ApplyPerParticleSimpleCollisionFastFriction(uniform FVector4f Partic
 				FVector3f Normal;
 				float Penetration;	
 
-				if(!GetSimplePenetrationWithNormal(PhiWithNormalInputFloat, Type, Ptr, TypedPtr, SizeofFImplicitObject, OffsetofGeometryType, OffsetOfMargin, PerGroupThickness, Penetration, Normal))
+				if(!GetSimplePenetrationWithNormal(PhiWithNormalInputFloat, Type, Ptr, TypedPtr, SizeofFImplicitObject, OffsetofGeometryType, OffsetOfMargin, CollisionThickness, Penetration, Normal))
 				{
 					uniform int IMMask = 0;
 					if(IM != 0)
@@ -275,7 +296,7 @@ export void ApplyPerParticleSimpleCollisionFastFriction(uniform FVector4f Partic
 					{
 						float Phi = 0.f;
 						GetPhiWithNormalCollisionParticleRange(CollisionParticlesRange, (const uniform float *uniform)&PhiWithNormalInputFloat, (uniform float *uniform)&Normal, (uniform float *uniform)&Phi, i, programCount, IMMask);
-						Penetration = PerGroupThickness - Phi;
+						Penetration = CollisionThickness - Phi;
 					}
 				}
 
@@ -309,7 +330,9 @@ export void ApplyPerParticleSimpleCollisionFastFrictionAndGeneratePlanarConstrai
 													const uniform FVector3f CollisionW[],
 													const uniform FVector4f CollisionR[],
 													const uniform float PerGroupFriction,
-													const uniform float PerGroupThickness,
+													const uniform float PerGroupCollisionThickness,
+													const uniform FVector2f& SoftBodyThicknessOffsetRange,
+													const uniform float SoftBodyThicknessMap[],
 													const uniform uint8 *uniform CollisionParticlesRange,
 													const uniform uint8 *uniform Geometry,
 													const uniform int SizeofFImplicitObject,
@@ -337,6 +360,12 @@ export void ApplyPerParticleSimpleCollisionFastFrictionAndGeneratePlanarConstrai
 			continue;
 		}
 
+		varying float CollisionThickness = PerGroupCollisionThickness + SoftBodyThicknessOffsetRange.V[0];
+		if (SoftBodyThicknessMap != NULL)
+		{
+			CollisionThickness += SoftBodyThicknessMap[Index] * SoftBodyThicknessOffsetRange.V[1];
+		}
+
 		for(uniform int i = 0; i < InnerCollisionCount; i++) // SequentialFor
 		{
 			const uniform SerializedPtr *uniform GeometryPtr = (const uniform SerializedPtr *uniform)Geometry;
@@ -356,7 +385,7 @@ export void ApplyPerParticleSimpleCollisionFastFrictionAndGeneratePlanarConstrai
 			FVector3f Normal;
 			float Penetration;
 			
-			if(!GetSimplePenetrationWithNormal(PhiWithNormalInputFloat, Type, Ptr, TypedPtr, SizeofFImplicitObject, OffsetofGeometryType, OffsetOfMargin, PerGroupThickness, Penetration, Normal))
+			if(!GetSimplePenetrationWithNormal(PhiWithNormalInputFloat, Type, Ptr, TypedPtr, SizeofFImplicitObject, OffsetofGeometryType, OffsetOfMargin, CollisionThickness, Penetration, Normal))
 			{
 				uniform int IMMask;
 				if(IM != 0)
@@ -366,7 +395,7 @@ export void ApplyPerParticleSimpleCollisionFastFrictionAndGeneratePlanarConstrai
 				
 				float Phi;
 				GetPhiWithNormalCollisionParticleRange(CollisionParticlesRange, (const uniform float *uniform)&PhiWithNormalInputFloat, (uniform float *uniform)&Normal, (uniform float *uniform)&Phi, i, programCount, IMMask);
-				Penetration = PerGroupThickness - Phi;
+				Penetration = CollisionThickness - Phi;
 			}
 
 			if(Penetration > 0.0f)
@@ -402,7 +431,9 @@ export void ApplyPerParticleSimpleCollisionNoFriction(uniform FVector4f Particle
 												const uniform FVector3f CollisionX[],
 												const uniform FVector4f CollisionR[],
 												const uniform FTransform3f LastCollisionTransforms[],
-												const uniform float PerGroupThickness,
+												const uniform float PerGroupCollisionThickness,
+												const uniform FVector2f& SoftBodyThicknessOffsetRange,
+												const uniform float SoftBodyThicknessMap[],
 												const uniform float *uniform InnerThicknessTable,
 												const uniform uint8 *uniform InnerThicknessIndices,
 												const uniform uint8 *uniform CollisionParticlesRange,
@@ -432,6 +463,12 @@ export void ApplyPerParticleSimpleCollisionNoFriction(uniform FVector4f Particle
 					continue;
 				}
 
+				varying float CollisionThickness = PerGroupCollisionThickness + SoftBodyThicknessOffsetRange.V[0];
+				if (SoftBodyThicknessMap != NULL)
+				{
+					CollisionThickness += SoftBodyThicknessMap[Index] * SoftBodyThicknessOffsetRange.V[1];
+				}
+
 				varying uint8 InnerThicknessIndex = InnerThicknessIndices[Index];
 				#pragma ignore warning(perf)
 				varying float InnerThickness = InnerThicknessTable[InnerThicknessIndex];
@@ -461,14 +498,14 @@ export void ApplyPerParticleSimpleCollisionNoFriction(uniform FVector4f Particle
 						IMMask = lanemask();
 					}
 			
-					if(!GetSimplePenetrationWithNormal(PhiWithNormalInputFloat, Type, Ptr, TypedPtr, SizeofFImplicitObject, OffsetofGeometryType, OffsetOfMargin, PerGroupThickness, Penetration, Normal))
+					if(!GetSimplePenetrationWithNormal(PhiWithNormalInputFloat, Type, Ptr, TypedPtr, SizeofFImplicitObject, OffsetofGeometryType, OffsetOfMargin, CollisionThickness, Penetration, Normal))
 					{
 						const uniform int Mask = IMMask & lanemask();
 						if(Mask != 0)  // Do not call c++ when there is no work to do
 						{
 							float Phi;
 							GetPhiWithNormalCollisionParticleRange(CollisionParticlesRange, (const uniform float *uniform)&PhiWithNormalInputFloat, (uniform float *uniform)&Normal, (uniform float *uniform)&Phi, i, programCount, Mask);
-							Penetration = PerGroupThickness - Phi;
+							Penetration = CollisionThickness - Phi;
 						}
 					}
 
@@ -483,7 +520,7 @@ export void ApplyPerParticleSimpleCollisionNoFriction(uniform FVector4f Particle
 								i,
 								LastCollisionTransforms[i],
 								(uniform float *uniform)&X,
-								PerGroupThickness,
+								(uniform float* uniform)&CollisionThickness,
 								(uniform float *uniform)&P,
 								programCount,
 								Mask);
@@ -518,6 +555,12 @@ export void ApplyPerParticleSimpleCollisionNoFriction(uniform FVector4f Particle
 				{
 					continue;
 				}
+				
+				varying float CollisionThickness = PerGroupCollisionThickness + SoftBodyThicknessOffsetRange.V[0];
+				if (SoftBodyThicknessMap != NULL)
+				{
+					CollisionThickness += SoftBodyThicknessMap[Index] * SoftBodyThicknessOffsetRange.V[1];
+				}
 
 				for(uniform int i = 0; i < InnerCollisionCount; i++) // SequentialFor
 				{
@@ -544,14 +587,14 @@ export void ApplyPerParticleSimpleCollisionNoFriction(uniform FVector4f Particle
 						IMMask = lanemask();
 					}
 			
-					if(!GetSimplePenetrationWithNormal(PhiWithNormalInputFloat, Type, Ptr, TypedPtr, SizeofFImplicitObject, OffsetofGeometryType, OffsetOfMargin, PerGroupThickness, Penetration, Normal))
+					if(!GetSimplePenetrationWithNormal(PhiWithNormalInputFloat, Type, Ptr, TypedPtr, SizeofFImplicitObject, OffsetofGeometryType, OffsetOfMargin, CollisionThickness, Penetration, Normal))
 					{
 						const uniform int Mask = IMMask & lanemask();
 						if(Mask != 0)  // Do not call c++ when there is no work to do
 						{
 							float Phi;
 							GetPhiWithNormalCollisionParticleRange(CollisionParticlesRange, (const uniform float *uniform)&PhiWithNormalInputFloat, (uniform float *uniform)&Normal, (uniform float *uniform)&Phi, i, programCount, Mask);
-							Penetration = PerGroupThickness - Phi;
+							Penetration = CollisionThickness - Phi;
 						}
 					}
 
@@ -566,7 +609,7 @@ export void ApplyPerParticleSimpleCollisionNoFriction(uniform FVector4f Particle
 								i,
 								LastCollisionTransforms[i],
 								(uniform float *uniform)&X,
-								PerGroupThickness,
+								(uniform float* uniform)&CollisionThickness,
 								(uniform float *uniform)&P,
 								programCount,
 								Mask);
@@ -600,6 +643,12 @@ export void ApplyPerParticleSimpleCollisionNoFriction(uniform FVector4f Particle
 			{
 				continue;
 			}
+			
+			varying float CollisionThickness = PerGroupCollisionThickness + SoftBodyThicknessOffsetRange.V[0];
+			if (SoftBodyThicknessMap != NULL)
+			{
+				CollisionThickness += SoftBodyThicknessMap[Index] * SoftBodyThicknessOffsetRange.V[1];
+			}
 
 			for(uniform int i = 0; i < InnerCollisionCount; i++) // SequentialFor
 			{
@@ -621,7 +670,7 @@ export void ApplyPerParticleSimpleCollisionNoFriction(uniform FVector4f Particle
 				float Penetration;
 
 			
-				if(!GetSimplePenetrationWithNormal(PhiWithNormalInputFloat, Type, Ptr, TypedPtr, SizeofFImplicitObject, OffsetofGeometryType, OffsetOfMargin, PerGroupThickness, Penetration, Normal))
+				if(!GetSimplePenetrationWithNormal(PhiWithNormalInputFloat, Type, Ptr, TypedPtr, SizeofFImplicitObject, OffsetofGeometryType, OffsetOfMargin, CollisionThickness, Penetration, Normal))
 				{
 					uniform int IMMask = 0;
 					if(IM != 0)
@@ -632,7 +681,7 @@ export void ApplyPerParticleSimpleCollisionNoFriction(uniform FVector4f Particle
 					{
 						float Phi;
 						GetPhiWithNormalCollisionParticleRange(CollisionParticlesRange, (const uniform float *uniform)&PhiWithNormalInputFloat, (uniform float *uniform)&Normal, (uniform float *uniform)&Phi, i, programCount, IMMask);
-						Penetration = PerGroupThickness - Phi;
+						Penetration = CollisionThickness - Phi;
 					}
 				}
 
@@ -659,7 +708,9 @@ export void ApplyPerParticleSimpleCollisionNoFrictionAndGeneratePlanarConstraint
 												const uniform FVector3f ParticlesX[],
 												const uniform FVector3f CollisionX[],
 												const uniform FVector4f CollisionR[],
-												const uniform float PerGroupThickness,
+												const uniform float PerGroupCollisionThickness,
+												const uniform FVector2f& SoftBodyThicknessOffsetRange,
+												const uniform float SoftBodyThicknessMap[],
 												const uniform uint8 *uniform CollisionParticlesRange,
 												const uniform uint8 *uniform Geometry,
 												const uniform int SizeofFImplicitObject,
@@ -685,6 +736,12 @@ export void ApplyPerParticleSimpleCollisionNoFrictionAndGeneratePlanarConstraint
 		{
 			continue;
 		}
+		
+		varying float CollisionThickness = PerGroupCollisionThickness + SoftBodyThicknessOffsetRange.V[0];
+		if (SoftBodyThicknessMap != NULL)
+		{
+			CollisionThickness += SoftBodyThicknessMap[Index] * SoftBodyThicknessOffsetRange.V[1];
+		}
 
 		for(uniform int i = 0; i < InnerCollisionCount; i++) // SequentialFor
 		{
@@ -705,7 +762,7 @@ export void ApplyPerParticleSimpleCollisionNoFrictionAndGeneratePlanarConstraint
 			FVector3f Normal;
 			float Penetration;
 			
-			if(!GetSimplePenetrationWithNormal(PhiWithNormalInputFloat, Type, Ptr, TypedPtr, SizeofFImplicitObject, OffsetofGeometryType, OffsetOfMargin, PerGroupThickness, Penetration, Normal))
+			if(!GetSimplePenetrationWithNormal(PhiWithNormalInputFloat, Type, Ptr, TypedPtr, SizeofFImplicitObject, OffsetofGeometryType, OffsetOfMargin, CollisionThickness, Penetration, Normal))
 			{
 				uniform int IMMask;
 				if(IM != 0)
@@ -715,7 +772,7 @@ export void ApplyPerParticleSimpleCollisionNoFrictionAndGeneratePlanarConstraint
 				
 				float Phi;
 				GetPhiWithNormalCollisionParticleRange(CollisionParticlesRange, (const uniform float *uniform)&PhiWithNormalInputFloat, (uniform float *uniform)&Normal, (uniform float *uniform)&Phi, i, programCount, IMMask);
-				Penetration = PerGroupThickness - Phi;
+				Penetration = CollisionThickness - Phi;
 			}
 
 			if(Penetration > 0.0f)
diff --git a/Engine/Source/Runtime/Experimental/Chaos/Private/Chaos/PerParticlePBDCollisionConstraint.isph b/Engine/Source/Runtime/Experimental/Chaos/Private/Chaos/PerParticlePBDCollisionConstraint.isph
--- a/Engine/Source/Runtime/Experimental/Chaos/Private/Chaos/PerParticlePBDCollisionConstraint.isph
+++ b/Engine/Source/Runtime/Experimental/Chaos/Private/Chaos/PerParticlePBDCollisionConstraint.isph
@@ -250,7 +250,8 @@ static inline void ReflectOneSidedCollision(const FVector3f& P, const uniform FV
 	Penetration = SafeNormalize(ImplicitNormal, FloatZeroVector);
 }
 
-static inline void PenetrationWithNormalSphere(const uniform uint8 *uniform TypedObjectPtr, const FVector3f &PhiWithNormalInput, const uniform float Thickness,  FVector3f &Normal, float &Penetration, const uniform float Radius)
+template<typename ThicknessType>
+static inline void PenetrationWithNormalSphere(const uniform uint8 *uniform TypedObjectPtr, const FVector3f &PhiWithNormalInput, const ThicknessType Thickness,  FVector3f &Normal, float &Penetration, const uniform float Radius)
 {
 	const uniform FVector3f Center = *((const uniform FVector3f *uniform)&TypedObjectPtr[0]);
 	Normal = PhiWithNormalInput - Center;
@@ -272,7 +273,8 @@ static inline float SignedDistancePlane(const uniform FPlaneConcrete3f *uniform
 	return VectorDot(x - P.MX, P.MNormal);
 }
 
-static inline void PenetrationWithNormalCapsule(const uniform uint8 *uniform TypedObjectPtr, const FVector3f &PhiWithNormalInput, const uniform float Thickness, FVector3f &Normal, float &Penetration, const uniform float Radius)
+template<typename ThicknessType>
+static inline void PenetrationWithNormalCapsule(const uniform uint8 *uniform TypedObjectPtr, const FVector3f &PhiWithNormalInput, const ThicknessType Thickness, FVector3f &Normal, float &Penetration, const uniform float Radius)
 {
 	const uniform Segment MSegment = *((const uniform Segment *uniform)&TypedObjectPtr[0]);
 
@@ -282,7 +284,8 @@ static inline void PenetrationWithNormalCapsule(const uniform uint8 *uniform Typ
 	Penetration = Thickness - (SafeNormalize(Normal) - Radius);
 }
 
-static inline void PenetrationWithNormalTaperedCapsule(const uniform uint8* uniform TypedObjectPtr, const FVector3f& PhiWithNormalInput, const uniform float Thickness, FVector3f& Normal, float& Penetration)
+template<typename ThicknessType>
+static inline void PenetrationWithNormalTaperedCapsule(const uniform uint8* uniform TypedObjectPtr, const FVector3f& PhiWithNormalInput, const ThicknessType Thickness, FVector3f& Normal, float& Penetration)
 {
 	const uniform FTaperedCapsule CapsuleObj = *((const uniform FTaperedCapsule *uniform)&TypedObjectPtr[0]);
 
@@ -299,7 +302,8 @@ static inline void PenetrationWithNormalTaperedCapsule(const uniform uint8* unif
 	}
 }
 
-static inline void PenetrationWithNormalExtrudedTaperedCapsule(const uniform uint8* uniform TypedObjectPtr, const FVector3f& PhiWithNormalInput, const uniform float Thickness, FVector3f& Normal, float& Penetration)
+template<typename ThicknessType>
+static inline void PenetrationWithNormalExtrudedTaperedCapsule(const uniform uint8* uniform TypedObjectPtr, const FVector3f& PhiWithNormalInput, const ThicknessType Thickness, FVector3f& Normal, float& Penetration)
 {
 	const uniform FExtrudedTaperedCapsule CapsuleObj = *((const uniform FExtrudedTaperedCapsule *uniform)&TypedObjectPtr[0]);
 
@@ -313,7 +317,8 @@ static inline void PenetrationWithNormalExtrudedTaperedCapsule(const uniform uin
 	Penetration = Thickness - (SafeNormalize(Normal) - Radius);
 }
 
-static inline void PenetrationWithNormalTaperedCylinder(const uniform uint8 *uniform TypedObjectPtr, const FVector3f &PhiWithNormalInput, const uniform float Thickness, FVector3f &Normal, float &Penetration, const uniform int SizeofFImplicitObject)
+template<typename ThicknessType>
+static inline void PenetrationWithNormalTaperedCylinder(const uniform uint8 *uniform TypedObjectPtr, const FVector3f &PhiWithNormalInput, const ThicknessType Thickness, FVector3f &Normal, float &Penetration, const uniform int SizeofFImplicitObject)
 {
 	const uniform FTaperedCylinder CylinderObj = *((const uniform FTaperedCylinder *uniform)&TypedObjectPtr[0]);
 
@@ -374,7 +379,8 @@ static inline void PenetrationWithNormalTaperedCylinder(const uniform uint8 *uni
 	Penetration = Thickness - Phi;
 }
 
-static inline void PenetrationWithNormalUnion(const uniform uint8 *uniform TypedPtr, const FVector3f &PhiWithNormalInput, const uniform float Thickness, FVector3f &Normal, float &Penetration, const uniform int SizeofFImplicitObject, const uniform int OffsetofGeometryType, const uniform int OffsetOfMargin)
+template<typename ThicknessType>
+static inline void PenetrationWithNormalUnion(const uniform uint8 *uniform TypedPtr, const FVector3f &PhiWithNormalInput, const ThicknessType Thickness, FVector3f &Normal, float &Penetration, const uniform int SizeofFImplicitObject, const uniform int OffsetofGeometryType, const uniform int OffsetOfMargin)
 {
 	// This is expecting a MObjects = [TaperedCylinder, Sphere, Sphere]
 	assert( ((const uniform TArray *uniform)TypedPtr)->ArrayNum == 3 );
@@ -561,7 +567,8 @@ static inline int32 GetPlaneVertex(const uniform FConvexStructureData *uniform S
 	}
 }
 
-static inline void PenetrationWithNormalConvex(const uniform uint8 *uniform TypedObjectPtr, const FVector3f &PhiWithNormalInput, const uniform float Thickness, FVector3f &Normal, float &Penetration)
+template<typename ThicknessType>
+static inline void PenetrationWithNormalConvex(const uniform uint8 *uniform TypedObjectPtr, const FVector3f &PhiWithNormalInput, const ThicknessType Thickness, FVector3f &Normal, float &Penetration)
 {
 	const uniform FConvex *uniform Object = (const uniform FConvex *uniform)TypedObjectPtr;
 	const uniform FPlaneConcrete3f *uniform Planes = (const uniform FPlaneConcrete3f *uniform)&Object->Planes.Data[0];
@@ -636,7 +643,8 @@ static inline void PenetrationWithNormalConvex(const uniform uint8 *uniform Type
 	Penetration = Thickness - Phi;
 }
 
-static inline bool GetSimplePenetrationWithNormal(const FVector3f& PhiWithNormalInputFloat, const uniform uint8 Type, const uniform uint8 *uniform Ptr, const uniform uint8 *uniform TypedPtr, const uniform int SizeofFImplicitObject, const uniform int OffsetofGeometryType, const uniform int OffsetOfMargin, const uniform float Thickness, float& Penetration, FVector3f& Normal)
+template<typename ThicknessType>
+static inline bool GetSimplePenetrationWithNormal(const FVector3f& PhiWithNormalInputFloat, const uniform uint8 Type, const uniform uint8 *uniform Ptr, const uniform uint8 *uniform TypedPtr, const uniform int SizeofFImplicitObject, const uniform int OffsetofGeometryType, const uniform int OffsetOfMargin, const ThicknessType Thickness, float& Penetration, FVector3f& Normal)
 {
 	switch(Type)
 	{
diff --git a/Engine/Source/Runtime/Experimental/Chaos/Public/Chaos/PBDSoftBodyCollisionConstraint.h b/Engine/Source/Runtime/Experimental/Chaos/Public/Chaos/PBDSoftBodyCollisionConstraint.h
--- a/Engine/Source/Runtime/Experimental/Chaos/Public/Chaos/PBDSoftBodyCollisionConstraint.h
+++ b/Engine/Source/Runtime/Experimental/Chaos/Public/Chaos/PBDSoftBodyCollisionConstraint.h
@@ -6,6 +6,7 @@
 #include "Chaos/PBDKinematicTriangleMeshCollisions.h"
 #include "Chaos/PBDSoftsEvolutionFwd.h"
 #include "Chaos/PBDSoftsSolverParticles.h"
+#include "Chaos/PBDFlatWeightMap.h"
 #include "Chaos/PBDWeightMap.h"
 #include "Chaos/PBDPlanarConstraints.h"
 #include "Chaos/SoftsSolverCollisionParticles.h"
@@ -52,7 +53,7 @@ public:
 		TArray<FSolverVec3>* InContacts = nullptr,
 		TArray<FSolverVec3>* InNormals = nullptr,
 		TArray<FSolverReal>* InPhis = nullptr,
-		const FSolverReal InSoftBodyCollisionThickness = 0.f,
+		const FSolverVec2 InSoftBodyCollisionThickness = FSolverVec2::ZeroVector,
 		bool bInEnableSimpleColliders = true,
 		bool bInEnableComplexColliders = true,
 		bool bInUsePlanarConstraintForSimpleColliders = bDefaultUsePlanarConstraintForSimpleColliders,
@@ -62,7 +63,7 @@ public:
 		: LastCollisionTransforms(InLastCollisionTransforms)
 		, InnerCollisionThickness(FSolverVec2::ZeroVector)
 		, CollisionThickness(InCollisionThickness)
-		, SoftBodyCollisionThickness(InSoftBodyCollisionThickness)
+		, SoftBodyCollisionThicknessWeighted(InSoftBodyCollisionThickness, TConstArrayView<FRealSingle>(), 0)
 		, FrictionCoefficient(InFrictionCoefficient)
 		, bUseCCD(bInUseCCD)
 		, ProximityStiffness(InProximityStiffness)
@@ -78,6 +79,33 @@ public:
 		, PlanarConstraint(FrictionCoefficient)
 	{}
 
+	UE_DEPRECATED(5.8, "Use the overload taking FSolverVec2 for InSoftBodyCollisionThickness")
+	FPBDSoftBodyCollisionConstraintBase(
+		const TArray<FSolverRigidTransform3>& InLastCollisionTransforms,
+		FSolverReal InCollisionThickness,
+		FSolverReal InFrictionCoefficient,
+		bool bInUseCCD,
+		FSolverReal InProximityStiffness,
+		TArray<bool>* InCollisionParticleCollided,
+		TArray<FSolverVec3>* InContacts,
+		TArray<FSolverVec3>* InNormals,
+		TArray<FSolverReal>* InPhis,
+		const FSolverReal InSoftBodyCollisionThickness,
+		bool bInEnableSimpleColliders = true,
+		bool bInEnableComplexColliders = true,
+		bool bInUsePlanarConstraintForSimpleColliders = bDefaultUsePlanarConstraintForSimpleColliders,
+		bool bInUsePlanarConstraintForComplexColliders = bDefaultUsePlanarConstraintForComplexColliders,
+		const TMap<FParticleRangeIndex, FPBDComplexColliderBoneData>& InComplexBoneData = TMap<FParticleRangeIndex, FPBDComplexColliderBoneData>()
+		)
+		: FPBDSoftBodyCollisionConstraintBase(InLastCollisionTransforms, InCollisionThickness,
+			InFrictionCoefficient, bInUseCCD, InProximityStiffness,
+			InCollisionParticleCollided, InContacts, InNormals, InPhis,
+			FSolverVec2(InSoftBodyCollisionThickness, InSoftBodyCollisionThickness),
+			bInEnableSimpleColliders, bInEnableComplexColliders,
+			bInUsePlanarConstraintForSimpleColliders, bInUsePlanarConstraintForComplexColliders,
+			InComplexBoneData)
+	{}
+
 	void SetWriteDebugContacts(bool bWrite) { bWriteDebugContacts = bWrite; }
 
 	// Follow planar constraint settings and either generate or apply planar constraints if requested.
@@ -99,6 +127,16 @@ public:
 		}
 	}
 
+	FSolverReal GetParticleClothThickness(int32 ParticleIndex) const
+	{
+		return SoftBodyCollisionThicknessWeighted.GetValue(ParticleIndex);
+	}
+
+	FSolverReal GetParticleInnerThickness(int32 ParticleIndex) const
+	{
+		return InnerCollisionThickness.GetValue(ParticleIndex);
+	}
+
 private:
 
 	template<bool bLockAndWriteContacts, bool bWithFriction, bool bGeneratePlanarConstraints>
@@ -119,7 +157,9 @@ protected:
 	const TArray<FSolverRigidTransform3>& LastCollisionTransforms; // Used by CCD
 	FPBDWeightMap InnerCollisionThickness;
 	FSolverReal CollisionThickness;
-	FSolverReal SoftBodyCollisionThickness;
+	FPBDFlatWeightMap SoftBodyCollisionThicknessWeighted;
+	UE_DEPRECATED(5.8, "Use SoftBodyCollisionThicknessWeighted instead")
+	FSolverReal SoftBodyCollisionThickness = 0.f;
 	FSolverReal FrictionCoefficient;
 	bool bUseCCD;
 	FSolverReal ProximityStiffness; // Used by force-based solver
@@ -188,13 +228,20 @@ public:
 		const TMap<FParticleRangeIndex, FPBDComplexColliderBoneData>& InComplexBoneData = TMap<FParticleRangeIndex, FPBDComplexColliderBoneData>(),
 		const TMap<FString, TConstArrayView<float>>& WeightMaps = TMap<FString, TConstArrayView<float>>());
 
-	CHAOS_API void SetProperties(const FCollectionPropertyConstFacade& PropertyCollection);
+	CHAOS_API void SetProperties(const FCollectionPropertyConstFacade& PropertyCollection,
+		const TMap<FString, TConstArrayView<float>>& WeightMaps);
+
+	UE_DEPRECATED(5.8, "Use the overload that accepts WeightMaps")
+	void SetProperties(const FCollectionPropertyConstFacade& PropertyCollection)
+	{
+		SetProperties(PropertyCollection, TMap<FString, TConstArrayView<float>>());
+	}
 
 private:
 	const FSolverReal MeshScale;
 
 	UE_CHAOS_DECLARE_PROPERTYCOLLECTION_NAME(CollisionThickness, float);
-	UE_CHAOS_DECLARE_PROPERTYCOLLECTION_NAME(SoftBodyCollisionThickness, float);
+	UE_CHAOS_DECLARE_PROPERTYCOLLECTION_NAME(ClothCollisionThickness, float);
 	UE_CHAOS_DECLARE_PROPERTYCOLLECTION_NAME(FrictionCoefficient, float);
 	UE_CHAOS_DECLARE_PROPERTYCOLLECTION_NAME(UseCCD, bool);
 	UE_CHAOS_DECLARE_PROPERTYCOLLECTION_NAME(ProximityStiffness, float);
