Dataflow:
- Minor updates to member functions

#Synced-CL 52237738

diff --git a/Engine/Source/Runtime/Experimental/Dataflow/Core/Private/Dataflow/DataflowImage.cpp b/Engine/Source/Runtime/Experimental/Dataflow/Core/Private/Dataflow/DataflowImage.cpp
--- a/Engine/Source/Runtime/Experimental/Dataflow/Core/Private/Dataflow/DataflowImage.cpp
+++ b/Engine/Source/Runtime/Experimental/Dataflow/Core/Private/Dataflow/DataflowImage.cpp
@@ -19,24 +19,34 @@ const FImage& FDataflowImage::GetImage() const
 	return Image;
 }
 
-void FDataflowImage::CreateR32F(EDataflowImageResolution Resolution)
+void FDataflowImage::CreateR32F(EDataflowImageResolution Resolution, TArrayView64<float> InPixels)
 {
-	CreateR32F((int32)Resolution, (int32)Resolution);
+	CreateR32F((int32)Resolution, (int32)Resolution, InPixels);
 }
 
-void FDataflowImage::CreateR32F(int32 Width, int32 Height)
+void FDataflowImage::CreateR32F(int32 Width, int32 Height, TArrayView64<float> InPixels)
 {
 	Image.Init(Width, Height, ERawImageFormat::R32F, EGammaSpace::Linear);
+
+	if (Width * Height == InPixels.Num())
+	{
+		CopyPixels(InPixels);
+	}
 }
 
-void FDataflowImage::CreateRGBA32F(EDataflowImageResolution Resolution)
+void FDataflowImage::CreateRGBA32F(EDataflowImageResolution Resolution, TArrayView64<FVector4f> InPixels)
 {
-	CreateRGBA32F((int32)Resolution, (int32)Resolution);
+	CreateRGBA32F((int32)Resolution, (int32)Resolution, InPixels);
 }
 
-void FDataflowImage::CreateRGBA32F(int32 Width, int32 Height)
+void FDataflowImage::CreateRGBA32F(int32 Width, int32 Height, TArrayView64<FVector4f> InPixels)
 {
 	Image.Init(Width, Height, ERawImageFormat::RGBA32F, EGammaSpace::Linear);
+
+	if (Width * Height == InPixels.Num())
+	{
+		CopyPixels(InPixels);
+	}
 }
 
 void FDataflowImage::CreateFromColor(EDataflowImageResolution Resolution, FLinearColor Color)
@@ -54,6 +64,26 @@ void FDataflowImage::CreateFromColor(int32 Width, int32 Height, FLinearColor Col
 	FImageCore::CopyImage(SourceImage, Image);
 }
 
+bool FDataflowImage::CopyPixels(TArrayView64<float> Pixels)
+{
+	if (Image.Format != ERawImageFormat::Type::R32F)
+	{
+		return false;
+	}
+	if ((GetWidth() * GetHeight()) != Pixels.Num())
+	{
+		return false;
+	}
+	const FImageView SrcImage(Pixels.GetData(), GetWidth(), GetHeight(), /*NumSlices*/1, ERawImageFormat::Type::R32F, EGammaSpace::Linear);
+	FImageCore::CopyImage(SrcImage, Image);
+	return true;
+}
+
+bool FDataflowImage::CopyPixels(TArrayView64<FVector4f> InPixels)
+{
+	return CopyRGBAPixels(InPixels);
+}
+
 bool FDataflowImage::CopyRGBAPixels(TArrayView64<FVector4f> Pixels)
 {
 	if (Image.Format != ERawImageFormat::Type::RGBA32F)
diff --git a/Engine/Source/Runtime/Experimental/Dataflow/Core/Public/Dataflow/DataflowImage.h b/Engine/Source/Runtime/Experimental/Dataflow/Core/Public/Dataflow/DataflowImage.h
--- a/Engine/Source/Runtime/Experimental/Dataflow/Core/Public/Dataflow/DataflowImage.h
+++ b/Engine/Source/Runtime/Experimental/Dataflow/Core/Public/Dataflow/DataflowImage.h
@@ -46,16 +46,16 @@ struct FDataflowImage
 	DATAFLOWCORE_API const FImage& GetImage() const;
 
 	/** Create a single channel float format image */
-	DATAFLOWCORE_API void CreateR32F(EDataflowImageResolution Resolution);
+	DATAFLOWCORE_API void CreateR32F(EDataflowImageResolution Resolution, TArrayView64<float> InPixels = {});
 
 	/** Create a single channel float format image */
-	DATAFLOWCORE_API void CreateR32F(int32 Width, int32 Height);
+	DATAFLOWCORE_API void CreateR32F(int32 Width, int32 Height, TArrayView64<float> InPixels = {});
 
 	/** Create a four channels float format image */
-	DATAFLOWCORE_API void CreateRGBA32F(EDataflowImageResolution Resolution);
+	DATAFLOWCORE_API void CreateRGBA32F(EDataflowImageResolution Resolution, TArrayView64<FVector4f> InPixels = {});
 
 	/** Create a four channels float format image */
-	DATAFLOWCORE_API void CreateRGBA32F(int32 Width, int32 Height);
+	DATAFLOWCORE_API void CreateRGBA32F(int32 Width, int32 Height, TArrayView64<FVector4f> InPixels = {});
 
 	/** Create a four channels float format image filled with a specific color */
 	DATAFLOWCORE_API void CreateFromColor(EDataflowImageResolution Resolution, FLinearColor Color);
@@ -63,11 +63,19 @@ struct FDataflowImage
 	/** Create a four channels float format image filled with a specific color */
 	DATAFLOWCORE_API void CreateFromColor(int32 Width, int32 Height, FLinearColor Color);
 
+	/**
+	* Copy R32F pixels to the image
+	* Number of pixels must match and format must be R32F
+	* return false if the copy could not be done
+	*/
+	DATAFLOWCORE_API bool CopyPixels(TArrayView64<float> Pixels);
+
 	/**
 	* Copy RGBA32F pixels to the image
 	* Number of pixels must match and format must be RGBA32F
 	* return false if the copy could not be done
 	*/
+	DATAFLOWCORE_API bool CopyPixels(TArrayView64<FVector4f> Pixels);
 	DATAFLOWCORE_API bool CopyRGBAPixels(TArrayView64<FVector4f> Pixels);
 
 	/** 
