Dataflow:
- Added AppendPoints()

#Synced-CL 52238487

diff --git a/Engine/Source/Runtime/Experimental/Chaos/Private/GeometryCollection/Facades/PointsFacade.cpp b/Engine/Source/Runtime/Experimental/Chaos/Private/GeometryCollection/Facades/PointsFacade.cpp
--- a/Engine/Source/Runtime/Experimental/Chaos/Private/GeometryCollection/Facades/PointsFacade.cpp
+++ b/Engine/Source/Runtime/Experimental/Chaos/Private/GeometryCollection/Facades/PointsFacade.cpp
@@ -57,8 +57,7 @@ namespace GeometryCollection::Facades
 	{
 		if (IsValid())
 		{
-			const TManagedArray<FVector3f>& Point = ConstCollection.GetAttribute<FVector3f>(FPointsFacade::PointAttribute, FPointsFacade::PointsGroup);
-			return Point.Num();
+			return ConstCollection.NumElements(FPointsFacade::PointsGroup);
 		}
 
 		return 0;
@@ -166,6 +165,56 @@ namespace GeometryCollection::Facades
 
 	/* ------------------------------------------------------------------------------------------------------------- */
 
+	void FPointsFacade::AppendPoints(const TArray<FVector3f>& InPoints)
+	{
+		if (ensure(MutableCollection))
+		{
+			const int32 NumNewPoints = InPoints.Num();
+
+			if (ConstCollection.HasGroup(FPointsFacade::PointsGroup) &&
+				ConstCollection.HasAttribute(FPointsFacade::PointAttribute, FPointsFacade::PointsGroup))
+			{
+				const int32 NumExistingPoints = ConstCollection.NumElements(FPointsFacade::PointsGroup);
+
+				MutableCollection->AddElements(NumNewPoints, FPointsFacade::PointsGroup);
+
+				TManagedArray<FVector3f>& Point = MutableCollection->ModifyAttribute<FVector3f>(FPointsFacade::PointAttribute, FPointsFacade::PointsGroup);
+
+				for (int32 Idx = 0; Idx < NumNewPoints; ++Idx)
+				{
+					Point[NumExistingPoints + Idx] = InPoints[Idx];
+				}
+			}
+		}
+	}
+
+	/* ------------------------------------------------------------------------------------------------------------- */
+
+	void FPointsFacade::AppendPoints(const TArray<FVector>& InPoints)
+	{
+		if (ensure(MutableCollection))
+		{
+			const int32 NumNewPoints = InPoints.Num();
+
+			if (ConstCollection.HasGroup(FPointsFacade::PointsGroup) &&
+				ConstCollection.HasAttribute(FPointsFacade::PointAttribute, FPointsFacade::PointsGroup))
+			{
+				const int32 NumExistingPoints = ConstCollection.NumElements(FPointsFacade::PointsGroup);
+
+				MutableCollection->AddElements(NumNewPoints, FPointsFacade::PointsGroup);
+
+				TManagedArray<FVector3f>& Point = MutableCollection->ModifyAttribute<FVector3f>(FPointsFacade::PointAttribute, FPointsFacade::PointsGroup);
+
+				for (int32 Idx = 0; Idx < NumNewPoints; ++Idx)
+				{
+					Point[NumExistingPoints + Idx] = FVector3f(InPoints[Idx]);
+				}
+			}
+		}
+	}
+
+	/* ------------------------------------------------------------------------------------------------------------- */
+
 	void FPointsFacade::DeletePoints(const TArray<int32>& InSortedDeletionList)
 	{
 		if (ensure(MutableCollection))
diff --git a/Engine/Source/Runtime/Experimental/Chaos/Public/GeometryCollection/Facades/PointsFacade.h b/Engine/Source/Runtime/Experimental/Chaos/Public/GeometryCollection/Facades/PointsFacade.h
--- a/Engine/Source/Runtime/Experimental/Chaos/Public/GeometryCollection/Facades/PointsFacade.h
+++ b/Engine/Source/Runtime/Experimental/Chaos/Public/GeometryCollection/Facades/PointsFacade.h
@@ -76,6 +76,9 @@ namespace GeometryCollection::Facades
 		CHAOS_API void AddPoints(const TArray<FVector3f>& InPoints);
 		CHAOS_API void AddPoints(const TArray<FVector>& InPoints);
 
+		CHAOS_API void AppendPoints(const TArray<FVector3f>& InPoints);
+		CHAOS_API void AppendPoints(const TArray<FVector>& InPoints);
+
 		/**
 		* Delete points from the collection
 		*/
