ChaosDeformable: Constraint hierarchy refactor and FSolverParticlesRange migration

- Add cross-asset dynamic spring constraint creator with Offset Geometric Contact collision detection
- Add position target facade for cross-asset constraint configuration
- Add FSolverParticlesRange convenience constructors to eliminate const_cast boilerplate
- Add FSolverParticlesRange overloads to Chaos runtime constraint APIs (GaussSeidelMainConstraint, corotated, bending, stretch bias, sphere repulsion, graph coloring)
- Change Gauss-Seidel lambda type aliases (FGSConstraintResidualHessianFunc, etc.) from FSolverParticles to FSolverParticlesRange

#Synced-CL 52219429

#Synced-CL 52219429

#Synced-CL 52219429

diff --git a/Engine/Source/Runtime/Experimental/Chaos/Private/Chaos/Deformable/ChaosDeformableSolver.cpp b/Engine/Source/Runtime/Experimental/Chaos/Private/Chaos/Deformable/ChaosDeformableSolver.cpp
--- a/Engine/Source/Runtime/Experimental/Chaos/Private/Chaos/Deformable/ChaosDeformableSolver.cpp
+++ b/Engine/Source/Runtime/Experimental/Chaos/Private/Chaos/Deformable/ChaosDeformableSolver.cpp
@@ -1391,8 +1391,8 @@ namespace Chaos::Softs
 									{
 										if (FleshInputBuffer)
 										{
-											TArray<int32> BoneIndices = WeightsFacade.GetBoneIndices()[ObjectVertexIndex];
-											TArray<float> BoneWeights = WeightsFacade.GetBoneWeights()[ObjectVertexIndex];
+											const TArray<int32>& BoneIndices = WeightsFacade.GetBoneIndices()[ObjectVertexIndex];
+											const TArray<float>& BoneWeights = WeightsFacade.GetBoneWeights()[ObjectVertexIndex];
 
 											FFleshThreadingProxy::FFleshInputBuffer* PreviousFleshBuffer = nullptr;
 											if (this->PreviousInputPackage && this->PreviousInputPackage->ObjectMap.Contains(Owner))
@@ -1685,7 +1685,7 @@ namespace Chaos::Softs
 
 			int32 StaticIndex = GSMainConstraint->AddStaticConstraintResidualAndHessianRange(1);
 
-			GSMainConstraint->StaticConstraintResidualAndHessian()[StaticIndex] = [this](const FSolverParticles& Particles, const int32 ElementIndex, const int32 ElementIndexLocal, const FSolverReal Dt, TVec3<FSolverReal>& ParticleResidual, Chaos::PMatrix<FSolverReal, 3, 3>& ParticleHessian)
+			GSMainConstraint->StaticConstraintResidualAndHessian()[StaticIndex] = [this](const FSolverParticlesRange& Particles, const int32 ElementIndex, const int32 ElementIndexLocal, const FSolverReal Dt, TVec3<FSolverReal>& ParticleResidual, Chaos::PMatrix<FSolverReal, 3, 3>& ParticleHessian)
 			{
 				this->GSCorotatedCodConstraints->AddHyperelasticResidualAndHessian(Particles, ElementIndex, ElementIndexLocal, Dt, ParticleResidual, ParticleHessian);
 			};
@@ -1700,7 +1700,7 @@ namespace Chaos::Softs
 
 			int32 StaticIndex = GSMainConstraint->AddStaticConstraintResidualAndHessianRange(1);
 
-			GSMainConstraint->StaticConstraintResidualAndHessian()[StaticIndex] = [this](const FSolverParticles& Particles, const int32 ElementIndex, const int32 ElementIndexLocal, const FSolverReal Dt, TVec3<FSolverReal>& ParticleResidual, Chaos::PMatrix<FSolverReal, 3, 3>& ParticleHessian)
+			GSMainConstraint->StaticConstraintResidualAndHessian()[StaticIndex] = [this](const FSolverParticlesRange& Particles, const int32 ElementIndex, const int32 ElementIndexLocal, const FSolverReal Dt, TVec3<FSolverReal>& ParticleResidual, Chaos::PMatrix<FSolverReal, 3, 3>& ParticleHessian)
 			{
 				this->GSLinearCodConstraints->AddHyperelasticResidualAndHessian(Particles, ElementIndex, ElementIndexLocal, Dt, ParticleResidual, ParticleHessian);
 			};
@@ -1731,7 +1731,7 @@ namespace Chaos::Softs
 
 			int32 StaticIndex = GSMainConstraint->AddStaticConstraintResidualAndHessianRange(1);
 
-			GSMainConstraint->StaticConstraintResidualAndHessian()[StaticIndex] = [this](const FSolverParticles& Particles, const int32 ElementIndex, const int32 ElementIndexLocal, const FSolverReal Dt, TVec3<FSolverReal>& ParticleResidual, Chaos::PMatrix<FSolverReal, 3, 3>& ParticleHessian)
+			GSMainConstraint->StaticConstraintResidualAndHessian()[StaticIndex] = [this](const FSolverParticlesRange& Particles, const int32 ElementIndex, const int32 ElementIndexLocal, const FSolverReal Dt, TVec3<FSolverReal>& ParticleResidual, Chaos::PMatrix<FSolverReal, 3, 3>& ParticleHessian)
 			{
 				this->GSNeohookeanConstraints->AddHyperelasticResidualAndHessian(Particles, ElementIndex, ElementIndexLocal, Dt, ParticleResidual, ParticleHessian);
 			};
@@ -1761,7 +1761,7 @@ namespace Chaos::Softs
 
 			int32 StaticIndex = GSMainConstraint->AddStaticConstraintResidualAndHessianRange(1);
 
-			GSMainConstraint->StaticConstraintResidualAndHessian()[StaticIndex] = [this](const FSolverParticles& Particles, const int32 ElementIndex, const int32 ElementIndexLocal, const FSolverReal Dt, TVec3<FSolverReal>& ParticleResidual, Chaos::PMatrix<FSolverReal, 3, 3>& ParticleHessian)
+			GSMainConstraint->StaticConstraintResidualAndHessian()[StaticIndex] = [this](const FSolverParticlesRange& Particles, const int32 ElementIndex, const int32 ElementIndexLocal, const FSolverReal Dt, TVec3<FSolverReal>& ParticleResidual, Chaos::PMatrix<FSolverReal, 3, 3>& ParticleHessian)
 				{
 					this->GSCorotatedConstraints->AddHyperelasticResidualAndHessian(Particles, ElementIndex, ElementIndexLocal, Dt, ParticleResidual, ParticleHessian);
 				};
@@ -1771,12 +1771,13 @@ namespace Chaos::Softs
 		{
 			GSWeakConstraints->ComputeInitialWCData(Evolution->Particles());
 
-			TArray<TArray<int32>> StaticIncidentElements, StaticIncidentElementsLocal;
+			const TArray<TArray<int32>>* StaticIncidentElements = nullptr;
+			const TArray<TArray<int32>>* StaticIncidentElementsLocal = nullptr;
 			const TArray<TArray<int32>>& StaticConstraints = GSWeakConstraints->GetStaticConstraintArrays(StaticIncidentElements, StaticIncidentElementsLocal);
-			GSMainConstraint->AddStaticConstraints(StaticConstraints, StaticIncidentElements, StaticIncidentElementsLocal);
+			GSMainConstraint->AddStaticConstraints(StaticConstraints, *StaticIncidentElements, *StaticIncidentElementsLocal);
 
 			const int32 StaticIndex1 = GSMainConstraint->AddStaticConstraintResidualAndHessianRange(1);
-			GSMainConstraint->StaticConstraintResidualAndHessian()[StaticIndex1] = [this](const FSolverParticles& Particles, const int32 ConstraintIndex, const int32 ConstraintIndexLocal, const FSolverReal Dt, TVec3<FSolverReal>& ParticleResidual, Chaos::PMatrix<FSolverReal, 3, 3>& ParticleHessian)
+			GSMainConstraint->StaticConstraintResidualAndHessian()[StaticIndex1] = [this](const FSolverParticlesRange& Particles, const int32 ConstraintIndex, const int32 ConstraintIndexLocal, const FSolverReal Dt, TVec3<FSolverReal>& ParticleResidual, Chaos::PMatrix<FSolverReal, 3, 3>& ParticleHessian)
 				{
 					this->GSWeakConstraints->AddWCResidual(Particles, ConstraintIndex, ConstraintIndexLocal, Dt, ParticleResidual, ParticleHessian);
 				};
@@ -1790,12 +1791,13 @@ namespace Chaos::Softs
 
 		if (GSVolumeConstraints->NumConstraints())
 		{
-			TArray<TArray<int32>> StaticIncidentElements, StaticIncidentElementsLocal;
-			const TArray<TArray<int32>> StaticConstraints = GSVolumeConstraints->GetStaticConstraintArrays(StaticIncidentElements, StaticIncidentElementsLocal);
-			GSMainConstraint->AddStaticConstraints(StaticConstraints, StaticIncidentElements, StaticIncidentElementsLocal);
+			const TArray<TArray<int32>>* StaticIncidentElements = nullptr;
+			const TArray<TArray<int32>>* StaticIncidentElementsLocal = nullptr;
+			const TArray<TArray<int32>>& StaticConstraints = GSVolumeConstraints->GetStaticConstraintArrays(StaticIncidentElements, StaticIncidentElementsLocal);
+			GSMainConstraint->AddStaticConstraints(StaticConstraints, *StaticIncidentElements, *StaticIncidentElementsLocal);
 
 			const int32 StaticIndex1 = GSMainConstraint->AddStaticConstraintResidualAndHessianRange(1);
-			GSMainConstraint->StaticConstraintResidualAndHessian()[StaticIndex1] = [this](const FSolverParticles& Particles, const int32 ConstraintIndex, const int32 ConstraintIndexLocal, const FSolverReal Dt, TVec3<FSolverReal>& ParticleResidual, Chaos::PMatrix<FSolverReal, 3, 3>& ParticleHessian)
+			GSMainConstraint->StaticConstraintResidualAndHessian()[StaticIndex1] = [this](const FSolverParticlesRange& Particles, const int32 ConstraintIndex, const int32 ConstraintIndexLocal, const FSolverReal Dt, TVec3<FSolverReal>& ParticleResidual, Chaos::PMatrix<FSolverReal, 3, 3>& ParticleHessian)
 				{
 					this->GSVolumeConstraints->AddResidualAndHessian(
 						Evolution->Particles(),
@@ -1818,7 +1820,7 @@ namespace Chaos::Softs
 		if (Property.bEnableDynamicSprings)
 		{
 			int32 DynamicIndex = GSMainConstraint->AddDynamicConstraintResidualAndHessianRange(1);
-			GSMainConstraint->DynamicConstraintResidualAndHessian()[DynamicIndex] = [this](const FSolverParticles& Particles, const int32 ConstraintIndex, const int32 ConstraintIndexLocal, const FSolverReal Dt, TVec3<FSolverReal>& ParticleResidual, Chaos::PMatrix<FSolverReal, 3, 3>& ParticleHessian)
+			GSMainConstraint->DynamicConstraintResidualAndHessian()[DynamicIndex] = [this](const FSolverParticlesRange& Particles, const int32 ConstraintIndex, const int32 ConstraintIndexLocal, const FSolverReal Dt, TVec3<FSolverReal>& ParticleResidual, Chaos::PMatrix<FSolverReal, 3, 3>& ParticleHessian)
 			{
 				this->GSDynamicWeakConstraints->AddWCResidual(Particles, ConstraintIndex, ConstraintIndexLocal, Dt, ParticleResidual, ParticleHessian);
 			};
@@ -1839,10 +1841,11 @@ namespace Chaos::Softs
 					this->GSDynamicWeakConstraints->Init(InParticles, Dt);
 					if (this->bDynamicConstraintIsUpdated)
 					{
-						TArray<TArray<int32>> WCDynamicIncidentElements,WCDynamicIncidentElementsLocal;
+						const TArray<TArray<int32>>* WCDynamicIncidentElements = nullptr;
+						const TArray<TArray<int32>>* WCDynamicIncidentElementsLocal = nullptr;
 						const TArray<TArray<int32>>& DynamicConstraints = GSDynamicWeakConstraints->GetStaticConstraintArrays(WCDynamicIncidentElements, WCDynamicIncidentElementsLocal);
 						GSMainConstraint->ResetDynamicConstraints();
-						GSMainConstraint->AddDynamicConstraints(DynamicConstraints, WCDynamicIncidentElements, WCDynamicIncidentElementsLocal, true);
+						GSMainConstraint->AddDynamicConstraints(DynamicConstraints, *WCDynamicIncidentElements, *WCDynamicIncidentElementsLocal, true);
 						this->GSMainConstraint->InitDynamicColor(InParticles);
 					}
 				};
@@ -1851,7 +1854,7 @@ namespace Chaos::Softs
 		if (Property.bDoSpringCollision)
 		{
 			int32 TransientIndex = GSMainConstraint->AddTransientConstraintResidualAndHessianRange(1);
-			GSMainConstraint->TransientConstraintResidualAndHessian()[TransientIndex] = [this](const FSolverParticles& Particles, const int32 ConstraintIndex, const int32 ConstraintIndexLocal, const FSolverReal Dt, TVec3<FSolverReal>& ParticleResidual, Chaos::PMatrix<FSolverReal, 3, 3>& ParticleHessian)
+			GSMainConstraint->TransientConstraintResidualAndHessian()[TransientIndex] = [this](const FSolverParticlesRange& Particles, const int32 ConstraintIndex, const int32 ConstraintIndexLocal, const FSolverReal Dt, TVec3<FSolverReal>& ParticleResidual, Chaos::PMatrix<FSolverReal, 3, 3>& ParticleHessian)
 				{
 					this->GSWeakConstraints->AddWCResidual(Particles, ConstraintIndex + this->GSWeakConstraints->GetInitialWCSize(), ConstraintIndexLocal, Dt, ParticleResidual, ParticleHessian);
 				};
@@ -1902,20 +1905,10 @@ namespace Chaos::Softs
 
 			int32 TransientIndex = GSMainConstraint->AddTransientConstraintResidualAndHessianRange(1);
 			GSMainConstraint->TransientConstraintResidualAndHessian()[TransientIndex] =
-				[this](const FSolverParticles& Particles, const int32 ConstraintIndex, const int32 ConstraintIndexLocal, const FSolverReal Dt, TVec3<FSolverReal>& ParticleResidual, Chaos::PMatrix<FSolverReal, 3, 3>& ParticleHessian)
+				[this](const FSolverParticlesRange& Particles, const int32 ConstraintIndex, const int32 ConstraintIndexLocal, const FSolverReal Dt, TVec3<FSolverReal>& ParticleResidual, Chaos::PMatrix<FSolverReal, 3, 3>& ParticleHessian)
 				{
 					this->GSSphereRepulsionConstraints->AddSphereRepulsionResidualAndHessian(Particles, ConstraintIndex, ConstraintIndexLocal, Dt, ParticleResidual, ParticleHessian);
 				};
-
-			if (!Property.bEnablePositionTargets)
-			{
-				int32 PerNodeIndex = GSMainConstraint->AddPerNodeHessianRange(1);
-				GSMainConstraint->PerNodeHessian()[PerNodeIndex] = 
-					[this](const int32 p, const FSolverReal Dt, Chaos::PMatrix<FSolverReal, 3, 3>& ParticleHessian)
-					{
-						this->GSSphereRepulsionConstraints->AddSphereRepulsionHessian(p, Dt, ParticleHessian);
-					};
-			}
 		}
 		if (Property.bEnableGravity && Property.bDoQuasistatics) // Quasistatic PBD evolution does not apply gravity, add here
 		{
diff --git a/Engine/Source/Runtime/Experimental/Chaos/Private/Chaos/Deformable/GaussSeidelMainConstraint.cpp b/Engine/Source/Runtime/Experimental/Chaos/Private/Chaos/Deformable/GaussSeidelMainConstraint.cpp
--- a/Engine/Source/Runtime/Experimental/Chaos/Private/Chaos/Deformable/GaussSeidelMainConstraint.cpp
+++ b/Engine/Source/Runtime/Experimental/Chaos/Private/Chaos/Deformable/GaussSeidelMainConstraint.cpp
@@ -15,26 +15,31 @@ namespace Chaos::Softs
 
 	FAutoConsoleVariableRef CVarClothCGTol(TEXT("p.Chaos.Cloth.CGTol"), CGTol, TEXT("CG Tolerance [def: 1e-4]"));
 
-	void FGaussSeidelMainConstraints::AddStaticConstraints(const TArray<TArray<int32>>& ExtraConstraints, TArray<TArray<int32>>& ExtraIncidentElements, TArray<TArray<int32>>& ExtraIncidentElementsLocal)
-	{	
+	void FGaussSeidelMainConstraints::AddStaticConstraints(const TArray<TArray<int32>>& ExtraConstraints, const TArray<TArray<int32>>& ExtraIncidentElements, const TArray<TArray<int32>>& ExtraIncidentElementsLocal)
+	{
+		const TArray<TArray<int32>>* IncidentElementsPtr = &ExtraIncidentElements;
+		const TArray<TArray<int32>>* IncidentElementsLocalPtr = &ExtraIncidentElementsLocal;
+		TArray<TArray<int32>> RecomputedIncidentElements, RecomputedIncidentElementsLocal;
 		if (!IsClean(ExtraConstraints, ExtraIncidentElements, ExtraIncidentElementsLocal))
 		{
-			ExtraIncidentElements = Chaos::Utilities::ComputeIncidentElements(ExtraConstraints, &ExtraIncidentElementsLocal);
+			RecomputedIncidentElements = Chaos::Utilities::ComputeIncidentElements(ExtraConstraints, &RecomputedIncidentElementsLocal);
+			IncidentElementsPtr = &RecomputedIncidentElements;
+			IncidentElementsLocalPtr = &RecomputedIncidentElementsLocal;
 		}
 
 		int32 Offset = StaticConstraints.Num();
 		StaticConstraints += ExtraConstraints;
-		for (int32 i = 0; i < ExtraIncidentElements.Num(); i++)
+		for (int32 ParticleIdx = 0; ParticleIdx < IncidentElementsPtr->Num(); ParticleIdx++)
 		{
-			if (ExtraIncidentElements[i].Num() > 0)
+			if ((*IncidentElementsPtr)[ParticleIdx].Num() > 0)
 			{
-				TArray<int32> ExtraIncidentElementsWithOffset = ExtraIncidentElements[i];
-				for (int32 j = 0; j < ExtraIncidentElementsWithOffset.Num(); j++)
+				TArray<int32> ExtraIncidentElementsWithOffset = (*IncidentElementsPtr)[ParticleIdx];
+				for (int32 LocalIdx = 0; LocalIdx < ExtraIncidentElementsWithOffset.Num(); LocalIdx++)
 				{
-					ExtraIncidentElementsWithOffset[j] += Offset;
+					ExtraIncidentElementsWithOffset[LocalIdx] += Offset;
 				}
-				StaticIncidentElements[i] += ExtraIncidentElementsWithOffset;
-				StaticIncidentElementsLocal[i] += ExtraIncidentElementsLocal[i];
+				StaticIncidentElements[ParticleIdx] += ExtraIncidentElementsWithOffset;
+				StaticIncidentElementsLocal[ParticleIdx] += (*IncidentElementsLocalPtr)[ParticleIdx];
 			}
 		}
 		if (StaticIncidentElementsOffsets.Num() > 0)
@@ -45,29 +50,31 @@ namespace Chaos::Softs
 		StaticIncidentElementsOffsets.Add(StaticConstraints.Num());
 	}
 
-	void FGaussSeidelMainConstraints::AddTransientConstraints(const TArray<TArray<int32>>& ExtraConstraints, TArray<TArray<int32>>& ExtraIncidentElements, TArray<TArray<int32>>& ExtraIncidentElementsLocal, bool CheckIncidentElements)
+	void FGaussSeidelMainConstraints::AddTransientConstraints(const TArray<TArray<int32>>& ExtraConstraints, const TArray<TArray<int32>>& ExtraIncidentElements, const TArray<TArray<int32>>& ExtraIncidentElementsLocal, bool CheckIncidentElements)
 	{
-		if (CheckIncidentElements)
+		const TArray<TArray<int32>>* IncidentElementsPtr = &ExtraIncidentElements;
+		const TArray<TArray<int32>>* IncidentElementsLocalPtr = &ExtraIncidentElementsLocal;
+		TArray<TArray<int32>> RecomputedIncidentElements, RecomputedIncidentElementsLocal;
+		if (CheckIncidentElements && !IsClean(ExtraConstraints, ExtraIncidentElements, ExtraIncidentElementsLocal))
 		{
-			if (!IsClean(ExtraConstraints, ExtraIncidentElements, ExtraIncidentElementsLocal))
-			{
-				ExtraIncidentElements = Chaos::Utilities::ComputeIncidentElements(ExtraConstraints, &ExtraIncidentElementsLocal);
-			}
+			RecomputedIncidentElements = Chaos::Utilities::ComputeIncidentElements(ExtraConstraints, &RecomputedIncidentElementsLocal);
+			IncidentElementsPtr = &RecomputedIncidentElements;
+			IncidentElementsLocalPtr = &RecomputedIncidentElementsLocal;
 		}
 
 		int32 Offset = TransientConstraints.Num();
 		TransientConstraints += ExtraConstraints;
-		for (int32 i = 0; i < ExtraIncidentElements.Num(); i++)
+		for (int32 ParticleIdx = 0; ParticleIdx < IncidentElementsPtr->Num(); ParticleIdx++)
 		{
-			if (ExtraIncidentElements[i].Num() > 0)
+			if ((*IncidentElementsPtr)[ParticleIdx].Num() > 0)
 			{
-				TArray<int32> ExtraIncidentElementsWithOffset = ExtraIncidentElements[i];
-				for (int32 j = 0; j < ExtraIncidentElementsWithOffset.Num(); j++)
+				TArray<int32> ExtraIncidentElementsWithOffset = (*IncidentElementsPtr)[ParticleIdx];
+				for (int32 LocalIdx = 0; LocalIdx < ExtraIncidentElementsWithOffset.Num(); LocalIdx++)
 				{
-					ExtraIncidentElementsWithOffset[j] += Offset;
+					ExtraIncidentElementsWithOffset[LocalIdx] += Offset;
 				}
-				TransientIncidentElements[i] += ExtraIncidentElementsWithOffset;
-				TransientIncidentElementsLocal[i] += ExtraIncidentElementsLocal[i];
+				TransientIncidentElements[ParticleIdx] += ExtraIncidentElementsWithOffset;
+				TransientIncidentElementsLocal[ParticleIdx] += (*IncidentElementsLocalPtr)[ParticleIdx];
 			}
 		}
 		if (TransientIncidentElementsOffsets.Num() > 0)
@@ -76,32 +83,33 @@ namespace Chaos::Softs
 		}
 		TransientIncidentElementsOffsets.Add(Offset);
 		TransientIncidentElementsOffsets.Add(TransientConstraints.Num());
-	
 	}
 
-	void FGaussSeidelMainConstraints::AddDynamicConstraints(const TArray<TArray<int32>>& ExtraConstraints, TArray<TArray<int32>>& ExtraIncidentElements, TArray<TArray<int32>>& ExtraIncidentElementsLocal, bool CheckIncidentElements)
+	void FGaussSeidelMainConstraints::AddDynamicConstraints(const TArray<TArray<int32>>& ExtraConstraints, const TArray<TArray<int32>>& ExtraIncidentElements, const TArray<TArray<int32>>& ExtraIncidentElementsLocal, bool CheckIncidentElements)
 	{
-		if (CheckIncidentElements)
+		const TArray<TArray<int32>>* IncidentElementsPtr = &ExtraIncidentElements;
+		const TArray<TArray<int32>>* IncidentElementsLocalPtr = &ExtraIncidentElementsLocal;
+		TArray<TArray<int32>> RecomputedIncidentElements, RecomputedIncidentElementsLocal;
+		if (CheckIncidentElements && !IsClean(ExtraConstraints, ExtraIncidentElements, ExtraIncidentElementsLocal))
 		{
-			if (!IsClean(ExtraConstraints, ExtraIncidentElements, ExtraIncidentElementsLocal))
-			{
-				ExtraIncidentElements = Chaos::Utilities::ComputeIncidentElements(ExtraConstraints, &ExtraIncidentElementsLocal);
-			}
+			RecomputedIncidentElements = Chaos::Utilities::ComputeIncidentElements(ExtraConstraints, &RecomputedIncidentElementsLocal);
+			IncidentElementsPtr = &RecomputedIncidentElements;
+			IncidentElementsLocalPtr = &RecomputedIncidentElementsLocal;
 		}
 
 		int32 Offset = DynamicConstraints.Num();
 		DynamicConstraints += ExtraConstraints;
-		for (int32 i = 0; i < ExtraIncidentElements.Num(); i++)
+		for (int32 ParticleIdx = 0; ParticleIdx < IncidentElementsPtr->Num(); ParticleIdx++)
 		{
-			if (ExtraIncidentElements[i].Num() > 0)
+			if ((*IncidentElementsPtr)[ParticleIdx].Num() > 0)
 			{
-				TArray<int32> ExtraIncidentElementsWithOffset = ExtraIncidentElements[i];
-				for (int32 j = 0; j < ExtraIncidentElementsWithOffset.Num(); j++)
+				TArray<int32> ExtraIncidentElementsWithOffset = (*IncidentElementsPtr)[ParticleIdx];
+				for (int32 LocalIdx = 0; LocalIdx < ExtraIncidentElementsWithOffset.Num(); LocalIdx++)
 				{
-					ExtraIncidentElementsWithOffset[j] += Offset;
+					ExtraIncidentElementsWithOffset[LocalIdx] += Offset;
 				}
-				DynamicIncidentElements[i] += ExtraIncidentElementsWithOffset;
-				DynamicIncidentElementsLocal[i] += ExtraIncidentElementsLocal[i];
+				DynamicIncidentElements[ParticleIdx] += ExtraIncidentElementsWithOffset;
+				DynamicIncidentElementsLocal[ParticleIdx] += (*IncidentElementsLocalPtr)[ParticleIdx];
 			}
 		}
 		if (DynamicIncidentElementsOffsets.Num() > 0)
@@ -112,7 +120,7 @@ namespace Chaos::Softs
 		DynamicIncidentElementsOffsets.Add(DynamicConstraints.Num());
 	}
 
-	TArray<TVec3<FSolverReal>> FGaussSeidelMainConstraints::ComputeNewtonResiduals(const FSolverParticles& Particles, const FSolverReal Dt, const bool Write2File, TArray<PMatrix<FSolverReal, 3, 3>>* AllParticleHessian)
+	TArray<TVec3<FSolverReal>> FGaussSeidelMainConstraints::ComputeNewtonResiduals(const FSolverParticlesRange& ParticlesRange, const FSolverReal Dt, const bool Write2File, TArray<PMatrix<FSolverReal, 3, 3>>* AllParticleHessian)
 	{
 		TRACE_CPUPROFILER_EVENT_SCOPE(STAT_ChaosGaussSeidelComputeNewtonResidual);
 
@@ -121,56 +129,56 @@ namespace Chaos::Softs
 
 		if (AllParticleHessian)
 		{
-			AllParticleHessian->Init(PMatrix<FSolverReal, 3, 3>((FSolverReal)0.), Particles.Size());
+			AllParticleHessian->Init(PMatrix<FSolverReal, 3, 3>((FSolverReal)0.), ParticlesRange.Size());
 		}
 
-		for (int32 k = 0; k < ParticlesPerColor.Num(); k++)
+		for (int32 ColorIdx = 0; ColorIdx < ParticlesPerColor.Num(); ColorIdx++)
 		{
-			PhysicsParallelFor(ParticlesPerColor[k].Num(), [this, &AllParticleHessian, Dt, k, &NewtonResidual, &Particles](const int32 j)
+			PhysicsParallelFor(ParticlesPerColor[ColorIdx].Num(), [this, &AllParticleHessian, Dt, ColorIdx, &NewtonResidual, &ParticlesRange](const int32 BatchIdx)
 				{
-					const int32 p = ParticlesPerColor[k][j];
+					const int32 ParticleIdx = ParticlesPerColor[ColorIdx][BatchIdx];
 
-					if (Particles.InvM(p) != (FSolverReal)0.)
+					if (ParticlesRange.InvM(ParticleIdx) != (FSolverReal)0.)
 					{
 						int32 ConstraintIndex = 0;
 						Chaos::PMatrix<FSolverReal, 3, 3> ParticleHessian((FSolverReal)0., (FSolverReal)0., (FSolverReal)0.);
 
-						this->ComputeInitialResidualAndHessian(Particles, p, Dt, NewtonResidual[p], ParticleHessian);
+						this->ComputeInitialResidualAndHessian(ParticlesRange, ParticleIdx, Dt, NewtonResidual[ParticleIdx], ParticleHessian);
 
-						for (int32 i = 0; i < StaticIncidentElements[p].Num(); i++)
+						for (int32 IncidentIdx = 0; IncidentIdx < StaticIncidentElements[ParticleIdx].Num(); IncidentIdx++)
 						{
-							while (StaticIncidentElements[p][i] >= StaticIncidentElementsOffsets[ConstraintIndex + 1] && ConstraintIndex < StaticIncidentElementsOffsets.Num() - 1)
+							while (StaticIncidentElements[ParticleIdx][IncidentIdx] >= StaticIncidentElementsOffsets[ConstraintIndex + 1] && ConstraintIndex < StaticIncidentElementsOffsets.Num() - 1)
 							{
 								ConstraintIndex += 1;
 							}
 
-							this->AddStaticConstraintResidualAndHessian[ConstraintIndex](Particles, StaticIncidentElements[p][i] - StaticIncidentElementsOffsets[ConstraintIndex], StaticIncidentElementsLocal[p][i], Dt, NewtonResidual[p], ParticleHessian);
+							this->AddStaticConstraintResidualAndHessian[ConstraintIndex](ParticlesRange, StaticIncidentElements[ParticleIdx][IncidentIdx] - StaticIncidentElementsOffsets[ConstraintIndex], StaticIncidentElementsLocal[ParticleIdx][IncidentIdx], Dt, NewtonResidual[ParticleIdx], ParticleHessian);
 						}
 
 						ConstraintIndex = 0;
 
-						for (int32 i = 0; i < TransientIncidentElements[p].Num(); i++)
+						for (int32 IncidentIdx = 0; IncidentIdx < TransientIncidentElements[ParticleIdx].Num(); IncidentIdx++)
 						{
-							while (TransientIncidentElements[p][i] >= TransientIncidentElementsOffsets[ConstraintIndex + 1] && ConstraintIndex < TransientIncidentElementsOffsets.Num() - 1)
+							while (TransientIncidentElements[ParticleIdx][IncidentIdx] >= TransientIncidentElementsOffsets[ConstraintIndex + 1] && ConstraintIndex < TransientIncidentElementsOffsets.Num() - 1)
 							{
 								ConstraintIndex += 1;
 							}
 
-							this->AddTransientConstraintResidualAndHessian[ConstraintIndex](Particles, TransientIncidentElements[p][i] - TransientIncidentElementsOffsets[ConstraintIndex], TransientIncidentElementsLocal[p][i], Dt, NewtonResidual[p], ParticleHessian);
+							this->AddTransientConstraintResidualAndHessian[ConstraintIndex](ParticlesRange, TransientIncidentElements[ParticleIdx][IncidentIdx] - TransientIncidentElementsOffsets[ConstraintIndex], TransientIncidentElementsLocal[ParticleIdx][IncidentIdx], Dt, NewtonResidual[ParticleIdx], ParticleHessian);
 						}
 
 						if (AllParticleHessian)
 						{
-							(*AllParticleHessian)[p] = ParticleHessian;
+							(*AllParticleHessian)[ParticleIdx] = ParticleHessian;
 						}
 					}
-			}, ParticlesPerColor[k].Num() < 1000);
+			}, ParticlesPerColor[ColorIdx].Num() < 1000);
 		}
 
 		FSolverReal NewtonNorm = (FSolverReal)0.;
-		for (int32 p = 0; p < NewtonResidual.Num(); p++)
+		for (int32 ParticleIdx = 0; ParticleIdx < NewtonResidual.Num(); ParticleIdx++)
 		{
-			NewtonNorm += NewtonResidual[p].SizeSquared();
+			NewtonNorm += NewtonResidual[ParticleIdx].SizeSquared();
 		}
 
 		NewtonNorm = FMath::Sqrt(NewtonNorm);
@@ -199,65 +207,76 @@ namespace Chaos::Softs
 
 	}
 
-	void FGaussSeidelMainConstraints::ApplyCG(FSolverParticles& Particles, const FSolverReal Dt)
+	TArray<TVec3<FSolverReal>> FGaussSeidelMainConstraints::ComputeNewtonResiduals(const FSolverParticles& Particles, const FSolverReal Dt, const bool Write2File, TArray<PMatrix<FSolverReal, 3, 3>>* AllParticleHessian)
+	{
+		return ComputeNewtonResiduals(FSolverParticlesRange(Particles), Dt, Write2File, AllParticleHessian);
+	}
+
+	void FGaussSeidelMainConstraints::ApplyCG(FSolverParticlesRange& ParticlesRange, const FSolverReal Dt)
 	{
 		TRACE_CPUPROFILER_EVENT_SCOPE(STAT_ChaosGaussSeidelApplyCG);
-		TFunction<void(TArray<FSolverVec3>&)> ProjectBCs = [&Particles](TArray<FSolverVec3>& y)
+		TFunction<void(TArray<FSolverVec3>&)> ProjectBCs = [&ParticlesRange](TArray<FSolverVec3>& Y)
 		{
-			for (int32 i = 0; i < y.Num(); i++)
+			for (int32 ParticleIdx = 0; ParticleIdx < Y.Num(); ParticleIdx++)
 			{
-				if (Particles.InvM(i) == 0.f)
+				if (ParticlesRange.InvM(ParticleIdx) == 0.f)
 				{
-					y[i] = FSolverVec3(0.f);
+					Y[ParticleIdx] = FSolverVec3(0.f);
 				}
 			}
 		};
 
-		int32 MaxNewtonIt = 1;
+		const int32 MaxNewtonIt = 1;
 
 		for (int32 It = 0; It < MaxNewtonIt; It++)
 		{
-			const TArray<FSolverVec3> Residual = ComputeNewtonResiduals(Particles, Dt, false, nullptr);
+			const TArray<FSolverVec3> Residual = ComputeNewtonResiduals(ParticlesRange, Dt, false, nullptr);
 
-			auto multiply = [this, &ProjectBCs, &Dt, &Particles](TArray<FSolverVec3>& y, const TArray<FSolverVec3>& x)
+			auto Multiply = [this, &ProjectBCs, &Dt, &ParticlesRange](TArray<FSolverVec3>& Y, const TArray<FSolverVec3>& X)
 			{
-				TArray<FSolverVec3> XProj = x;
+				TArray<FSolverVec3> XProj = X;
 				ProjectBCs(XProj);
 
-				y.Init(FSolverVec3((FSolverReal)0.), XProj.Num());
+				Y.Init(FSolverVec3((FSolverReal)0.), XProj.Num());
 
-				for (int32 i = 0; i < this->AddInternalForceDifferentials.Num(); i++)
+				for (int32 ParticleIdx = 0; ParticleIdx < this->AddInternalForceDifferentials.Num(); ParticleIdx++)
 				{
-					this->AddInternalForceDifferentials[i](Particles, XProj, y);
+					this->AddInternalForceDifferentials[ParticleIdx](ParticlesRange, XProj, Y);
 				}
 
-				PhysicsParallelFor(y.Num(),
-					[&y, &Particles, &XProj, Dt](const int32 i)
+				PhysicsParallelFor(Y.Num(),
+					[&Y, &ParticlesRange, &XProj, Dt](const int32 ParticleIdx)
 					{
-						y[i] = Particles.M(i) * XProj[i] + Dt * Dt * y[i];
+						Y[ParticleIdx] = ParticlesRange.M(ParticleIdx) * XProj[ParticleIdx] + Dt * Dt * Y[ParticleIdx];
 					});
 
-				ProjectBCs(y);
+				ProjectBCs(Y);
 			};
 
-			TArray<FSolverVec3> Deltax;
-			Deltax.Init(FSolverVec3(0.f), Particles.Size());
+			TArray<FSolverVec3> DeltaX;
+			DeltaX.Init(FSolverVec3(0.f), ParticlesRange.Size());
 
-			Chaos::LanczosCG<FSolverReal>(multiply, Deltax, Residual, MaxItCG, CGTol, use_list.Get());
+			Chaos::LanczosCG<FSolverReal>(Multiply, DeltaX, Residual, MaxItCG, CGTol, use_list.Get());
 
-			PhysicsParallelFor(Particles.Size(), [&Particles, &Deltax](const int32 i)
+			PhysicsParallelFor(ParticlesRange.Size(), [&ParticlesRange, &DeltaX](const int32 ParticleIdx)
 				{
-					Particles.P(i) -= Deltax[i];
+					ParticlesRange.P(ParticleIdx) -= DeltaX[ParticleIdx];
 				});
 		}
 	}
 
+	void FGaussSeidelMainConstraints::ApplyCG(FSolverParticles& Particles, const FSolverReal Dt)
+	{
+		FSolverParticlesRange Range(Particles);
+		ApplyCG(Range, Dt);
+	}
+
 	// -------------------------------------------------------------------
 	// Methods moved from header
 	// -------------------------------------------------------------------
 
 	FGaussSeidelMainConstraints::FGaussSeidelMainConstraints(
-		const FSolverParticles& InParticles,
+		const FSolverParticlesRange& InParticles,
 		const bool bDoQuasistaticsIn,
 		const bool bDoSORIn,
 		const FSolverReal InOmegaSOR,
@@ -275,24 +294,37 @@ namespace Chaos::Softs
 		InitializeLambdas();
 
 		TVec3<FSolverReal> MaxCoord((FSolverReal)100.), MinCoord((FSolverReal)-100.);
-		for (int32 i = 0; i < (int32)InParticles.Size(); i++)
+		for (int32 ParticleIdx = 0; ParticleIdx < (int32)InParticles.Size(); ParticleIdx++)
 		{
-			for (int32 j = 0; j < 3; j++)
+			for (int32 Alpha = 0; Alpha < 3; Alpha++)
 			{
-				if (InParticles.X(i)[j] < MinCoord[j])
+				if (InParticles.X(ParticleIdx)[Alpha] < MinCoord[Alpha])
 				{
-					MinCoord[j] = InParticles.X(i)[j];
+					MinCoord[Alpha] = InParticles.X(ParticleIdx)[Alpha];
 				}
-				if (InParticles.X(i)[j] > MaxCoord[j])
+				if (InParticles.X(ParticleIdx)[Alpha] > MaxCoord[Alpha])
 				{
-					MaxCoord[j] = InParticles.X(i)[j];
+					MaxCoord[Alpha] = InParticles.X(ParticleIdx)[Alpha];
 				}
 			}
 		}
 		MaxDxSize = (MaxCoord - MinCoord).Size() * MaxDxRatioIn;
 	}
 
-	FGaussSeidelMainConstraints::FGaussSeidelMainConstraints(const FSolverParticles& InParticles, const FCollectionPropertyConstFacade& Property, const FSolverReal MaxDxRatioIn)
+	FGaussSeidelMainConstraints::FGaussSeidelMainConstraints(
+		const FSolverParticles& InParticles,
+		const bool bDoQuasistaticsIn,
+		const bool bDoSORIn,
+		const FSolverReal InOmegaSOR,
+		const int32 ParallelMaxIn,
+		const FSolverReal MaxDxRatioIn,
+		const FDeformableXPBDCorotatedParams& InParams)
+		: FGaussSeidelMainConstraints(
+			FSolverParticlesRange(InParticles),
+			bDoQuasistaticsIn, bDoSORIn, InOmegaSOR, ParallelMaxIn, MaxDxRatioIn, InParams)
+	{}
+
+	FGaussSeidelMainConstraints::FGaussSeidelMainConstraints(const FSolverParticlesRange& InParticles, const FCollectionPropertyConstFacade& Property, const FSolverReal MaxDxRatioIn)
 		: bDoQuasistatics(GetDoQuasistatics(Property, false))
 		, bDoAcceleration(GetAccelerateSolverUsingSOR(Property, true))
 		, OmegaSOR(GetOmegaSOR(Property, (FSolverReal)1.6))
@@ -310,6 +342,12 @@ namespace Chaos::Softs
 		MaxDxSize = (MaxCoord - MinCoord).Size() * MaxDxRatioIn;
 	}
 
+	FGaussSeidelMainConstraints::FGaussSeidelMainConstraints(const FSolverParticles& InParticles, const FCollectionPropertyConstFacade& Property, const FSolverReal MaxDxRatioIn)
+		: FGaussSeidelMainConstraints(
+			FSolverParticlesRange(InParticles),
+			Property, MaxDxRatioIn)
+	{}
+
 	void FGaussSeidelMainConstraints::Resize(const int32 NewSize)
 	{
 		xtilde.SetNum(NewSize);
@@ -409,20 +447,26 @@ namespace Chaos::Softs
 	}
 
 	void FGaussSeidelMainConstraints::Apply(FSolverParticles& Particles, const FSolverReal Dt, const int32 MaxWriteIters, const bool Write2File, const TPBDActiveView<FSolverParticles>* InParticleActiveView)
+	{
+		FSolverParticlesRange ParticlesRange(Particles);
+		Apply(ParticlesRange, Dt, MaxWriteIters, Write2File, InParticleActiveView);
+	}
+
+	void FGaussSeidelMainConstraints::Apply(FSolverParticlesRange& ParticlesRange, const FSolverReal Dt, const int32 MaxWriteIters, const bool Write2File, const TPBDActiveView<FSolverParticles>* InParticleActiveView)
 	{
 		PERF_SCOPE(STAT_ChaosGSMainConstraint_Apply);
 
 		if (DebugResidual && PassedIters < MaxWriteIters)
 		{
-			ComputeNewtonResiduals(Particles, Dt, Write2File);
+			ComputeNewtonResiduals(ParticlesRange, Dt, Write2File);
 		}
 
 		std::atomic<int32> ParticleFailureCounter(0);
 
-		for (int32 i = 0; i < ParticlesPerColor.Num(); i++)
+		for (int32 ColorIdx = 0; ColorIdx < ParticlesPerColor.Num(); ColorIdx++)
 		{
-			int32 NumBatch = ParticlesPerColor[i].Num() / CorotatedParams.XPBDCorotatedBatchSize;
-			if (ParticlesPerColor[i].Num() % CorotatedParams.XPBDCorotatedBatchSize != 0)
+			int32 NumBatch = ParticlesPerColor[ColorIdx].Num() / CorotatedParams.XPBDCorotatedBatchSize;
+			if (ParticlesPerColor[ColorIdx].Num() % CorotatedParams.XPBDCorotatedBatchSize != 0)
 			{
 				NumBatch ++;
 			}
@@ -430,13 +474,13 @@ namespace Chaos::Softs
 			PhysicsParallelFor(NumBatch, [&](const int32 BatchIndex)
 				{
 					for (int32 BatchSubIndex = 0; BatchSubIndex < CorotatedParams.XPBDCorotatedBatchSize; BatchSubIndex++) {
-						int32 TaskIndex = CorotatedParams.XPBDCorotatedBatchSize * BatchIndex + BatchSubIndex;
-						if (TaskIndex < ParticlesPerColor[i].Num())
+						const int32 TaskIndex = CorotatedParams.XPBDCorotatedBatchSize * BatchIndex + BatchSubIndex;
+						if (TaskIndex < ParticlesPerColor[ColorIdx].Num())
 						{
-							if (Particles.InvM(ParticlesPerColor[i][TaskIndex]) != FSolverReal(0))
+							if (ParticlesRange.InvM(ParticlesPerColor[ColorIdx][TaskIndex]) != FSolverReal(0))
 							{
-								int32 ParticleIndex = ParticlesPerColor[i][TaskIndex];
-								if (!ApplySingleParticle(ParticleIndex, Dt, Particles))
+								const int32 ParticleIndex = ParticlesPerColor[ColorIdx][TaskIndex];
+								if (!ApplySingleParticle(ParticleIndex, Dt, ParticlesRange))
 								{
 									ParticleFailureCounter.fetch_add(1, std::memory_order_relaxed);
 								}
@@ -453,14 +497,19 @@ namespace Chaos::Softs
 			PERF_SCOPE(STAT_ChaosGSMainConstraint_Acceleration);
 			if (InParticleActiveView)
 			{
-				InParticleActiveView->ParallelFor(AccelerationTechniquePerParticle, CorotatedParams.XPBDCorotatedBatchSize);
+				// TODO: Update TPBDActiveView to support FSolverParticlesRange
+				InParticleActiveView->ParallelFor(
+					[this, &ParticlesRange](FSolverParticles& /*Unused*/, int32 ParticleIndex)
+					{
+						this->AccelerationTechniquePerParticle(ParticlesRange, ParticleIndex);
+					}, CorotatedParams.XPBDCorotatedBatchSize);
 			}
 			else
 			{
-				PhysicsParallelFor(Particles.Size(), [this, &Particles](const int32 ParticleIndex)
+				PhysicsParallelFor(ParticlesRange.Size(), [this, &ParticlesRange](const int32 ParticleIndex)
 					{
-						this->AccelerationTechniquePerParticle(Particles, ParticleIndex);
-					}, Particles.Size() < 1000);
+						this->AccelerationTechniquePerParticle(ParticlesRange, ParticleIndex);
+					}, ParticlesRange.Size() < 1000);
 			}
 		}
 
@@ -472,59 +521,81 @@ namespace Chaos::Softs
 		}
 	}
 
-	void FGaussSeidelMainConstraints::InitStaticColor(const FSolverParticles& Particles, const TPBDActiveView<FSolverParticles>* InParticleActiveView)
+	void FGaussSeidelMainConstraints::InitStaticColor(const FSolverParticlesRange& ParticlesRange, const TPBDActiveView<FSolverParticles>* InParticleActiveView)
 	{
-		StaticParticlesPerColor = ComputeNodalColoring(StaticConstraints, Particles, 0, Particles.Size(), StaticIncidentElements, StaticIncidentElementsLocal, InParticleActiveView,  &StaticParticleColors);
+		StaticParticlesPerColor = ComputeNodalColoring(StaticConstraints, ParticlesRange, 0, ParticlesRange.Size(), StaticIncidentElements, StaticIncidentElementsLocal, InParticleActiveView, &StaticParticleColors);
 		ParticleColors = StaticParticleColors;
 		ParticlesPerColor = StaticParticlesPerColor;
 	}
 
-	void FGaussSeidelMainConstraints::InitTransientColor(const FSolverParticles& Particles)
+	void FGaussSeidelMainConstraints::InitStaticColor(const FSolverParticles& Particles, const TPBDActiveView<FSolverParticles>* InParticleActiveView)
+	{
+		InitStaticColor(FSolverParticlesRange(Particles), InParticleActiveView);
+	}
+
+	void FGaussSeidelMainConstraints::InitTransientColor(const FSolverParticlesRange& ParticlesRange)
 	{
+		ensureMsgf(StaticParticleColors.Num() == ParticlesRange.Size(), TEXT("InitStaticColor must be called before InitTransientColor"));
 		PERF_SCOPE(STAT_ChaosGSMainConstraint_InitTransientColor);
 		ParticleColors = StaticParticleColors;
 		ParticlesPerColor = StaticParticlesPerColor;
-		Chaos::ComputeExtraNodalColoring(StaticConstraints, DynamicConstraints, TransientConstraints, Particles, StaticIncidentElements, DynamicIncidentElements, TransientIncidentElements, ParticleColors, ParticlesPerColor);
+		Chaos::ComputeExtraNodalColoring(StaticConstraints, DynamicConstraints, TransientConstraints, ParticlesRange, StaticIncidentElements, DynamicIncidentElements, TransientIncidentElements, ParticleColors, ParticlesPerColor);
 	}
 
-	void FGaussSeidelMainConstraints::InitDynamicColor(const FSolverParticles& Particles)
+	void FGaussSeidelMainConstraints::InitTransientColor(const FSolverParticles& Particles)
+	{
+		InitTransientColor(FSolverParticlesRange(Particles));
+	}
+
+	void FGaussSeidelMainConstraints::InitDynamicColor(const FSolverParticlesRange& ParticlesRange)
 	{
+		ensureMsgf(StaticParticleColors.Num() == ParticlesRange.Size(), TEXT("InitStaticColor must be called before InitDynamicColor"));
 		PERF_SCOPE(STAT_ChaosGSMainConstraint_InitDynamicColor);
 		ParticleColors = StaticParticleColors;
 		ParticlesPerColor = StaticParticlesPerColor;
-		Chaos::ComputeExtraNodalColoring(StaticConstraints, DynamicConstraints, Particles, StaticIncidentElements, TransientIncidentElements, ParticleColors, ParticlesPerColor);
+		Chaos::ComputeExtraNodalColoring(StaticConstraints, DynamicConstraints, ParticlesRange, StaticIncidentElements, TransientIncidentElements, ParticleColors, ParticlesPerColor);
 	}
 
-	void FGaussSeidelMainConstraints::Init(const FSolverReal Dt, const FSolverParticles& Particles)
+	void FGaussSeidelMainConstraints::InitDynamicColor(const FSolverParticles& Particles)
 	{
-		Resize((int32)Particles.Size());
+		InitDynamicColor(FSolverParticlesRange(Particles));
+	}
+
+	void FGaussSeidelMainConstraints::Init(const FSolverReal Dt, const FSolverParticlesRange& ParticlesRange)
+	{
+		Resize((int32)ParticlesRange.Size());
 
 		PERF_SCOPE(STAT_ChaosGSMainConstraint_Init);
 		if (!bDoQuasistatics)
 		{
-			for (int32 i = 0; i < xtilde.Num(); i++)
+			for (int32 ParticleIdx = 0; ParticleIdx < xtilde.Num(); ParticleIdx++)
 			{
-				xtilde[i] = Particles.P(i); // Here P = xn + dt*(vn + dt*an) after evolution initial guess
+				xtilde[ParticleIdx] = ParticlesRange.P(ParticleIdx); // Here P = xn + dt*(vn + dt*an) after evolution initial guess
 			}
 		}
 		CurrentIt = 0;
 	}
 
+	void FGaussSeidelMainConstraints::Init(const FSolverReal Dt, const FSolverParticles& Particles)
+	{
+		Init(Dt, FSolverParticlesRange(Particles));
+	}
+
 	bool FGaussSeidelMainConstraints::IsClean(const TArray<TArray<int32>>& ConstraintsIn, const TArray<TArray<int32>>& IncidentElementsIn, const TArray<TArray<int32>>& IncidentElementsLocalIn)
 	{
 		if (IncidentElementsIn.Num() == IncidentElementsLocalIn.Num())
 		{
 			int32 TotalEntries = 0;
-			for (int32 i = 0; i < IncidentElementsIn.Num(); i++)
+			for (int32 ParticleIdx = 0; ParticleIdx < IncidentElementsIn.Num(); ParticleIdx++)
 			{
-				TotalEntries += IncidentElementsIn[i].Num();
-				if (IncidentElementsIn[i].Num() != IncidentElementsLocalIn[i].Num())
+				TotalEntries += IncidentElementsIn[ParticleIdx].Num();
+				if (IncidentElementsIn[ParticleIdx].Num() != IncidentElementsLocalIn[ParticleIdx].Num())
 				{
 					return false;
 				}
-				for (int32 j = 0; j < IncidentElementsIn[i].Num(); j++)
+				for (int32 IncidentIdx = 0; IncidentIdx < IncidentElementsIn[ParticleIdx].Num(); IncidentIdx++)
 				{
-					if (IncidentElementsIn[i][j] >= ConstraintsIn.Num() || IncidentElementsLocalIn[i][j] >= ConstraintsIn[IncidentElementsIn[i][j]].Num())
+					if (IncidentElementsIn[ParticleIdx][IncidentIdx] >= ConstraintsIn.Num() || IncidentElementsLocalIn[ParticleIdx][IncidentIdx] >= ConstraintsIn[IncidentElementsIn[ParticleIdx][IncidentIdx]].Num())
 					{
 						return false;
 					}
@@ -538,7 +609,7 @@ namespace Chaos::Softs
 		return false;
 	}
 
-	bool FGaussSeidelMainConstraints::ApplySingleParticle(const int32 ParticleIdx, const FSolverReal Dt, FSolverParticles& Particles)
+	bool FGaussSeidelMainConstraints::ApplySingleParticle(const int32 ParticleIdx, const FSolverReal Dt, FSolverParticlesRange& Particles)
 	{
 		int32 ConstraintIndex = 0;
 		Chaos::TVector<FSolverReal, 3> ParticleResidual((FSolverReal)0.);
@@ -546,47 +617,47 @@ namespace Chaos::Softs
 
 		ComputeInitialResidualAndHessian(Particles, ParticleIdx, Dt, ParticleResidual, ParticleHessian);
 
-		for (int32 i = 0; i < StaticIncidentElements[ParticleIdx].Num(); i++)
+		for (int32 IncidentIdx = 0; IncidentIdx < StaticIncidentElements[ParticleIdx].Num(); IncidentIdx++)
 		{
-			while (StaticIncidentElements[ParticleIdx][i] >= StaticIncidentElementsOffsets[ConstraintIndex + 1] && ConstraintIndex < StaticIncidentElementsOffsets.Num() - 1)
+			while (StaticIncidentElements[ParticleIdx][IncidentIdx] >= StaticIncidentElementsOffsets[ConstraintIndex + 1] && ConstraintIndex < StaticIncidentElementsOffsets.Num() - 1)
 			{
 				ConstraintIndex ++;
 			}
 
-			AddStaticConstraintResidualAndHessian[ConstraintIndex](Particles, StaticIncidentElements[ParticleIdx][i] - StaticIncidentElementsOffsets[ConstraintIndex], StaticIncidentElementsLocal[ParticleIdx][i], Dt, ParticleResidual, ParticleHessian);
+			AddStaticConstraintResidualAndHessian[ConstraintIndex](Particles, StaticIncidentElements[ParticleIdx][IncidentIdx] - StaticIncidentElementsOffsets[ConstraintIndex], StaticIncidentElementsLocal[ParticleIdx][IncidentIdx], Dt, ParticleResidual, ParticleHessian);
 		}
 
 		ConstraintIndex = 0;
 
-		for (int32 i = 0; i < DynamicIncidentElements[ParticleIdx].Num(); i++)
+		for (int32 IncidentIdx = 0; IncidentIdx < DynamicIncidentElements[ParticleIdx].Num(); IncidentIdx++)
 		{
-			while (DynamicIncidentElements[ParticleIdx][i] >= DynamicIncidentElementsOffsets[ConstraintIndex + 1] && ConstraintIndex < DynamicIncidentElementsOffsets.Num() - 1)
+			while (DynamicIncidentElements[ParticleIdx][IncidentIdx] >= DynamicIncidentElementsOffsets[ConstraintIndex + 1] && ConstraintIndex < DynamicIncidentElementsOffsets.Num() - 1)
 			{
 				ConstraintIndex ++;
 			}
 
-			AddDynamicConstraintResidualAndHessian[ConstraintIndex](Particles, DynamicIncidentElements[ParticleIdx][i] - DynamicIncidentElementsOffsets[ConstraintIndex], DynamicIncidentElementsLocal[ParticleIdx][i], Dt, ParticleResidual, ParticleHessian);
+			AddDynamicConstraintResidualAndHessian[ConstraintIndex](Particles, DynamicIncidentElements[ParticleIdx][IncidentIdx] - DynamicIncidentElementsOffsets[ConstraintIndex], DynamicIncidentElementsLocal[ParticleIdx][IncidentIdx], Dt, ParticleResidual, ParticleHessian);
 		}
 
 		ConstraintIndex = 0;
 
-		for (int32 i = 0; i < TransientIncidentElements[ParticleIdx].Num(); i++)
+		for (int32 IncidentIdx = 0; IncidentIdx < TransientIncidentElements[ParticleIdx].Num(); IncidentIdx++)
 		{
-			while (TransientIncidentElements[ParticleIdx][i] >= TransientIncidentElementsOffsets[ConstraintIndex + 1] && ConstraintIndex < TransientIncidentElementsOffsets.Num() - 1)
+			while (TransientIncidentElements[ParticleIdx][IncidentIdx] >= TransientIncidentElementsOffsets[ConstraintIndex + 1] && ConstraintIndex < TransientIncidentElementsOffsets.Num() - 1)
 			{
 				ConstraintIndex ++;
 			}
 
-			AddTransientConstraintResidualAndHessian[ConstraintIndex](Particles, TransientIncidentElements[ParticleIdx][i] - TransientIncidentElementsOffsets[ConstraintIndex], TransientIncidentElementsLocal[ParticleIdx][i], Dt, ParticleResidual, ParticleHessian);
+			AddTransientConstraintResidualAndHessian[ConstraintIndex](Particles, TransientIncidentElements[ParticleIdx][IncidentIdx] - TransientIncidentElementsOffsets[ConstraintIndex], TransientIncidentElementsLocal[ParticleIdx][IncidentIdx], Dt, ParticleResidual, ParticleHessian);
 		}
 
-		for (int32 i = 0; i < AddPerNodeHessian.Num(); i++)
+		for (int32 HessianIdx = 0; HessianIdx < AddPerNodeHessian.Num(); HessianIdx++)
 		{
-			AddPerNodeHessian[i](ParticleIdx, Dt, ParticleHessian);
+			AddPerNodeHessian[HessianIdx](ParticleIdx, Dt, ParticleHessian);
 		}
-		for (int32 i = 0; i < AddTransientPerNodeHessian.Num(); i++)
+		for (int32 HessianIdx = 0; HessianIdx < AddTransientPerNodeHessian.Num(); HessianIdx++)
 		{
-			AddTransientPerNodeHessian[i](ParticleIdx, Dt, ParticleHessian);
+			AddTransientPerNodeHessian[HessianIdx](ParticleIdx, Dt, ParticleHessian);
 		}
 		FSolverReal HessianScale = (FSolverReal)1.;
 		FSolverReal HessianDet = ParticleHessian.Determinant();
@@ -638,7 +709,7 @@ namespace Chaos::Softs
 
 	void FGaussSeidelMainConstraints::InitializeLambdas()
 	{
-		ComputeInitialResidualAndHessian = [this](const FSolverParticles& Particles, const int32 ParticleIdx, const FSolverReal Dt, TVec3<FSolverReal>& ParticleResidual, Chaos::PMatrix<FSolverReal, 3, 3>& ParticleHessian)
+		ComputeInitialResidualAndHessian = [this](const FSolverParticlesRange& Particles, const int32 ParticleIdx, const FSolverReal Dt, TVec3<FSolverReal>& ParticleResidual, Chaos::PMatrix<FSolverReal, 3, 3>& ParticleHessian)
 		{
 			if (!this->bDoQuasistatics) {
 				for (int32 Alpha = 0; Alpha < 3; Alpha++) {
@@ -656,7 +727,7 @@ namespace Chaos::Softs
 			}
 		};
 
-		AccelerationTechniquePerParticle = [this](FSolverParticles& Particles, int32 ParticleIndex)
+		AccelerationTechniquePerParticle = [this](FSolverParticlesRange& Particles, int32 ParticleIndex)
 		{
 			if (Particles.InvM(ParticleIndex) != FSolverReal(0) && CurrentIt > SORStart)
 			{
diff --git a/Engine/Source/Runtime/Experimental/Chaos/Private/Chaos/Deformable/GaussSeidelWeakConstraints.cpp b/Engine/Source/Runtime/Experimental/Chaos/Private/Chaos/Deformable/GaussSeidelWeakConstraints.cpp
--- a/Engine/Source/Runtime/Experimental/Chaos/Private/Chaos/Deformable/GaussSeidelWeakConstraints.cpp
+++ b/Engine/Source/Runtime/Experimental/Chaos/Private/Chaos/Deformable/GaussSeidelWeakConstraints.cpp
@@ -144,6 +144,35 @@ namespace Chaos::Softs
 		}
 	}
 
+	FGaussSeidelSpringConstraints::FGaussSeidelSpringConstraints(
+		const TArray<TArray<int32>>& InIndices,
+		const TArray<TArray<FSolverReal>>& InWeights,
+		const TArray<FSolverReal>& InStiffness,
+		const TArray<TArray<int32>>& InSecondIndices,
+		const TArray<TArray<FSolverReal>>& InSecondWeights,
+		const FDeformableXPBDWeightedSpringConstraintParams& InParams,
+		bool bIsAnisotropic,
+		bool bIsZeroRestLength
+	) : DebugDrawParams(InParams)
+	{
+		ensureMsgf(InIndices.Num() == InSecondIndices.Num(), TEXT("Input Double Bindings have wrong size"));
+
+		ConstraintsData.AddConstraints(InIndices.Num());
+
+		for (int32 SpringIdx = 0; SpringIdx < InIndices.Num(); ++SpringIdx)
+		{
+			FGaussSeidelSpringConstraintSingleData SingleConstraintData;
+			SingleConstraintData.SingleIndices = InIndices[SpringIdx];
+			SingleConstraintData.SingleSecondIndices = InSecondIndices[SpringIdx];
+			SingleConstraintData.SingleWeights = InWeights[SpringIdx];
+			SingleConstraintData.SingleSecondWeights = InSecondWeights[SpringIdx];
+			SingleConstraintData.SingleStiffness = InStiffness[SpringIdx];
+			SingleConstraintData.bIsAnisotropic = bIsAnisotropic;
+			SingleConstraintData.bIsZeroRestLength = bIsZeroRestLength;
+			ConstraintsData.SetSingleConstraint(SingleConstraintData, SpringIdx);
+		}
+	}
+
 PRAGMA_DISABLE_DEPRECATION_WARNINGS
 	FGaussSeidelSpringConstraints::FGaussSeidelSpringConstraints(
 		const TArray<TArray<int32>>& InIndices,
@@ -158,24 +187,24 @@ PRAGMA_DISABLE_DEPRECATION_WARNINGS
 
 		ConstraintsData.AddConstraints(InIndices.Num());
 
-		for (int32 i = 0; i < InIndices.Num(); i++)
+		for (int32 SpringIdx = 0; SpringIdx < InIndices.Num(); SpringIdx++)
 		{
 			FGaussSeidelSpringConstraintSingleData SingleConstraintData;
-			SingleConstraintData.SingleIndices = InIndices[i];
-			SingleConstraintData.SingleSecondIndices = InSecondIndices[i];
-			SingleConstraintData.SingleWeights = InWeights[i];
-			SingleConstraintData.SingleSecondWeights = InSecondWeights[i];
-			SingleConstraintData.SingleStiffness = InStiffness[i];
-			ConstraintsData.SetSingleConstraint(SingleConstraintData, i);
+			SingleConstraintData.SingleIndices = InIndices[SpringIdx];
+			SingleConstraintData.SingleSecondIndices = InSecondIndices[SpringIdx];
+			SingleConstraintData.SingleWeights = InWeights[SpringIdx];
+			SingleConstraintData.SingleSecondWeights = InSecondWeights[SpringIdx];
+			SingleConstraintData.SingleStiffness = InStiffness[SpringIdx];
+			ConstraintsData.SetSingleConstraint(SingleConstraintData, SpringIdx);
 		}
 
-		for (int32 i = 0; i < ConstraintsData.Size(); i++)
+		for (int32 SpringIdx = 0; SpringIdx < ConstraintsData.Size(); SpringIdx++)
 		{
-			const TArray<int32>& SingleIndices = ConstraintsData.GetIndices(i);
-			const TArray<int32>& SingleSecondIndices = ConstraintsData.GetSecondIndices(i);
-			for (int32 j = 0; j < SingleSecondIndices.Num(); j++)
+			const TArray<int32>& SingleIndices = ConstraintsData.GetIndices(SpringIdx);
+			const TArray<int32>& SingleSecondIndices = ConstraintsData.GetSecondIndices(SpringIdx);
+			for (int32 SecondIdx = 0; SecondIdx < SingleSecondIndices.Num(); SecondIdx++)
 			{
-				ensureMsgf(!SingleIndices.Contains(SingleSecondIndices[j]), TEXT("Indices and Second Indices overlaps. Currently not supported"));
+				ensureMsgf(!SingleIndices.Contains(SingleSecondIndices[SecondIdx]), TEXT("Indices and Second Indices overlaps. Currently not supported"));
 			}
 		}
 	}
@@ -192,7 +221,7 @@ PRAGMA_ENABLE_DEPRECATION_WARNINGS
 
 	void FGaussSeidelSpringConstraints::ComputeInitialWCData(const FSolverParticles& InParticles)
 	{
-		ComputeInitialWCData(FSolverParticlesRange(const_cast<FSolverParticles*>(&InParticles), /*InOffset = */ 0, /*InRange = */ InParticles.Size()));
+		ComputeInitialWCData(FSolverParticlesRange(InParticles));
 	}
 
 	void FGaussSeidelSpringConstraints::AddWCHessian(const int32 ParticleIdx, const FSolverReal Dt, Chaos::PMatrix<FSolverReal, 3, 3>& ParticleHessian) const
@@ -225,17 +254,17 @@ PRAGMA_ENABLE_DEPRECATION_WARNINGS
 
 		ConstraintsData.AddConstraints(InIndices.Num());
 
-		for (int32 i = 0; i < InIndices.Num(); i++)
+		for (int32 SpringIdx = 0; SpringIdx < InIndices.Num(); SpringIdx++)
 		{
 			FGaussSeidelSpringConstraintSingleData SingleConstraintData;
-			SingleConstraintData.SingleIndices = InIndices[i];
-			SingleConstraintData.SingleSecondIndices = InSecondIndices[i];
-			SingleConstraintData.SingleWeights = InWeights[i];
-			SingleConstraintData.SingleSecondWeights = InSecondWeights[i];
-			SingleConstraintData.SingleStiffness = InStiffness[i];
-			SingleConstraintData.bIsAnisotropic = InIsAnisotropic[i];
-			SingleConstraintData.bIsZeroRestLength = InIsZeroRestLength[i];
-			ConstraintsData.SetSingleConstraint(SingleConstraintData, i + Offset);
+			SingleConstraintData.SingleIndices = InIndices[SpringIdx];
+			SingleConstraintData.SingleSecondIndices = InSecondIndices[SpringIdx];
+			SingleConstraintData.SingleWeights = InWeights[SpringIdx];
+			SingleConstraintData.SingleSecondWeights = InSecondWeights[SpringIdx];
+			SingleConstraintData.SingleStiffness = InStiffness[SpringIdx];
+			SingleConstraintData.bIsAnisotropic = InIsAnisotropic[SpringIdx];
+			SingleConstraintData.bIsZeroRestLength = InIsZeroRestLength[SpringIdx];
+			ConstraintsData.SetSingleConstraint(SingleConstraintData, SpringIdx + Offset);
 		}
 	}
 
@@ -250,7 +279,7 @@ PRAGMA_ENABLE_DEPRECATION_WARNINGS
 
 	void FGaussSeidelSpringConstraints::Init(const FSolverParticles& InParticles, const FSolverReal Dt)
 	{
-		Init(FSolverParticlesRange(const_cast<FSolverParticles*>(&InParticles), /*InOffset = */ 0, /*InRange = */ InParticles.Size()), Dt);
+		Init(FSolverParticlesRange(InParticles), Dt);
 	}
 
 	void FGaussSeidelSpringConstraints::CollisionDetectionBVH(const FSolverParticles& Particles, const TArray<TVec3<int32>>& SurfaceElements, const TArray<int32>& ComponentIndex, float DetectRadius, float PositionTargetStiffness, bool UseAnisotropicSpring)
@@ -259,12 +288,12 @@ PRAGMA_ENABLE_DEPRECATION_WARNINGS
 		Resize(InitialWCSize);
 
 		TArray<Chaos::TVector<int32, 3>> SurfaceElementsArray;
-		for (int32 i = 0; i < SurfaceElements.Num(); i++)
+		for (int32 ElemIdx = 0; ElemIdx < SurfaceElements.Num(); ElemIdx++)
 		{
 			Chaos::TVector<int32, 3> CurrentSurfaceElements(0);
-			for (int32 j = 0; j < 3; j++)
+			for (int32 TriLocalIdx = 0; TriLocalIdx < 3; TriLocalIdx++)
 			{
-				CurrentSurfaceElements[j] = SurfaceElements[i][j];
+				CurrentSurfaceElements[TriLocalIdx] = SurfaceElements[ElemIdx][TriLocalIdx];
 			}
 			if (CurrentSurfaceElements[0] != INDEX_NONE
 				&& CurrentSurfaceElements[1] != INDEX_NONE
@@ -277,9 +306,9 @@ PRAGMA_ENABLE_DEPRECATION_WARNINGS
 		TArray<TArray<int32>>* LocalIndexPtr = &LocalIndex;
 		TArray<TArray<int>> GlobalIndex = Chaos::Utilities::ComputeIncidentElements(SurfaceElementsArray, LocalIndexPtr);
 		int32 ActualParticleCount = 0;
-		for (int32 l = 0; l < GlobalIndex.Num(); l++)
+		for (int32 ParticleIdx = 0; ParticleIdx < GlobalIndex.Num(); ParticleIdx++)
 		{
-			if (GlobalIndex[l].Num() > 0)
+			if (GlobalIndex[ParticleIdx].Num() > 0)
 			{
 				ActualParticleCount += 1;
 			}
@@ -289,12 +318,12 @@ PRAGMA_ENABLE_DEPRECATION_WARNINGS
 		TArray<int32> SurfaceElementsMap;
 		SurfaceElementsMap.SetNum(ActualParticleCount);
 		int32 CurrentParticleIndex = 0;
-		for (int32 i = 0; i < GlobalIndex.Num(); i++)
+		for (int32 ParticleIdx = 0; ParticleIdx < GlobalIndex.Num(); ParticleIdx++)
 		{
-			if (GlobalIndex[i].Num() > 0)
+			if (GlobalIndex[ParticleIdx].Num() > 0)
 			{
-				SurfaceElementsPositions[CurrentParticleIndex] = Particles.P(SurfaceElements[GlobalIndex[i][0]][LocalIndex[i][0]]);
-				SurfaceElementsMap[CurrentParticleIndex] = SurfaceElements[GlobalIndex[i][0]][LocalIndex[i][0]];
+				SurfaceElementsPositions[CurrentParticleIndex] = Particles.P(SurfaceElements[GlobalIndex[ParticleIdx][0]][LocalIndex[ParticleIdx][0]]);
+				SurfaceElementsMap[CurrentParticleIndex] = SurfaceElements[GlobalIndex[ParticleIdx][0]][LocalIndex[ParticleIdx][0]];
 				CurrentParticleIndex += 1;
 			}
 		}
@@ -305,12 +334,12 @@ PRAGMA_ENABLE_DEPRECATION_WARNINGS
 		VertexSpheres.Init(Chaos::FSphere(Chaos::TVec3<Chaos::FReal>(0), DetectRadius), SurfaceElementsPositions.Num());
 		VertexSpherePtrs.SetNum(SurfaceElementsPositions.Num());
 
-		for (int32 i = 0; i < SurfaceElementsPositions.Num(); i++)
+		for (int32 SphereIdx = 0; SphereIdx < SurfaceElementsPositions.Num(); SphereIdx++)
 		{
-			Chaos::TVec3<Chaos::FReal> SphereCenter(SurfaceElementsPositions[i]);
+			Chaos::TVec3<Chaos::FReal> SphereCenter(SurfaceElementsPositions[SphereIdx]);
 			Chaos::FSphere VertexSphere(SphereCenter, DetectRadius);
-			VertexSpheres[i] = Chaos::FSphere(SphereCenter, DetectRadius);
-			VertexSpherePtrs[i] = &VertexSpheres[i];
+			VertexSpheres[SphereIdx] = Chaos::FSphere(SphereCenter, DetectRadius);
+			VertexSpherePtrs[SphereIdx] = &VertexSpheres[SphereIdx];
 		}
 		Chaos::TBoundingVolumeHierarchy<
 			TArray<Chaos::FSphere*>,
@@ -318,37 +347,37 @@ PRAGMA_ENABLE_DEPRECATION_WARNINGS
 			Chaos::FReal,
 			3> VertexBVH(VertexSpherePtrs);
 
-		for (int32 i = 0; i < SurfaceElements.Num(); i++)
+		for (int32 ElemIdx = 0; ElemIdx < SurfaceElements.Num(); ElemIdx++)
 		{
-			TArray<int32> TriangleIntersections0 = VertexBVH.FindAllIntersections(Particles.P(SurfaceElements[i][0]));
-			TArray<int32> TriangleIntersections1 = VertexBVH.FindAllIntersections(Particles.P(SurfaceElements[i][1]));
-			TArray<int32> TriangleIntersections2 = VertexBVH.FindAllIntersections(Particles.P(SurfaceElements[i][2]));
+			TArray<int32> TriangleIntersections0 = VertexBVH.FindAllIntersections(Particles.P(SurfaceElements[ElemIdx][0]));
+			TArray<int32> TriangleIntersections1 = VertexBVH.FindAllIntersections(Particles.P(SurfaceElements[ElemIdx][1]));
+			TArray<int32> TriangleIntersections2 = VertexBVH.FindAllIntersections(Particles.P(SurfaceElements[ElemIdx][2]));
 			TriangleIntersections0.Sort();
 			TriangleIntersections1.Sort();
 			TriangleIntersections2.Sort();
 
 			TArray<int32> TriangleIntersections({});
-			for (int32 k = 0; k < TriangleIntersections0.Num(); k++)
+			for (int32 IntersectionIdx = 0; IntersectionIdx < TriangleIntersections0.Num(); IntersectionIdx++)
 			{
-				if (TriangleIntersections1.Contains(TriangleIntersections0[k])
-					&& TriangleIntersections2.Contains(TriangleIntersections0[k]))
+				if (TriangleIntersections1.Contains(TriangleIntersections0[IntersectionIdx])
+					&& TriangleIntersections2.Contains(TriangleIntersections0[IntersectionIdx]))
 				{
-					TriangleIntersections.Emplace(TriangleIntersections0[k]);
+					TriangleIntersections.Emplace(TriangleIntersections0[IntersectionIdx]);
 				}
 			}
 
-			int32 TriangleIndex = ComponentIndex[SurfaceElements[i][0]];
+			const int32 TriangleIndex = ComponentIndex[SurfaceElements[ElemIdx][0]];
 			int32 MinIndex = INDEX_NONE;
 			float MinDis = DetectRadius;
 			Chaos::TVector<float, 3> ClosestBary(0.f);
 			Chaos::TVector<float, 3> FaceNormal(0.f);
-			for (int32 j = 0; j < TriangleIntersections.Num(); j++)
+			for (int32 IntersectionIdx = 0; IntersectionIdx < TriangleIntersections.Num(); IntersectionIdx++)
 			{
-				if (ComponentIndex[SurfaceElementsMap[TriangleIntersections[j]]] >= 0 && TriangleIndex >= 0 && ComponentIndex[SurfaceElementsMap[TriangleIntersections[j]]] != TriangleIndex)
+				if (ComponentIndex[SurfaceElementsMap[TriangleIntersections[IntersectionIdx]]] >= 0 && TriangleIndex >= 0 && ComponentIndex[SurfaceElementsMap[TriangleIntersections[IntersectionIdx]]] != TriangleIndex)
 				{
-					Chaos::TVector<float, 3> Bary, TriPos0(Particles.P(SurfaceElements[i][0])), TriPos1(Particles.P(SurfaceElements[i][1])), TriPos2(Particles.P(SurfaceElements[i][2])), ParticlePos(Particles.P(SurfaceElementsMap[TriangleIntersections[j]]));
+					Chaos::TVector<float, 3> Bary, TriPos0(Particles.P(SurfaceElements[ElemIdx][0])), TriPos1(Particles.P(SurfaceElements[ElemIdx][1])), TriPos2(Particles.P(SurfaceElements[ElemIdx][2])), ParticlePos(Particles.P(SurfaceElementsMap[TriangleIntersections[IntersectionIdx]]));
 					Chaos::TVector<Chaos::FRealSingle, 3> ClosestPoint = Chaos::FindClosestPointAndBaryOnTriangle(TriPos0, TriPos1, TriPos2, ParticlePos, Bary);
-					Chaos::FRealSingle CurrentDistance = (Particles.P(SurfaceElementsMap[TriangleIntersections[j]]) - ClosestPoint).Size();
+					const Chaos::FRealSingle CurrentDistance = (Particles.P(SurfaceElementsMap[TriangleIntersections[IntersectionIdx]]) - ClosestPoint).Size();
 					if (CurrentDistance < MinDis)
 					{
 						Chaos::TVector<FSolverReal, 3> Normal = FVector3f::CrossProduct(TriPos2 - TriPos0, TriPos1 - TriPos0); //The normal needs to point outwards of the geometry
@@ -356,7 +385,7 @@ PRAGMA_ENABLE_DEPRECATION_WARNINGS
 						{
 							Normal.SafeNormalize(1e-8f);
 							MinDis = CurrentDistance;
-							MinIndex = SurfaceElementsMap[TriangleIntersections[j]];
+							MinIndex = SurfaceElementsMap[TriangleIntersections[IntersectionIdx]];
 							ClosestBary = Bary;
 							FaceNormal = Normal;
 						}
@@ -365,12 +394,12 @@ PRAGMA_ENABLE_DEPRECATION_WARNINGS
 				}
 			}
 			if (MinIndex != INDEX_NONE
-				&& MinIndex != SurfaceElements[i][0]
-				&& MinIndex != SurfaceElements[i][1]
-				&& MinIndex != SurfaceElements[i][2])
+				&& MinIndex != SurfaceElements[ElemIdx][0]
+				&& MinIndex != SurfaceElements[ElemIdx][1]
+				&& MinIndex != SurfaceElements[ElemIdx][2])
 			{
 				FGaussSeidelSpringConstraintSingleData SingleConstraintData;
-				SingleConstraintData.SingleIndices = {SurfaceElements[i][0], SurfaceElements[i][1], SurfaceElements[i][2]};
+				SingleConstraintData.SingleIndices = {SurfaceElements[ElemIdx][0], SurfaceElements[ElemIdx][1], SurfaceElements[ElemIdx][2]};
 				SingleConstraintData.SingleSecondIndices = {MinIndex};
 				SingleConstraintData.SingleWeights = {ClosestBary[0], ClosestBary[1], ClosestBary[2]};
 				SingleConstraintData.SingleSecondWeights = {FSolverReal(1.f)};
@@ -379,9 +408,9 @@ PRAGMA_ENABLE_DEPRECATION_WARNINGS
 				SingleConstraintData.bIsZeroRestLength = true; //Push-out type collision springs should be zero rest length
 
 				float SpringStiffness = 0.f;
-				for (int32 k = 0; k < 3; k++)
+				for (int32 TriLocalIdx = 0; TriLocalIdx < 3; TriLocalIdx++)
 				{
-					SpringStiffness += ClosestBary[k] * PositionTargetStiffness * Particles.M(SurfaceElements[i][k]);
+					SpringStiffness += ClosestBary[TriLocalIdx] * PositionTargetStiffness * Particles.M(SurfaceElements[ElemIdx][TriLocalIdx]);
 				}
 				SpringStiffness += PositionTargetStiffness * Particles.M(MinIndex);
 				SingleConstraintData.SingleStiffness = (FSolverReal)SpringStiffness;
@@ -395,38 +424,38 @@ PRAGMA_ENABLE_DEPRECATION_WARNINGS
 		ensureMsgf(ConstraintsData.Size() >= InitialWCSize, TEXT("The size of Indices is smaller than InitialWCSize"));
 
 		ExtraConstraints.Init(TArray<int32>(), ConstraintsData.Size() - InitialWCSize);
-		for (int32 i = InitialWCSize; i < static_cast<int32>(ConstraintsData.Size()); i++)
+		for (int32 ConstraintIdx = InitialWCSize; ConstraintIdx < static_cast<int32>(ConstraintsData.Size()); ConstraintIdx++)
 		{
-			const TArray<int32>& LocalIndices = ConstraintsData.GetIndices(i);
-			const TArray<int32>& LocalSecondIndices = ConstraintsData.GetSecondIndices(i);
+			const TArray<int32>& LocalIndices = ConstraintsData.GetIndices(ConstraintIdx);
+			const TArray<int32>& LocalSecondIndices = ConstraintsData.GetSecondIndices(ConstraintIdx);
 
-			ExtraConstraints[i - InitialWCSize].SetNum(LocalIndices.Num() + LocalSecondIndices.Num());
-			for (int32 j = 0; j < LocalIndices.Num(); j++)
+			ExtraConstraints[ConstraintIdx - InitialWCSize].SetNum(LocalIndices.Num() + LocalSecondIndices.Num());
+			for (int32 LocalIdx = 0; LocalIdx < LocalIndices.Num(); LocalIdx++)
 			{
-				ExtraConstraints[i - InitialWCSize][j] = LocalIndices[j];
+				ExtraConstraints[ConstraintIdx - InitialWCSize][LocalIdx] = LocalIndices[LocalIdx];
 			}
-			for (int32 j = 0; j < LocalSecondIndices.Num(); j++)
+			for (int32 SecondLocalIdx = 0; SecondLocalIdx < LocalSecondIndices.Num(); SecondLocalIdx++)
 			{
-				ExtraConstraints[i - InitialWCSize][j + LocalIndices.Num()] = LocalSecondIndices[j];
+				ExtraConstraints[ConstraintIdx - InitialWCSize][SecondLocalIdx + LocalIndices.Num()] = LocalSecondIndices[SecondLocalIdx];
 			}
 		}
 
 		ExtraWCIncidentElements = Chaos::Utilities::ComputeIncidentElements(ExtraConstraints, &ExtraWCIncidentElementsLocal);
 
 		NodalWeights = NoCollisionNodalWeights;
-		for (int32 i = 0; i < ExtraWCIncidentElements.Num(); i++)
+		for (int32 IncidentIdx = 0; IncidentIdx < ExtraWCIncidentElements.Num(); IncidentIdx++)
 		{
-			if (ExtraWCIncidentElements[i].Num() > 0)
+			if (ExtraWCIncidentElements[IncidentIdx].Num() > 0)
 			{
-				int32 ParticleIdx = ExtraConstraints[ExtraWCIncidentElements[i][0]][ExtraWCIncidentElementsLocal[i][0]];
+				const int32 ParticleIdx = ExtraConstraints[ExtraWCIncidentElements[IncidentIdx][0]][ExtraWCIncidentElementsLocal[IncidentIdx][0]];
 				if (NodalWeights[ParticleIdx].Num() == 0)
 				{
 					NodalWeights[ParticleIdx].Init(FSolverReal(0), 6);
 				}
-				for (int32 j = 0; j < ExtraWCIncidentElements[i].Num(); j++)
+				for (int32 LocalIdx = 0; LocalIdx < ExtraWCIncidentElements[IncidentIdx].Num(); LocalIdx++)
 				{
-					int32 LocalIndex = ExtraWCIncidentElementsLocal[i][j];
-					int32 ConstraintIndex = ExtraWCIncidentElements[i][j] + InitialWCSize;
+					const int32 LocalIndex = ExtraWCIncidentElementsLocal[IncidentIdx][LocalIdx];
+					const int32 ConstraintIndex = ExtraWCIncidentElements[IncidentIdx][LocalIdx] + InitialWCSize;
 
 					const FGaussSeidelSpringConstraintSingleData& SingleData = ConstraintsData.GetSingleConstraintData(ConstraintIndex);
 
@@ -464,6 +493,13 @@ PRAGMA_ENABLE_DEPRECATION_WARNINGS
 		}
 	}
 
+	const TArray<TArray<int32>>& FGaussSeidelSpringConstraints::GetStaticConstraintArrays(const TArray<TArray<int32>>*& OutIncidentElements, const TArray<TArray<int32>>*& OutIncidentElementsLocal) const
+	{
+		OutIncidentElements = &NoCollisionWCIncidentElements;
+		OutIncidentElementsLocal = &NoCollisionWCIncidentElementsLocal;
+		return NoCollisionConstraints;
+	}
+
 	const TArray<TArray<int32>>& FGaussSeidelSpringConstraints::GetStaticConstraintArrays(TArray<TArray<int32>>& IncidentElements, TArray<TArray<int32>>& IncidentElementsLocal) const
 	{
 		IncidentElements = NoCollisionWCIncidentElements;
@@ -475,18 +511,18 @@ PRAGMA_ENABLE_DEPRECATION_WARNINGS
 	{
 		TArray<TArray<int32>> ExtraConstraints;
 		ExtraConstraints.Init(TArray<int32>(), ConstraintsData.Size());
-		for (int32 i = InitialWCSize; i < ConstraintsData.Size(); i++)
+		for (int32 ConstraintIdx = InitialWCSize; ConstraintIdx < ConstraintsData.Size(); ConstraintIdx++)
 		{
-			const TArray<int32>& LocalIndices = ConstraintsData.GetIndices(i);
-			const TArray<int32>& LocalSecondIndices = ConstraintsData.GetSecondIndices(i);
-			ExtraConstraints[i - InitialWCSize].SetNum(LocalIndices.Num() + LocalSecondIndices.Num());
-			for (int32 j = 0; j < LocalIndices.Num(); j++)
+			const TArray<int32>& LocalIndices = ConstraintsData.GetIndices(ConstraintIdx);
+			const TArray<int32>& LocalSecondIndices = ConstraintsData.GetSecondIndices(ConstraintIdx);
+			ExtraConstraints[ConstraintIdx - InitialWCSize].SetNum(LocalIndices.Num() + LocalSecondIndices.Num());
+			for (int32 LocalIdx = 0; LocalIdx < LocalIndices.Num(); LocalIdx++)
 			{
-				ExtraConstraints[i - InitialWCSize][j] = LocalIndices[j];
+				ExtraConstraints[ConstraintIdx - InitialWCSize][LocalIdx] = LocalIndices[LocalIdx];
 			}
-			for (int32 j = 0; j < LocalSecondIndices.Num(); j++)
+			for (int32 SecondLocalIdx = 0; SecondLocalIdx < LocalSecondIndices.Num(); SecondLocalIdx++)
 			{
-				ExtraConstraints[i - InitialWCSize][j + LocalIndices.Num()] = LocalSecondIndices[j];
+				ExtraConstraints[ConstraintIdx - InitialWCSize][SecondLocalIdx + LocalIndices.Num()] = LocalSecondIndices[SecondLocalIdx];
 			}
 		}
 
@@ -513,8 +549,8 @@ PRAGMA_DISABLE_DEPRECATION_WARNINGS
 		}
 		if (SingleData.bIsAnisotropic)
 		{
-			FSolverReal Comp = TVec3<FSolverReal>::DotProduct(SpringEdge, SingleData.SingleNormal);
-			TVec3<FSolverReal> Proj = SingleData.SingleNormal * Comp;
+			const FSolverReal Comp = TVec3<FSolverReal>::DotProduct(SpringEdge, SingleData.SingleNormal);
+			const TVec3<FSolverReal> Proj = SingleData.SingleNormal * Comp;
 			for (int32 Alpha = 0; Alpha < 3; Alpha++)
 			{
 				ParticleResidual[Alpha] += Dt * Dt * SingleData.SingleStiffness * Proj[Alpha] * Weight;
@@ -548,7 +584,7 @@ PRAGMA_ENABLE_DEPRECATION_WARNINGS
 		TVec3<FSolverReal> Projection;
 		if (SingleData.bIsAnisotropic)
 		{
-			FSolverReal LengthDiff = TVec3<FSolverReal>::DotProduct(SpringEdge, SingleData.SingleNormal) - SingleData.RestLength;
+			const FSolverReal LengthDiff = TVec3<FSolverReal>::DotProduct(SpringEdge, SingleData.SingleNormal) - SingleData.RestLength;
 			Projection = SingleData.SingleNormal * LengthDiff;
 		}
 		else
@@ -567,7 +603,7 @@ PRAGMA_ENABLE_DEPRECATION_WARNINGS
 
 	void FGaussSeidelSpringConstraints::AddWCResidual(const FSolverParticles& InParticles, const int32 ConstraintIndex, const int32 LocalIndex, const FSolverReal Dt, TVec3<FSolverReal>& ParticleResidual, Chaos::PMatrix<FSolverReal, 3, 3>& ParticleHessian) const
 	{
-		AddWCResidual(FSolverParticlesRange(const_cast<FSolverParticles*>(&InParticles), /*InOffset = */ 0, /*InRange = */ InParticles.Size()), ConstraintIndex, LocalIndex, Dt, ParticleResidual, ParticleHessian);
+		AddWCResidual(FSolverParticlesRange(InParticles), ConstraintIndex, LocalIndex, Dt, ParticleResidual, ParticleHessian);
 	}
 
 	void FGaussSeidelSpringConstraints::SetConstraintsRestLength(const FSolverParticlesRange& InParticles, const FSolverParticlesRange& InSecondParticles)
@@ -631,8 +667,8 @@ PRAGMA_ENABLE_DEPRECATION_WARNINGS
 		const TArray<int32>& LocalIndices, const TArray<int32>& LocalSecondIndices,
 		const TArray<FSolverReal>& Weight, const TArray<FSolverReal>& SecondWeight, bool bUseParticleX) const
 	{
-		return ComputeSpringEdge(FSolverParticlesRange(const_cast<FSolverParticles*>(&InParticles), /*InOffset = */ 0, /*InRange = */ InParticles.Size()),
-			FSolverParticlesRange(const_cast<FSolverParticles*>(&InSecondParticles), /*InOffset = */ 0, /*InRange = */ InSecondParticles.Size()),
+		return ComputeSpringEdge(FSolverParticlesRange(InParticles),
+			FSolverParticlesRange(InSecondParticles),
 			LocalIndices, LocalSecondIndices, Weight, SecondWeight, bUseParticleX);
 	}
 
@@ -646,7 +682,7 @@ PRAGMA_ENABLE_DEPRECATION_WARNINGS
 	TVec3<FSolverReal> FGaussSeidelSpringConstraints::ComputeSpringEdge(const FSolverParticles& InParticles, const TArray<int32>& LocalIndices,
 		const TArray<int32>& LocalSecondIndices, const TArray<FSolverReal>& Weight, const TArray<FSolverReal>& SecondWeight, bool bUseParticleX) const
 	{
-		return ComputeSpringEdge(FSolverParticlesRange(const_cast<FSolverParticles*>(&InParticles), /*InOffset = */ 0, /*InRange = */ InParticles.Size()),
+		return ComputeSpringEdge(FSolverParticlesRange(InParticles),
 			LocalIndices, LocalSecondIndices, Weight, SecondWeight, bUseParticleX);
 	}
 
@@ -684,12 +720,12 @@ PRAGMA_ENABLE_DEPRECATION_WARNINGS
 			{
 				return bUseParticleX ? InParticles.X(Idx) : InParticles.P(Idx);
 			};
-		for (int32 i = 0; i < ConstraintsData.Size(); i++)
+		for (int32 ConstraintIdx = 0; ConstraintIdx < ConstraintsData.Size(); ConstraintIdx++)
 		{
-			if (ConstraintsData.GetIsAnisotropic(i))
+			if (ConstraintsData.GetIsAnisotropic(ConstraintIdx))
 			{
-				const TArray<int32>& IndicesTemp = ConstraintsData.GetIndices(i);
-				const TArray<int32>& SecondIndicesTemp = ConstraintsData.GetSecondIndices(i);
+				const TArray<int32>& IndicesTemp = ConstraintsData.GetIndices(ConstraintIdx);
+				const TArray<int32>& SecondIndicesTemp = ConstraintsData.GetSecondIndices(ConstraintIdx);
 				Chaos::TVector<FSolverReal, 3> Normal((FSolverReal)0.);
 				if (IndicesTemp.Num() == 3)
 				{
@@ -705,7 +741,7 @@ PRAGMA_ENABLE_DEPRECATION_WARNINGS
 						TriPos2(GetPosition(InSecondParticles, SecondIndicesTemp[2]));
 					Normal = FVector3f::CrossProduct(TriPos2 - TriPos0, TriPos1 - TriPos0).GetSafeNormal(); //triangle normal convention (see FTriangleMesh::GetFaceNormals())
 				}
-				ConstraintsData.SetNormal(i, Normal);
+				ConstraintsData.SetNormal(ConstraintIdx, Normal);
 			}
 		}
 	}
@@ -718,10 +754,10 @@ PRAGMA_ENABLE_DEPRECATION_WARNINGS
 			if (WCIncidentElements[ParticleIdx].Num() > 0)
 			{
 				NodalWeights[ParticleIdx].Init(FSolverReal(0), 6);
-				for (int32 j = 0; j < WCIncidentElements[ParticleIdx].Num(); j++)
+				for (int32 LocalIdx = 0; LocalIdx < WCIncidentElements[ParticleIdx].Num(); LocalIdx++)
 				{
-					int32 ConstraintIndex = WCIncidentElements[ParticleIdx][j];
-					int32 LocalIndex = WCIncidentElementsLocal[ParticleIdx][j];
+					const int32 ConstraintIndex = WCIncidentElements[ParticleIdx][LocalIdx];
+					const int32 LocalIndex = WCIncidentElementsLocal[ParticleIdx][LocalIdx];
 
 					FSolverReal Weight = FSolverReal(0);
 					if (LocalIndex >= ConstraintsData.GetIndices(ConstraintIndex).Num())
@@ -774,48 +810,48 @@ PRAGMA_ENABLE_DEPRECATION_WARNINGS
 
 		ConstraintsData.Resize(InitialWCSize);
 
-		for (int32 i = InitialWCSize; i < OriginalConstraintsData.Size(); i++)
+		for (int32 ConstraintIdx = InitialWCSize; ConstraintIdx < OriginalConstraintsData.Size(); ConstraintIdx++)
 		{
-			const TArray<int32>& IndicesTemp = OriginalConstraintsData.GetIndices(i);
-			const TArray<int32>& SecondIndicesTemp = OriginalConstraintsData.GetSecondIndices(i);
-			ensureMsgf(OriginalConstraintsData.GetIndices(i).Num() == 3, TEXT("Collision format is not point-triangle"));
-			ensureMsgf(OriginalConstraintsData.GetSecondIndices(i).Num() == 1, TEXT("Collision format is not point-triangle"));
+			const TArray<int32>& IndicesTemp = OriginalConstraintsData.GetIndices(ConstraintIdx);
+			const TArray<int32>& SecondIndicesTemp = OriginalConstraintsData.GetSecondIndices(ConstraintIdx);
+			ensureMsgf(OriginalConstraintsData.GetIndices(ConstraintIdx).Num() == 3, TEXT("Collision format is not point-triangle"));
+			ensureMsgf(OriginalConstraintsData.GetSecondIndices(ConstraintIdx).Num() == 1, TEXT("Collision format is not point-triangle"));
 			Chaos::TVector<float, 3> TriPos0(Particles.P(IndicesTemp[0])), TriPos1(Particles.P(IndicesTemp[1])), TriPos2(Particles.P(IndicesTemp[2])), ParticlePos(Particles.P(SecondIndicesTemp[0]));
 			Chaos::TVector<FSolverReal, 3> Normal = FVector3f::CrossProduct(TriPos2 - TriPos0, TriPos1 - TriPos0); //triangle normal convention (see FTriangleMesh::GetFaceNormals())
 			if (FVector3f::DotProduct(ParticlePos - TriPos0, Normal) < 0.f) //not resolved, keep the spring
 			{
-				ConstraintsData.AddSingleConstraint(OriginalConstraintsData.GetSingleConstraintData(i));
+				ConstraintsData.AddSingleConstraint(OriginalConstraintsData.GetSingleConstraintData(ConstraintIdx));
 			}
 		}
 	}
 
 	void FGaussSeidelSpringConstraints::VisualizeAllBindings(const FSolverParticlesRange& InParticles, const FSolverReal Dt) const
 	{
-#if WITH_EDITOR
-		auto DoubleVert = [](Chaos::TVec3<FSolverReal> V) { return FVector3d(V.X, V.Y, V.Z); };
-		for (int32 i = 0; i < ConstraintsData.Size(); i++)
+#if CHAOS_DEBUG_DRAW 
+		auto DoubleVert = [](const Chaos::TVec3<FSolverReal>& V) { return FVector3d(V.X, V.Y, V.Z); };
+		for (int32 ConstraintIdx = 0; ConstraintIdx < ConstraintsData.Size(); ConstraintIdx++)
 		{
-			const FGaussSeidelSpringConstraintSingleData& SingleConstraintData = ConstraintsData.GetSingleConstraintData(i);
+			const FGaussSeidelSpringConstraintSingleData& SingleConstraintData = ConstraintsData.GetSingleConstraintData(ConstraintIdx);
 			Chaos::TVec3<FSolverReal> SourcePos((FSolverReal)0.), TargetPos((FSolverReal)0.);
-			for (int32 j = 0; j < SingleConstraintData.SingleIndices.Num(); j++)
+			for (int32 Idx = 0; Idx < SingleConstraintData.SingleIndices.Num(); Idx++)
 			{
-				SourcePos += SingleConstraintData.SingleWeights[j] * InParticles.P(SingleConstraintData.SingleIndices[j]);
+				SourcePos += SingleConstraintData.SingleWeights[Idx] * InParticles.P(SingleConstraintData.SingleIndices[Idx]);
 			}
-			for (int32 j = 0; j < SingleConstraintData.SingleSecondIndices.Num(); j++)
+			for (int32 Idx = 0; Idx < SingleConstraintData.SingleSecondIndices.Num(); Idx++)
 			{
-				TargetPos += SingleConstraintData.SingleSecondWeights[j] * InParticles.P(SingleConstraintData.SingleSecondIndices[j]);
+				TargetPos += SingleConstraintData.SingleSecondWeights[Idx] * InParticles.P(SingleConstraintData.SingleSecondIndices[Idx]);
 			}
 
-			float ParticleThickness = DebugDrawParams.DebugParticleWidth;
-			float LineThickness = DebugDrawParams.DebugLineWidth;
+			const float ParticleThickness = DebugDrawParams.DebugParticleWidth;
+			const float LineThickness = DebugDrawParams.DebugLineWidth;
 
 			if (SingleConstraintData.SingleIndices.Num() == 1)
 			{
 				Chaos::FDebugDrawQueue::GetInstance().DrawDebugPoint(DoubleVert(SourcePos), FColor::Red, false, Dt, 0, ParticleThickness);
-				for (int32 j = 0; j < SingleConstraintData.SingleSecondIndices.Num(); j++)
+				for (int32 Idx = 0; Idx < SingleConstraintData.SingleSecondIndices.Num(); Idx++)
 				{
-					Chaos::FDebugDrawQueue::GetInstance().DrawDebugPoint(DoubleVert(InParticles.P(SingleConstraintData.SingleSecondIndices[j])), FColor::Green, false, Dt, 0, ParticleThickness);
-					Chaos::FDebugDrawQueue::GetInstance().DrawDebugLine(DoubleVert(InParticles.P(SingleConstraintData.SingleSecondIndices[j])), DoubleVert(InParticles.P(SingleConstraintData.SingleSecondIndices[(j + 1) % SingleConstraintData.SingleSecondIndices.Num()])), FColor::Green, false, Dt, 0, LineThickness);
+					Chaos::FDebugDrawQueue::GetInstance().DrawDebugPoint(DoubleVert(InParticles.P(SingleConstraintData.SingleSecondIndices[Idx])), FColor::Green, false, Dt, 0, ParticleThickness);
+					Chaos::FDebugDrawQueue::GetInstance().DrawDebugLine(DoubleVert(InParticles.P(SingleConstraintData.SingleSecondIndices[Idx])), DoubleVert(InParticles.P(SingleConstraintData.SingleSecondIndices[(Idx + 1) % SingleConstraintData.SingleSecondIndices.Num()])), FColor::Green, false, Dt, 0, LineThickness);
 				}
 
 			}
@@ -823,10 +859,10 @@ PRAGMA_ENABLE_DEPRECATION_WARNINGS
 			if (SingleConstraintData.SingleSecondIndices.Num() == 1)
 			{
 				Chaos::FDebugDrawQueue::GetInstance().DrawDebugPoint(DoubleVert(TargetPos), FColor::Red, false, Dt, 0, ParticleThickness);
-				for (int32 j = 0; j < SingleConstraintData.SingleIndices.Num(); j++)
+				for (int32 Idx = 0; Idx < SingleConstraintData.SingleIndices.Num(); Idx++)
 				{
-					Chaos::FDebugDrawQueue::GetInstance().DrawDebugPoint(DoubleVert(InParticles.P(SingleConstraintData.SingleIndices[j])), FColor::Green, false, Dt, 0, ParticleThickness);
-					Chaos::FDebugDrawQueue::GetInstance().DrawDebugLine(DoubleVert(InParticles.P(SingleConstraintData.SingleIndices[j])), DoubleVert(InParticles.P(SingleConstraintData.SingleIndices[(j + 1) % SingleConstraintData.SingleIndices.Num()])), FColor::Green, false, Dt, 0, LineThickness);
+					Chaos::FDebugDrawQueue::GetInstance().DrawDebugPoint(DoubleVert(InParticles.P(SingleConstraintData.SingleIndices[Idx])), FColor::Green, false, Dt, 0, ParticleThickness);
+					Chaos::FDebugDrawQueue::GetInstance().DrawDebugLine(DoubleVert(InParticles.P(SingleConstraintData.SingleIndices[Idx])), DoubleVert(InParticles.P(SingleConstraintData.SingleIndices[(Idx + 1) % SingleConstraintData.SingleIndices.Num()])), FColor::Green, false, Dt, 0, LineThickness);
 				}
 			}
 
@@ -848,9 +884,9 @@ PRAGMA_ENABLE_DEPRECATION_WARNINGS
 		const TArray<TVec3<int32>>& Elements = TriangleMesh.GetSurfaceElements();
 		const float HalfRadius = DetectRadius / 2;
 		PhysicsParallelFor(SurfaceVertices.Num(),
-			[this, &Spatial, &Particles, &SurfaceVertices, &ConstraintIndex, &TriangleMesh, &Elements, &HalfRadius, &ComponentIndex, &PositionTargetStiffness, &UseAnisotropicSpring](int32 i)
+			[this, &Spatial, &Particles, &SurfaceVertices, &ConstraintIndex, &TriangleMesh, &Elements, &HalfRadius, &ComponentIndex, &PositionTargetStiffness, &UseAnisotropicSpring](int32 SurfaceVdx)
 			{
-				const int32 Index = SurfaceVertices[i];
+				const int32 Index = SurfaceVertices[SurfaceVdx];
 				TArray< TTriangleCollisionPoint<FSolverReal>> Result;
 				//PointProximityQuery
 				if (TriangleMesh.PointClosestTriangleQuery(Spatial, static_cast<const TArrayView<const FSolverVec3>&>(Particles.XArray()), Index, Particles.GetX(Index), HalfRadius, HalfRadius,
@@ -878,9 +914,9 @@ PRAGMA_ENABLE_DEPRECATION_WARNINGS
 							SingleConstraintData.bIsZeroRestLength = true; //Push-out type collision springs should be zero rest length
 
 							float SpringStiffness = 0.f;
-							for (int32 k = 0; k < 3; k++)
+							for (int32 TriLocalIdx = 0; TriLocalIdx < 3; TriLocalIdx++)
 							{
-								SpringStiffness += SingleConstraintData.SingleWeights[k] * PositionTargetStiffness * Particles.M(Elem[k]);
+								SpringStiffness += SingleConstraintData.SingleWeights[TriLocalIdx] * PositionTargetStiffness * Particles.M(Elem[TriLocalIdx]);
 							}
 							SpringStiffness += PositionTargetStiffness * Particles.M(Index);
 							SingleConstraintData.SingleStiffness = (FSolverReal)SpringStiffness;
@@ -905,9 +941,9 @@ PRAGMA_ENABLE_DEPRECATION_WARNINGS
 		const TArray<TVec3<int32>>& Elements = TriangleMesh.GetSurfaceElements();
 		const float HalfRadius = DetectRadius/2;
 		PhysicsParallelFor(SurfaceVertices.Num(),
-			[this, &Spatial, &Particles, &SurfaceVertices, &ConstraintIndex, &TriangleMesh, &ExcludeMap, &Elements, &HalfRadius, &PositionTargetStiffness, &UseAnisotropicSpring](int32 i)
+			[this, &Spatial, &Particles, &SurfaceVertices, &ConstraintIndex, &TriangleMesh, &ExcludeMap, &Elements, &HalfRadius, &PositionTargetStiffness, &UseAnisotropicSpring](int32 SurfaceVdx)
 			{
-				const int32 Index = SurfaceVertices[i];
+				const int32 Index = SurfaceVertices[SurfaceVdx];
 				TArray< TTriangleCollisionPoint<FSolverReal>> Result;
 				//PointProximityQuery
 				if (TriangleMesh.PointClosestTriangleQuery(Spatial, static_cast<const TArrayView<const FSolverVec3>&>(Particles.XArray()), Index, Particles.GetX(Index), HalfRadius, HalfRadius,
@@ -934,9 +970,9 @@ PRAGMA_ENABLE_DEPRECATION_WARNINGS
 							SingleConstraintData.bIsZeroRestLength = true; //Push-out type collision springs should be zero rest length
 
 							float SpringStiffness = 0.f;
-							for (int32 k = 0; k < 3; k++)
+							for (int32 TriLocalIdx = 0; TriLocalIdx < 3; TriLocalIdx++)
 							{
-								SpringStiffness += SingleConstraintData.SingleWeights[k] * PositionTargetStiffness * Particles.M(Elem[k]);
+								SpringStiffness += SingleConstraintData.SingleWeights[TriLocalIdx] * PositionTargetStiffness * Particles.M(Elem[TriLocalIdx]);
 							}
 							SpringStiffness += PositionTargetStiffness * Particles.M(Index);
 							SingleConstraintData.SingleStiffness = (FSolverReal)SpringStiffness;
