ChaosDeformable: Constraint hierarchy refactor and FSolverParticlesRange migration

diff --git a/Engine/Source/Runtime/Experimental/Chaos/Private/Chaos/Deformable/XPBDWeakConstraints.cpp b/Engine/Source/Runtime/Experimental/Chaos/Private/Chaos/Deformable/XPBDWeakConstraints.cpp
--- a/Engine/Source/Runtime/Experimental/Chaos/Private/Chaos/Deformable/XPBDWeakConstraints.cpp
+++ b/Engine/Source/Runtime/Experimental/Chaos/Private/Chaos/Deformable/XPBDWeakConstraints.cpp
@@ -8,7 +8,7 @@ namespace Chaos::Softs
 	using Chaos::TVec3;
 
 	// -------------------------------------------------------------------
-	// Constructors
+	// Constructors (FSolverParticles — delegate to Range versions)
 	// -------------------------------------------------------------------
 
 	FXPBDWeightedSpringConstraints::FXPBDWeightedSpringConstraints(
@@ -16,15 +16,62 @@ namespace Chaos::Softs
 		const TArray<TArray<int32>>& InIndices,
 		const TArray<TArray<FSolverReal>>& InWeights,
 		const TArray<FSolverReal>& InStiffness
+	)
+		: FXPBDWeightedSpringConstraints(
+			FSolverParticlesRange(InParticles),
+			InIndices, InWeights, InStiffness)
+	{}
+
+	FXPBDWeightedSpringConstraints::FXPBDWeightedSpringConstraints(
+		const FSolverParticles& InParticles,
+		const TArray<TArray<int32>>& InIndices,
+		const TArray<TArray<FSolverReal>>& InWeights,
+		const TArray<FSolverReal>& InStiffness,
+		const TArray<TArray<int32>>& InSecondIndices,
+		const TArray<TArray<FSolverReal>>& InSecondWeights,
+		const FDeformableXPBDWeightedSpringConstraintParams& InParams
+	)
+		: FXPBDWeightedSpringConstraints(
+			FSolverParticlesRange(InParticles),
+			InIndices, InWeights, InStiffness, InSecondIndices, InSecondWeights, InParams)
+	{}
+
+	FXPBDWeightedSpringConstraints::FXPBDWeightedSpringConstraints(
+		const FSolverParticles& InParticles,
+		const FSolverParticles& InSecondParticles,
+		const TArray<TArray<int32>>& InIndices,
+		const TArray<TArray<FSolverReal>>& InWeights,
+		const TArray<FSolverReal>& InStiffness,
+		const TArray<TArray<int32>>& InSecondIndices,
+		const TArray<TArray<FSolverReal>>& InSecondWeights,
+		const FDeformableXPBDWeightedSpringConstraintParams& InParams
+	)
+		: FXPBDWeightedSpringConstraints(
+			FSolverParticlesRange(InParticles),
+			FSolverParticlesRange(InSecondParticles),
+			InIndices, InWeights, InStiffness, InSecondIndices, InSecondWeights, InParams)
+	{}
+
+	// -------------------------------------------------------------------
+	// Constructors (FSolverParticlesRange — primary implementations)
+	// -------------------------------------------------------------------
+
+	FXPBDWeightedSpringConstraints::FXPBDWeightedSpringConstraints(
+		const FSolverParticlesRange& InParticlesRange,
+		const TArray<TArray<int32>>& InIndices,
+		const TArray<TArray<FSolverReal>>& InWeights,
+		const TArray<FSolverReal>& InStiffness
 	)
 		: Indices(InIndices), Weights(InWeights), Stiffness(InStiffness)
 	{
-		InitColor(InParticles);
+		// Initialize target positions to zero — caller must call UpdateTargets() before Apply() to set actual targets
+		Constraints.Init(TVector<FSolverReal, 3>((FSolverReal)0.), Indices.Num());
+		InitColor(InParticlesRange);
 		LambdaArray.Init((FSolverReal)0., Indices.Num() * 3);
 	}
 
 	FXPBDWeightedSpringConstraints::FXPBDWeightedSpringConstraints(
-		const FSolverParticles& InParticles,
+		const FSolverParticlesRange& InParticlesRange,
 		const TArray<TArray<int32>>& InIndices,
 		const TArray<TArray<FSolverReal>>& InWeights,
 		const TArray<FSolverReal>& InStiffness,
@@ -36,21 +83,21 @@ namespace Chaos::Softs
 	{
 		ensureMsgf(Indices.Num() == SecondIndices.Num(), TEXT("Input Double Bindings have wrong size"));
 
-		for (int32 i = 0; i < Indices.Num(); i++)
+		for (int32 ConstraintIdx = 0; ConstraintIdx < Indices.Num(); ConstraintIdx++)
 		{
-			TSet<int32> IndicesSet = TSet<int32>(Indices[i]);
-			for (int32 j = 0; j < SecondIndices[i].Num(); j++)
+			const TSet<int32> IndicesSet = TSet<int32>(Indices[ConstraintIdx]);
+			for (int32 SecondLocalIdx = 0; SecondLocalIdx < SecondIndices[ConstraintIdx].Num(); SecondLocalIdx++)
 			{
-				ensureMsgf(!IndicesSet.Contains(SecondIndices[i][j]), TEXT("Indices and Second Indices overlaps. Currently not supported"));
+				ensureMsgf(!IndicesSet.Contains(SecondIndices[ConstraintIdx][SecondLocalIdx]), TEXT("Indices and Second Indices overlaps. Currently not supported"));
 			}
 		}
-		InitColor(InParticles);
+		InitColor(InParticlesRange);
 		LambdaArray.Init((FSolverReal)0., Indices.Num() * 3);
 	}
 
 	FXPBDWeightedSpringConstraints::FXPBDWeightedSpringConstraints(
-		const FSolverParticles& InParticles,
-		const FSolverParticles& InSecondParticles,
+		const FSolverParticlesRange& InParticlesRange,
+		const FSolverParticlesRange& InSecondParticlesRange,
 		const TArray<TArray<int32>>& InIndices,
 		const TArray<TArray<FSolverReal>>& InWeights,
 		const TArray<FSolverReal>& InStiffness,
@@ -62,7 +109,7 @@ namespace Chaos::Softs
 	{
 		ensureMsgf(Indices.Num() == SecondIndices.Num(), TEXT("Input Double Bindings have wrong size"));
 
-		InitColor(InParticles, InSecondParticles);
+		InitColor(InParticlesRange, InSecondParticlesRange);
 		LambdaArray.Init((FSolverReal)0., Indices.Num() * 3);
 	}
 
@@ -104,7 +151,7 @@ namespace Chaos::Softs
 
 	void FXPBDWeightedSpringConstraints::ApplyInParallel(FSolverParticles& Particles, const FSolverReal Dt) const
 	{
-		FSolverParticlesRange Range(&Particles, /*InOffset = */ 0, /*InRange = */ Particles.Size());
+		FSolverParticlesRange Range(Particles);
 		ApplyInParallel(Range, Dt);
 	}
 
@@ -157,7 +204,7 @@ namespace Chaos::Softs
 
 	void FXPBDWeightedSpringConstraints::Init(const FSolverParticles& InParticles, const FSolverReal Dt) const
 	{
-		Init(FSolverParticlesRange(const_cast<FSolverParticles*>(&InParticles), /*InOffset = */ 0, /*InRange = */ InParticles.Size()), Dt);
+		Init(FSolverParticlesRange(InParticles), Dt);
 	}
 
 	void FXPBDWeightedSpringConstraints::Init(const FSolverParticlesRange& InParticles, const FSolverParticlesRange& InSecondParticles, const FSolverReal Dt) const
@@ -181,13 +228,13 @@ namespace Chaos::Softs
 		if (ensure(Indices[ConstraintIndex].Num() == Weights[ConstraintIndex].Num() && 
 			(bUseConstraintTargetPosition || SecondIndices[ConstraintIndex].Num() == SecondWeights[ConstraintIndex].Num())))
 		{
-			auto GetPosition = [bUseParticleX](const FSolverParticlesRange& InParticles, int32 Idx) -> const TVec3<FSolverReal>&
+			auto GetPosition = [bUseParticleX](const FSolverParticlesRange& InParticles, int32 Pdx) -> const TVec3<FSolverReal>&
 				{
-					return bUseParticleX ? InParticles.X(Idx) : InParticles.P(Idx);
+					return bUseParticleX ? InParticles.X(Pdx) : InParticles.P(Pdx);
 				};
-			for (int32 Idx = 0; Idx < Indices[ConstraintIndex].Num(); ++Idx)
+			for (int32 LocalIdx = 0; LocalIdx < Indices[ConstraintIndex].Num(); ++LocalIdx)
 			{
-				OutSpringEdge += Weights[ConstraintIndex][Idx] * GetPosition(InParticles, Indices[ConstraintIndex][Idx]);
+				OutSpringEdge += Weights[ConstraintIndex][LocalIdx] * GetPosition(InParticles, Indices[ConstraintIndex][LocalIdx]);
 			}
 			if (bUseConstraintTargetPosition)
 			{
@@ -195,9 +242,9 @@ namespace Chaos::Softs
 			}
 			else
 			{
-				for (int32 SecondIdx = 0; SecondIdx < SecondIndices[ConstraintIndex].Num(); ++SecondIdx)
+				for (int32 SecondLocalIdx = 0; SecondLocalIdx < SecondIndices[ConstraintIndex].Num(); ++SecondLocalIdx)
 				{
-					OutSpringEdge -= SecondWeights[ConstraintIndex][SecondIdx] * GetPosition(InSecondParticles, SecondIndices[ConstraintIndex][SecondIdx]);
+					OutSpringEdge -= SecondWeights[ConstraintIndex][SecondLocalIdx] * GetPosition(InSecondParticles, SecondIndices[ConstraintIndex][SecondLocalIdx]);
 				}
 			}
 			return true;
@@ -212,14 +259,14 @@ namespace Chaos::Softs
 	// InitColor (private)
 	// -------------------------------------------------------------------
 
-	void FXPBDWeightedSpringConstraints::InitColor(const FSolverParticles& Particles)
+	void FXPBDWeightedSpringConstraints::InitColor(const FSolverParticlesRange& ParticlesRange)
 	{
-		Chaos::ComputeWeakConstraintsColoring(Indices, SecondIndices, Particles, ConstraintsPerColor);
+		Chaos::ComputeWeakConstraintsColoring(Indices, SecondIndices, ParticlesRange.Size(), ConstraintsPerColor);
 	}
 
-	void FXPBDWeightedSpringConstraints::InitColor(const FSolverParticles& Particles, const FSolverParticles& SecondParticles)
+	void FXPBDWeightedSpringConstraints::InitColor(const FSolverParticlesRange& ParticlesRange, const FSolverParticlesRange& SecondParticlesRange)
 	{
-		Chaos::ComputeWeakConstraintsColoring(Indices, SecondIndices, Particles, SecondParticles, ConstraintsPerColor);
+		Chaos::ComputeWeakConstraintsColoring(Indices, SecondIndices, ParticlesRange.Size(), SecondParticlesRange.Size(), ConstraintsPerColor);
 	}
 
 	// -------------------------------------------------------------------
@@ -242,9 +289,9 @@ namespace Chaos::Softs
 		}
 
 		FSolverReal Denom = AlphaTilde;
-		for (int32 j = 0; j < Indices[ConstraintIndex].Num(); j++) {
-			Denom += Weights[ConstraintIndex][j] * Weights[ConstraintIndex][j] * Particles.InvM(Indices[ConstraintIndex][j]);
-
+		for (int32 LocalIdx = 0; LocalIdx < Indices[ConstraintIndex].Num(); LocalIdx++)
+		{
+			Denom += Weights[ConstraintIndex][LocalIdx] * Weights[ConstraintIndex][LocalIdx] * Particles.InvM(Indices[ConstraintIndex][LocalIdx]);
 		}
 
 		if (Denom == (FSolverReal)0.)
@@ -254,13 +301,14 @@ namespace Chaos::Softs
 
 		for (int32 Beta = 0; Beta < 3; Beta++)
 		{
-			FSolverReal Cj = SpringEdge[Beta];
+			const FSolverReal Cj = SpringEdge[Beta];
 			FSolverReal DLambda = -Cj - AlphaTilde * LambdaArray[ConstraintIndex * 3 + Beta];
 
 			DLambda /= Denom;
 			LambdaArray[ConstraintIndex * 3 + Beta] += DLambda;
-			for (int32 j = 0; j < Indices[ConstraintIndex].Num(); j++) {
-				Particles.P(Indices[ConstraintIndex][j])[Beta] += DLambda * Weights[ConstraintIndex][j] * Particles.InvM(Indices[ConstraintIndex][j]);
+			for (int32 LocalIdx = 0; LocalIdx < Indices[ConstraintIndex].Num(); LocalIdx++)
+			{
+				Particles.P(Indices[ConstraintIndex][LocalIdx])[Beta] += DLambda * Weights[ConstraintIndex][LocalIdx] * Particles.InvM(Indices[ConstraintIndex][LocalIdx]);
 			}
 		}
 	}
@@ -282,12 +330,14 @@ namespace Chaos::Softs
 		}
 
 		FSolverReal Denom = AlphaTilde;
-		for (int32 j = 0; j < Indices[ConstraintIndex].Num(); j++) {
-			Denom += Weights[ConstraintIndex][j] * Weights[ConstraintIndex][j] * Particles.InvM(Indices[ConstraintIndex][j]);
+		for (int32 LocalIdx = 0; LocalIdx < Indices[ConstraintIndex].Num(); LocalIdx++)
+		{
+			Denom += Weights[ConstraintIndex][LocalIdx] * Weights[ConstraintIndex][LocalIdx] * Particles.InvM(Indices[ConstraintIndex][LocalIdx]);
 		}
 
-		for (int32 j = 0; j < SecondIndices[ConstraintIndex].Num(); j++) {
-			Denom += SecondWeights[ConstraintIndex][j] * SecondWeights[ConstraintIndex][j] * Particles.InvM(SecondIndices[ConstraintIndex][j]);
+		for (int32 SecondLocalIdx = 0; SecondLocalIdx < SecondIndices[ConstraintIndex].Num(); SecondLocalIdx++)
+		{
+			Denom += SecondWeights[ConstraintIndex][SecondLocalIdx] * SecondWeights[ConstraintIndex][SecondLocalIdx] * Particles.InvM(SecondIndices[ConstraintIndex][SecondLocalIdx]);
 		}
 
 		if (Denom == (FSolverReal)0.)
@@ -297,17 +347,19 @@ namespace Chaos::Softs
 
 		for (int32 Beta = 0; Beta < 3; Beta++)
 		{
-			FSolverReal Cj = SpringEdge[Beta];
+			const FSolverReal Cj = SpringEdge[Beta];
 			FSolverReal DLambda = -Cj - AlphaTilde * LambdaArray[ConstraintIndex * 3 + Beta];
 
 			DLambda /= Denom;
 			LambdaArray[ConstraintIndex * 3 + Beta] += DLambda;
-			for (int32 j = 0; j < Indices[ConstraintIndex].Num(); j++) {
-				Particles.P(Indices[ConstraintIndex][j])[Beta] += DLambda * Weights[ConstraintIndex][j] * Particles.InvM(Indices[ConstraintIndex][j]);
+			for (int32 LocalIdx = 0; LocalIdx < Indices[ConstraintIndex].Num(); LocalIdx++)
+			{
+				Particles.P(Indices[ConstraintIndex][LocalIdx])[Beta] += DLambda * Weights[ConstraintIndex][LocalIdx] * Particles.InvM(Indices[ConstraintIndex][LocalIdx]);
 			}
 
-			for (int32 j = 0; j < SecondIndices[ConstraintIndex].Num(); j++) {
-				Particles.P(SecondIndices[ConstraintIndex][j])[Beta] -= DLambda * SecondWeights[ConstraintIndex][j] * Particles.InvM(SecondIndices[ConstraintIndex][j]);
+			for (int32 SecondLocalIdx = 0; SecondLocalIdx < SecondIndices[ConstraintIndex].Num(); SecondLocalIdx++)
+			{
+				Particles.P(SecondIndices[ConstraintIndex][SecondLocalIdx])[Beta] -= DLambda * SecondWeights[ConstraintIndex][SecondLocalIdx] * Particles.InvM(SecondIndices[ConstraintIndex][SecondLocalIdx]);
 			}
 		}
 	}
@@ -328,26 +380,26 @@ namespace Chaos::Softs
 		}
 
 		TVec3<FSolverReal> SpringEdge((FSolverReal)0.);
-		for (int32 j = 0; j < Indices[ConstraintIndex].Num(); j++)
+		for (int32 LocalIdx = 0; LocalIdx < Indices[ConstraintIndex].Num(); LocalIdx++)
 		{
 			for (int32 Alpha = 0; Alpha < 3; Alpha++)
 			{
-				SpringEdge[Alpha] += Weights[ConstraintIndex][j] * Particles.P(Indices[ConstraintIndex][j])[Alpha];
+				SpringEdge[Alpha] += Weights[ConstraintIndex][LocalIdx] * Particles.P(Indices[ConstraintIndex][LocalIdx])[Alpha];
 			}
 		}
-		for (int32 j = 0; j < SecondIndices[ConstraintIndex].Num(); j++)
+		for (int32 SecondLocalIdx = 0; SecondLocalIdx < SecondIndices[ConstraintIndex].Num(); SecondLocalIdx++)
 		{
-			SpringEdge -= SecondWeights[ConstraintIndex][j] * SecondParticles.P(SecondIndices[ConstraintIndex][j]);
+			SpringEdge -= SecondWeights[ConstraintIndex][SecondLocalIdx] * SecondParticles.P(SecondIndices[ConstraintIndex][SecondLocalIdx]);
 		}
 
 		FSolverReal Denom = AlphaTilde;
-		for (int32 j = 0; j < Indices[ConstraintIndex].Num(); j++)
+		for (int32 LocalIdx = 0; LocalIdx < Indices[ConstraintIndex].Num(); LocalIdx++)
 		{
-			Denom += Weights[ConstraintIndex][j] * Weights[ConstraintIndex][j] * Particles.InvM(Indices[ConstraintIndex][j]);
+			Denom += Weights[ConstraintIndex][LocalIdx] * Weights[ConstraintIndex][LocalIdx] * Particles.InvM(Indices[ConstraintIndex][LocalIdx]);
 		}
-		for (int32 j = 0; j < SecondIndices[ConstraintIndex].Num(); j++)
+		for (int32 SecondLocalIdx = 0; SecondLocalIdx < SecondIndices[ConstraintIndex].Num(); SecondLocalIdx++)
 		{
-			Denom += SecondWeights[ConstraintIndex][j] * SecondWeights[ConstraintIndex][j] * SecondParticles.InvM(SecondIndices[ConstraintIndex][j]);
+			Denom += SecondWeights[ConstraintIndex][SecondLocalIdx] * SecondWeights[ConstraintIndex][SecondLocalIdx] * SecondParticles.InvM(SecondIndices[ConstraintIndex][SecondLocalIdx]);
 		}
 		if (Denom == (FSolverReal)0.)
 		{
@@ -355,17 +407,17 @@ namespace Chaos::Softs
 		}
 		for (int32 Beta = 0; Beta < 3; Beta++)
 		{
-			FSolverReal Cj = SpringEdge[Beta];
+			const FSolverReal Cj = SpringEdge[Beta];
 			FSolverReal DLambda = -Cj - AlphaTilde * LambdaArray[ConstraintIndex * 3 + Beta];
 			DLambda /= Denom;
 			LambdaArray[ConstraintIndex * 3 + Beta] += DLambda;
-			for (int32 j = 0; j < Indices[ConstraintIndex].Num(); j++)
+			for (int32 LocalIdx = 0; LocalIdx < Indices[ConstraintIndex].Num(); LocalIdx++)
 			{
-				Particles.P(Indices[ConstraintIndex][j])[Beta] += DLambda * Weights[ConstraintIndex][j] * Particles.InvM(Indices[ConstraintIndex][j]);
+				Particles.P(Indices[ConstraintIndex][LocalIdx])[Beta] += DLambda * Weights[ConstraintIndex][LocalIdx] * Particles.InvM(Indices[ConstraintIndex][LocalIdx]);
 			}
-			for (int32 j = 0; j < SecondIndices[ConstraintIndex].Num(); j++)
+			for (int32 SecondLocalIdx = 0; SecondLocalIdx < SecondIndices[ConstraintIndex].Num(); SecondLocalIdx++)
 			{
-				SecondParticles.P(SecondIndices[ConstraintIndex][j])[Beta] -= DLambda * SecondWeights[ConstraintIndex][j] * SecondParticles.InvM(SecondIndices[ConstraintIndex][j]);
+				SecondParticles.P(SecondIndices[ConstraintIndex][SecondLocalIdx])[Beta] -= DLambda * SecondWeights[ConstraintIndex][SecondLocalIdx] * SecondParticles.InvM(SecondIndices[ConstraintIndex][SecondLocalIdx]);
 			}
 		}
 	}
@@ -378,7 +430,7 @@ namespace Chaos::Softs
 			return;
 		}
 
-		FSolverReal AlphaTilde = Stiffness[ConstraintIndex] > SoftMaxStiffness? FSolverReal(2) / (Stiffness[ConstraintIndex] * Dt * Dt) : FSolverReal(0);
+		const FSolverReal AlphaTilde = Stiffness[ConstraintIndex] > SoftMaxStiffness ? FSolverReal(2) / (Stiffness[ConstraintIndex] * Dt * Dt) : FSolverReal(0);
 
 		TVec3<FSolverReal> SpringEdge;
 		if (!ComputeSpringEdge(Particles, SecondParticles, SpringEdge, ConstraintIndex,
@@ -388,13 +440,13 @@ namespace Chaos::Softs
 		}
 
 		FSolverReal Denom = AlphaTilde;
-		for (int32 j = 0; j < Indices[ConstraintIndex].Num(); j++)
+		for (int32 LocalIdx = 0; LocalIdx < Indices[ConstraintIndex].Num(); LocalIdx++)
 		{
-			Denom += Weights[ConstraintIndex][j] * Weights[ConstraintIndex][j] * Particles.InvM(Indices[ConstraintIndex][j]);
+			Denom += Weights[ConstraintIndex][LocalIdx] * Weights[ConstraintIndex][LocalIdx] * Particles.InvM(Indices[ConstraintIndex][LocalIdx]);
 		}
-		for (int32 j = 0; j < SecondIndices[ConstraintIndex].Num(); j++)
+		for (int32 SecondLocalIdx = 0; SecondLocalIdx < SecondIndices[ConstraintIndex].Num(); SecondLocalIdx++)
 		{
-			Denom += SecondWeights[ConstraintIndex][j] * SecondWeights[ConstraintIndex][j] * SecondParticles.InvM(SecondIndices[ConstraintIndex][j]);
+			Denom += SecondWeights[ConstraintIndex][SecondLocalIdx] * SecondWeights[ConstraintIndex][SecondLocalIdx] * SecondParticles.InvM(SecondIndices[ConstraintIndex][SecondLocalIdx]);
 		}
 		if (Denom == (FSolverReal)0.)
 		{
@@ -402,17 +454,17 @@ namespace Chaos::Softs
 		}
 		for (int32 Beta = 0; Beta < 3; Beta++)
 		{
-			FSolverReal Cj = SpringEdge[Beta];
+			const FSolverReal Cj = SpringEdge[Beta];
 			FSolverReal DLambda = -Cj - AlphaTilde * LambdaArray[ConstraintIndex * 3 + Beta];
 			DLambda /= Denom;
 			LambdaArray[ConstraintIndex * 3 + Beta] += DLambda;
-			for (int32 j = 0; j < Indices[ConstraintIndex].Num(); j++)
+			for (int32 LocalIdx = 0; LocalIdx < Indices[ConstraintIndex].Num(); LocalIdx++)
 			{
-				Particles.P(Indices[ConstraintIndex][j])[Beta] += DLambda * Weights[ConstraintIndex][j] * Particles.InvM(Indices[ConstraintIndex][j]);
+				Particles.P(Indices[ConstraintIndex][LocalIdx])[Beta] += DLambda * Weights[ConstraintIndex][LocalIdx] * Particles.InvM(Indices[ConstraintIndex][LocalIdx]);
 			}
-			for (int32 j = 0; j < SecondIndices[ConstraintIndex].Num(); j++)
+			for (int32 SecondLocalIdx = 0; SecondLocalIdx < SecondIndices[ConstraintIndex].Num(); SecondLocalIdx++)
 			{
-				SecondParticles.P(SecondIndices[ConstraintIndex][j])[Beta] -= DLambda * SecondWeights[ConstraintIndex][j] * SecondParticles.InvM(SecondIndices[ConstraintIndex][j]);
+				SecondParticles.P(SecondIndices[ConstraintIndex][SecondLocalIdx])[Beta] -= DLambda * SecondWeights[ConstraintIndex][SecondLocalIdx] * SecondParticles.InvM(SecondIndices[ConstraintIndex][SecondLocalIdx]);
 			}
 		}
 	}
@@ -425,8 +477,8 @@ namespace Chaos::Softs
 		const TFunctionRef<void(const FVector& Pos0, const FVector& Pos1, const FLinearColor& Color)>& DrawLine,
 		const TFunctionRef<void(const FVector& Pos, const FLinearColor& Color)>& DrawPoint) const
 	{
-#if WITH_EDITOR
-		auto DoubleVert = [](Chaos::TVec3<FSolverReal> V) { return FVector3d(V.X, V.Y, V.Z); };
+#if CHAOS_DEBUG_DRAW
+		auto DoubleVert = [](const Chaos::TVec3<FSolverReal>& V) { return FVector3d(V.X, V.Y, V.Z); };
 		for (int32 ConstraintIndex = 0; ConstraintIndex < Indices.Num(); ++ConstraintIndex)
 		{
 			Chaos::TVec3<FSolverReal> SourcePos((FSolverReal)0.), TargetPos((FSolverReal)0.);
@@ -463,8 +515,8 @@ namespace Chaos::Softs
 
 	void FXPBDWeightedSpringConstraints::VisualizeAllBindings(const FSolverParticlesRange& InParticles, const FSolverReal Dt) const
 	{
-#if WITH_EDITOR
-		auto DoubleVert = [](Chaos::TVec3<FSolverReal> V) { return FVector3d(V.X, V.Y, V.Z); };
+#if CHAOS_DEBUG_DRAW
+		auto DoubleVert = [](const Chaos::TVec3<FSolverReal>& V) { return FVector3d(V.X, V.Y, V.Z); };
 		for (int32 ConstraintIndex = 0; ConstraintIndex < Indices.Num(); ++ConstraintIndex)
 		{
 			Chaos::TVec3<FSolverReal> SourcePos((FSolverReal)0.), TargetPos((FSolverReal)0.);
@@ -487,8 +539,8 @@ namespace Chaos::Softs
 				}
 			}
 
-			float ParticleThickness = DebugDrawParams.DebugParticleWidth;
-			float LineThickness = DebugDrawParams.DebugLineWidth;
+			const float ParticleThickness = DebugDrawParams.DebugParticleWidth;
+			const float LineThickness = DebugDrawParams.DebugLineWidth;
 
 			if (Indices[ConstraintIndex].Num() > 1)
 			{
diff --git a/Engine/Source/Runtime/Experimental/Chaos/Private/Chaos/GraphColoring.cpp b/Engine/Source/Runtime/Experimental/Chaos/Private/Chaos/GraphColoring.cpp
--- a/Engine/Source/Runtime/Experimental/Chaos/Private/Chaos/GraphColoring.cpp
+++ b/Engine/Source/Runtime/Experimental/Chaos/Private/Chaos/GraphColoring.cpp
@@ -62,130 +62,23 @@ static bool VerifyGridBasedSubColoring(const TArray<TArray<int32>>& ElementsPerC
 
 }
 
-template <typename T>
-static bool VerifyWeakConstraintsColoring(const TArray<TArray<int32>>& Indices, const TArray<TArray<int32>>& SecondIndices, const Chaos::TDynamicParticles<T, 3>& InParticles, const TArray<TArray<int32>>& ConstraintsPerColor)
-{
-	TArray<bool> ConstraintIsIncluded;
-	ConstraintIsIncluded.Init(false, Indices.Num());
-	for (int32 i = 0; i < ConstraintsPerColor.Num(); i++)
-	{
-		for (int32 j = 0; j < ConstraintsPerColor[i].Num(); j++)
-		{
-			ConstraintIsIncluded[ConstraintsPerColor[i][j]] = true;
-		}
-	}
-
-	for (int32 kk = 0; kk < Indices.Num(); kk++)
-	{
-		if (!ConstraintIsIncluded[kk])
-		{
-			return false;
-		}
-	}
-
-	for (int32 i = 0; i < ConstraintsPerColor.Num(); i++)
-	{
-		TSet<int32> CoveredParticles;
-		for (int32 j = 0; j < ConstraintsPerColor[i].Num(); j++)
-		{
-			for (int32 Node : Indices[ConstraintsPerColor[i][j]])
-			{
-				if (CoveredParticles.Contains(Node))
-				{
-					return false;
-				}
-				CoveredParticles.Emplace(Node);
-			}
-
-			if (SecondIndices.Num() > 0)
-			{
-				for (int32 Node : SecondIndices[ConstraintsPerColor[i][j]])
-				{
-					if (CoveredParticles.Contains(Node))
-					{
-						return false;
-					}
-					CoveredParticles.Emplace(Node);
-				}
-
-			}
-		}
-	}
-
-	return true;
-
-}
-
-template <typename T>
-static bool VerifyWeakConstraintsColoring(const TArray<TArray<int32>>& Indices, const TArray<TArray<int32>>& SecondIndices, const Chaos::TDynamicParticles<T, 3>& InParticles, const Chaos::TDynamicParticles<T, 3>& InSecondParticles, const TArray<TArray<int32>>& ConstraintsPerColor)
-{
-	TArray<bool> ConstraintIsIncluded;
-	ConstraintIsIncluded.Init(false, Indices.Num());
-	for (int32 ColorIndex = 0; ColorIndex < ConstraintsPerColor.Num(); ++ColorIndex)
-	{
-		for (int32 ConstraintIndex = 0; ConstraintIndex < ConstraintsPerColor[ColorIndex].Num(); ++ConstraintIndex)
-		{
-			ConstraintIsIncluded[ConstraintsPerColor[ColorIndex][ConstraintIndex]] = true;
-		}
-	}
-
-	for (int32 ConstraintIndex = 0; ConstraintIndex < Indices.Num(); ++ConstraintIndex)
-	{
-		if (!ConstraintIsIncluded[ConstraintIndex])
-		{
-			return false;
-		}
-	}
-
-	for (int32 ColorIndex = 0; ColorIndex < ConstraintsPerColor.Num(); ++ColorIndex)
-	{
-		TSet<int32> CoveredParticles;
-		TSet<int32> CoveredSecondParticles;
-		for (int32 ConstraintIndex = 0; ConstraintIndex < ConstraintsPerColor[ColorIndex].Num(); ++ConstraintIndex)
-		{
-			for (int32 ParticleIndex : Indices[ConstraintsPerColor[ColorIndex][ConstraintIndex]])
-			{
-				if (CoveredParticles.Contains(ParticleIndex))
-				{
-					return false;
-				}
-				CoveredParticles.Emplace(ParticleIndex);
-			}
-
-			if (SecondIndices.Num() > 0)
-			{
-				for (int32 SecondParticleIndex : SecondIndices[ConstraintsPerColor[ColorIndex][ConstraintIndex]])
-				{
-					if (CoveredSecondParticles.Contains(SecondParticleIndex))
-					{
-						return false;
-					}
-					CoveredSecondParticles.Emplace(SecondParticleIndex);
-				}
-
-			}
-		}
-	}
-	return true;
-}
-
-template <typename T>
-static bool VerifyNodalColoring(const TArray<Chaos::TVec4<int32>>& Graph, const Chaos::TDynamicParticles<T, 3>& InParticles, const int32 GraphParticlesStart, const int32 GraphParticlesEnd, const TArray<TArray<int32>>& IncidentElements, const TArray<TArray<int32>>& IncidentElementsLocalIndex, const TArray<TArray<int32>> ParticlesPerColor)
+// Impl VerifyNodalColoring (TVec4 graph)
+static bool VerifyNodalColoringImpl(const TArray<Chaos::TVec4<int32>>& Graph, int32 NumParticles, const TFunctionRef<bool(int32)>& IsDynamicParticle, const int32 GraphParticlesStart, const int32 GraphParticlesEnd, const TArray<TArray<int32>>& IncidentElements, const TArray<TArray<int32>>& IncidentElementsLocalIndex, const TArray<TArray<int32>>& ParticlesPerColor)
 {
 	TArray<bool> ParticleIsIncluded;
-	ParticleIsIncluded.Init(false, InParticles.Size());
-	for (int32 i = 0; i < ParticlesPerColor.Num(); i++) 
+	ParticleIsIncluded.Init(false, NumParticles);
+	for (int32 ColorIdx = 0; ColorIdx < ParticlesPerColor.Num(); ColorIdx++)
 	{
-		for (int32 j = 0; j < ParticlesPerColor[i].Num(); j++)
+		for (int32 ParticleInColorIdx = 0; ParticleInColorIdx < ParticlesPerColor[ColorIdx].Num(); ParticleInColorIdx++)
 		{
-			ParticleIsIncluded[ParticlesPerColor[i][j]] = true;
+			ParticleIsIncluded[ParticlesPerColor[ColorIdx][ParticleInColorIdx]] = true;
 		}
 	}
 
-	for (int32 p = 0; p < GraphParticlesEnd - GraphParticlesStart; p++)
+	for (int32 LocalParticleIdx = 0; LocalParticleIdx < GraphParticlesEnd - GraphParticlesStart; LocalParticleIdx++)
 	{
-		int32 ParticleIndex = p + GraphParticlesStart;
-		if (InParticles.InvM(ParticleIndex) != (T)0.)
+		int32 ParticleIndex = LocalParticleIdx + GraphParticlesStart;
+		if (IsDynamicParticle(ParticleIndex))
 		{
 			if (!ParticleIsIncluded[ParticleIndex])
 			{
@@ -195,28 +88,28 @@ static bool VerifyNodalColoring(const TArray<Chaos::TVec4<int32>>& Graph, const
 	}
 
 	TArray<int32> Particle2Incident;
-	Particle2Incident.Init(-1, GraphParticlesEnd - GraphParticlesStart);
-	for (int32 i = 0; i < IncidentElements.Num(); i++)
+	Particle2Incident.Init(INDEX_NONE, GraphParticlesEnd - GraphParticlesStart);
+	for (int32 IncidentIdx = 0; IncidentIdx < IncidentElements.Num(); IncidentIdx++)
 	{
-		if (IncidentElements[i].Num() > 0)
+		if (IncidentElements[IncidentIdx].Num() > 0)
 		{
-			Particle2Incident[Graph[IncidentElements[i][0]][IncidentElementsLocalIndex[i][0]]] = i;
+			Particle2Incident[Graph[IncidentElements[IncidentIdx][0]][IncidentElementsLocalIndex[IncidentIdx][0]]] = IncidentIdx;
 		}
 	}
 
-	for (int32 i = 0; i < ParticlesPerColor.Num(); i++)
+	for (int32 ColorIdx = 0; ColorIdx < ParticlesPerColor.Num(); ColorIdx++)
 	{
 		TSet<int32> IncidentParticles;
-		for (int32 j = 0; j < ParticlesPerColor[i].Num(); j++)
+		for (int32 ParticleInColorIdx = 0; ParticleInColorIdx < ParticlesPerColor[ColorIdx].Num(); ParticleInColorIdx++)
 		{
-			int32 ParticleIndex = ParticlesPerColor[i][j];
+			int32 ParticleIndex = ParticlesPerColor[ColorIdx][ParticleInColorIdx];
 			TSet<int32> LocalIncidentParticles;
-			for (int32 k = 0; k < IncidentElements[Particle2Incident[ParticleIndex]].Num(); k++)
+			for (int32 IncidentIdx = 0; IncidentIdx < IncidentElements[Particle2Incident[ParticleIndex]].Num(); IncidentIdx++)
 			{
-				int32 ElementIndex = IncidentElements[Particle2Incident[ParticleIndex]][k];
-				for (int32 ie = 0; ie < 4; ie++)
+				int32 ElementIndex = IncidentElements[Particle2Incident[ParticleIndex]][IncidentIdx];
+				for (int32 ElementLocalIdx = 0; ElementLocalIdx < 4; ElementLocalIdx++)
 				{
-					LocalIncidentParticles.Emplace(Graph[ElementIndex][ie]);
+					LocalIncidentParticles.Emplace(Graph[ElementIndex][ElementLocalIdx]);
 				}
 			}
 			if (IncidentParticles.Contains(ParticleIndex))
@@ -234,28 +127,27 @@ static bool VerifyNodalColoring(const TArray<Chaos::TVec4<int32>>& Graph, const
 	}
 
 	return true;
-
 }
 
-template <typename T>
-static bool VerifyNodalColoring(const TArray<TArray<int32>>& Graph, const Chaos::TDynamicParticles<T, 3>& InParticles, const int32 GraphParticlesStart, const int32 GraphParticlesEnd, const TArray<TArray<int32>>& IncidentElements, const TArray<TArray<int32>>& IncidentElementsLocalIndex, const TArray<TArray<int32>> ParticlesPerColor)
+// Impl VerifyNodalColoring (TArray graph)
+static bool VerifyNodalColoringImpl(const TArray<TArray<int32>>& Graph, int32 NumParticles, const TFunctionRef<bool(int32)>& IsDynamicParticle, const int32 GraphParticlesStart, const int32 GraphParticlesEnd, const TArray<TArray<int32>>& IncidentElements, const TArray<TArray<int32>>& IncidentElementsLocalIndex, const TArray<TArray<int32>>& ParticlesPerColor)
 {
- 	checkSlow(GraphParticlesStart <= GraphParticlesEnd);
-	checkSlow(GraphParticlesEnd <= (int32)InParticles.Size());
+	checkSlow(GraphParticlesStart <= GraphParticlesEnd);
+	checkSlow(GraphParticlesEnd <= NumParticles);
 	TArray<bool> ParticleIsIncluded;
-	ParticleIsIncluded.Init(false, InParticles.Size());
-	for (int32 i = 0; i < ParticlesPerColor.Num(); i++)
+	ParticleIsIncluded.Init(false, NumParticles);
+	for (int32 ColorIdx = 0; ColorIdx < ParticlesPerColor.Num(); ColorIdx++)
 	{
-		for (int32 j = 0; j < ParticlesPerColor[i].Num(); j++)
+		for (int32 ParticleInColorIdx = 0; ParticleInColorIdx < ParticlesPerColor[ColorIdx].Num(); ParticleInColorIdx++)
 		{
-			ParticleIsIncluded[ParticlesPerColor[i][j]] = true;
+			ParticleIsIncluded[ParticlesPerColor[ColorIdx][ParticleInColorIdx]] = true;
 		}
 	}
 
-	for (int32 p = 0; p < GraphParticlesEnd - GraphParticlesStart; p++)
+	for (int32 LocalParticleIdx = 0; LocalParticleIdx < GraphParticlesEnd - GraphParticlesStart; LocalParticleIdx++)
 	{
-		int32 ParticleIndex = p + GraphParticlesStart;
-		if (InParticles.InvM(ParticleIndex) != (T)0. && IncidentElements[ParticleIndex].Num() > 0 )
+		int32 ParticleIndex = LocalParticleIdx + GraphParticlesStart;
+		if (IsDynamicParticle(ParticleIndex) && IncidentElements[ParticleIndex].Num() > 0)
 		{
 			if (!ParticleIsIncluded[ParticleIndex])
 			{
@@ -266,92 +158,24 @@ static bool VerifyNodalColoring(const TArray<TArray<int32>>& Graph, const Chaos:
 
 	TArray<int32> Particle2Incident;
 	Particle2Incident.Init(INDEX_NONE, GraphParticlesEnd - GraphParticlesStart);
-	for (int32 i = 0; i < GraphParticlesEnd - GraphParticlesStart; i++)
-	{
-		Particle2Incident[i] = i + GraphParticlesStart;
-	}
-
-	for (int32 i = 0; i < ParticlesPerColor.Num(); i++)
-	{
-		TSet<int32> IncidentParticles;
-		for (int32 j = 0; j < ParticlesPerColor[i].Num(); j++)
-		{
-			int32 ParticleIndex = ParticlesPerColor[i][j];
-			TSet<int32> LocalIncidentParticles;
-			for (int32 k = 0; k < IncidentElements[Particle2Incident[ParticleIndex]].Num(); k++)
-			{
-				int32 ElementIndex = IncidentElements[Particle2Incident[ParticleIndex]][k];
-				for (int32 ie = 0; ie < Graph[ElementIndex].Num(); ie++)
-				{
-					LocalIncidentParticles.Add(Graph[ElementIndex][ie]);
-				}
-			}
-			if (IncidentParticles.Contains(ParticleIndex))
-			{
-				return false;
-			}
-			else
-			{
-				for (int32 LocalParticle : LocalIncidentParticles)
-				{
-					IncidentParticles.Add(LocalParticle);
-				}
-			}
-		}
-	}
-
-	return true;
-}
-
-template <typename T>
-static bool VerifyExtraNodalColoring(const Chaos::TDynamicParticles<T, 3>& InParticles, const TArray<TArray<int32>>& Graph, const TArray<TArray<int32>>& ExtraGraph, const TArray<TArray<int32>>& IncidentElements, const TArray<TArray<int32>>& ExtraIncidentElements,const TArray<int32>& ParticleColors,const TArray<TArray<int32>>& ParticlesPerColor)
-{
-	TBitArray ParticleIsIncluded(false, InParticles.Size());
-	for (int32 i = 0; i < ParticlesPerColor.Num(); i++)
-	{
-		for (int32 j = 0; j < ParticlesPerColor[i].Num(); j++)
-		{
-			ParticleIsIncluded[ParticlesPerColor[i][j]] = true;
-		}
-	}
-
-	for (int32 ParticleIndex = 0; ParticleIndex < (int32)InParticles.Size(); ParticleIndex++)
+	for (int32 LocalParticleIdx = 0; LocalParticleIdx < GraphParticlesEnd - GraphParticlesStart; LocalParticleIdx++)
 	{
-		if (InParticles.InvM(ParticleIndex) != (T)0. && (IncidentElements[ParticleIndex].Num() > 0 || ExtraIncidentElements[ParticleIndex].Num() > 0))
-		{
-			if (!ParticleIsIncluded[ParticleIndex])
-			{
-				return false;
-			}
-		}
+		Particle2Incident[LocalParticleIdx] = LocalParticleIdx + GraphParticlesStart;
 	}
 
-
-	for (int32 i = 0; i < ParticlesPerColor.Num(); i++)
+	for (int32 ColorIdx = 0; ColorIdx < ParticlesPerColor.Num(); ColorIdx++)
 	{
 		TSet<int32> IncidentParticles;
-		IncidentParticles.Reserve(ParticlesPerColor[i].Num());
-		for (int32 j = 0; j < ParticlesPerColor[i].Num(); j++)
+		for (int32 ParticleInColorIdx = 0; ParticleInColorIdx < ParticlesPerColor[ColorIdx].Num(); ParticleInColorIdx++)
 		{
-			int32 ParticleIndex = ParticlesPerColor[i][j];
+			int32 ParticleIndex = ParticlesPerColor[ColorIdx][ParticleInColorIdx];
 			TSet<int32> LocalIncidentParticles;
-			for (int32 k = 0; k < IncidentElements[ParticleIndex].Num(); k++)
+			for (int32 IncidentIdx = 0; IncidentIdx < IncidentElements[Particle2Incident[ParticleIndex]].Num(); IncidentIdx++)
 			{
-				int32 ElementIndex = IncidentElements[ParticleIndex][k];
-				for (int32 ie = 0; ie < Graph[ElementIndex].Num(); ie++)
+				int32 ElementIndex = IncidentElements[Particle2Incident[ParticleIndex]][IncidentIdx];
+				for (int32 ElementLocalIdx = 0; ElementLocalIdx < Graph[ElementIndex].Num(); ElementLocalIdx++)
 				{
-					LocalIncidentParticles.Add(Graph[ElementIndex][ie]);
-				}
-			}
-			if (ParticleIndex < ExtraIncidentElements.Num())
-			{
-				for (int32 k = 0; k < ExtraIncidentElements[ParticleIndex].Num(); k++)
-				{
-					int32 ElementIndex = ExtraIncidentElements[ParticleIndex][k];
-					for (int32 ie = 0; ie < ExtraGraph[ElementIndex].Num(); ie++)
-					{
-						LocalIncidentParticles.Add(ExtraGraph[ElementIndex][ie]);
-					}
+					LocalIncidentParticles.Add(Graph[ElementIndex][ElementLocalIdx]);
 				}
 			}
 			if (IncidentParticles.Contains(ParticleIndex))
@@ -371,21 +195,21 @@ static bool VerifyExtraNodalColoring(const Chaos::TDynamicParticles<T, 3>& InPar
 	return true;
 }
 
-template <typename T>
-static bool VerifyExtraNodalColoring(const Chaos::TDynamicParticles<T, 3>& InParticles, const TArray<TArray<int32>>& StaticGraph, const TArray<TArray<int32>>& DynamicGraph, const TArray<TArray<int32>>& ExtraGraph, const TArray<TArray<int32>>& StaticIncidentElements, const TArray<TArray<int32>>& DynamicIncidentElements, const TArray<TArray<int32>>& ExtraIncidentElements,const TArray<int32>& ParticleColors,const TArray<TArray<int32>>& ParticlesPerColor)
+// VerifyExtraNodalColoringImpl (3-graph)
+static bool VerifyExtraNodalColoringImpl(int32 NumParticles, const TFunctionRef<bool(int32)>& IsDynamicParticle, const TArray<TArray<int32>>& StaticGraph, const TArray<TArray<int32>>& DynamicGraph, const TArray<TArray<int32>>& ExtraGraph, const TArray<TArray<int32>>& StaticIncidentElements, const TArray<TArray<int32>>& DynamicIncidentElements, const TArray<TArray<int32>>& ExtraIncidentElements, const TArray<int32>& ParticleColors, const TArray<TArray<int32>>& ParticlesPerColor)
 {
-	TBitArray ParticleIsIncluded(false, InParticles.Size());
-	for (int32 i = 0; i < ParticlesPerColor.Num(); i++)
+	TBitArray ParticleIsIncluded(false, NumParticles);
+	for (int32 ColorIdx = 0; ColorIdx < ParticlesPerColor.Num(); ColorIdx++)
 	{
-		for (int32 j = 0; j < ParticlesPerColor[i].Num(); j++)
+		for (int32 ParticleInColorIdx = 0; ParticleInColorIdx < ParticlesPerColor[ColorIdx].Num(); ParticleInColorIdx++)
 		{
-			ParticleIsIncluded[ParticlesPerColor[i][j]] = true;
+			ParticleIsIncluded[ParticlesPerColor[ColorIdx][ParticleInColorIdx]] = true;
 		}
 	}
 
-	for (int32 ParticleIndex = 0; ParticleIndex < (int32)InParticles.Size(); ParticleIndex++)
+	for (int32 ParticleIndex = 0; ParticleIndex < NumParticles; ParticleIndex++)
 	{
-		if (InParticles.InvM(ParticleIndex) != (T)0. && (StaticIncidentElements[ParticleIndex].Num() > 0 || ExtraIncidentElements[ParticleIndex].Num() > 0 || DynamicIncidentElements[ParticleIndex].Num() > 0))
+		if (IsDynamicParticle(ParticleIndex) && (StaticIncidentElements[ParticleIndex].Num() > 0 || ExtraIncidentElements[ParticleIndex].Num() > 0 || DynamicIncidentElements[ParticleIndex].Num() > 0))
 		{
 			if (!ParticleIsIncluded[ParticleIndex])
 			{
@@ -394,39 +218,38 @@ static bool VerifyExtraNodalColoring(const Chaos::TDynamicParticles<T, 3>& InPar
 		}
 	}
 
-
-	for (int32 i = 0; i < ParticlesPerColor.Num(); i++)
+	for (int32 ColorIdx = 0; ColorIdx < ParticlesPerColor.Num(); ColorIdx++)
 	{
 		TSet<int32> IncidentParticles;
-		IncidentParticles.Reserve(ParticlesPerColor[i].Num());
-		for (int32 j = 0; j < ParticlesPerColor[i].Num(); j++)
+		IncidentParticles.Reserve(ParticlesPerColor[ColorIdx].Num());
+		for (int32 ParticleInColorIdx = 0; ParticleInColorIdx < ParticlesPerColor[ColorIdx].Num(); ParticleInColorIdx++)
 		{
-			int32 ParticleIndex = ParticlesPerColor[i][j];
+			int32 ParticleIndex = ParticlesPerColor[ColorIdx][ParticleInColorIdx];
 			TSet<int32> LocalIncidentParticles;
-			for (int32 k = 0; k < StaticIncidentElements[ParticleIndex].Num(); k++)
+			for (int32 IncidentIdx = 0; IncidentIdx < StaticIncidentElements[ParticleIndex].Num(); IncidentIdx++)
 			{
-				int32 ElementIndex = StaticIncidentElements[ParticleIndex][k];
-				for (int32 ie = 0; ie < StaticGraph[ElementIndex].Num(); ie++)
+				int32 ElementIndex = StaticIncidentElements[ParticleIndex][IncidentIdx];
+				for (int32 ElementLocalIdx = 0; ElementLocalIdx < StaticGraph[ElementIndex].Num(); ElementLocalIdx++)
 				{
-					LocalIncidentParticles.Add(StaticGraph[ElementIndex][ie]);
+					LocalIncidentParticles.Add(StaticGraph[ElementIndex][ElementLocalIdx]);
 				}
 			}
-			for (int32 k = 0; k < DynamicIncidentElements[ParticleIndex].Num(); k++)
+			for (int32 IncidentIdx = 0; IncidentIdx < DynamicIncidentElements[ParticleIndex].Num(); IncidentIdx++)
 			{
-				int32 ElementIndex = DynamicIncidentElements[ParticleIndex][k];
-				for (int32 ie = 0; ie < DynamicGraph[ElementIndex].Num(); ie++)
+				int32 ElementIndex = DynamicIncidentElements[ParticleIndex][IncidentIdx];
+				for (int32 ElementLocalIdx = 0; ElementLocalIdx < DynamicGraph[ElementIndex].Num(); ElementLocalIdx++)
 				{
-					LocalIncidentParticles.Add(DynamicGraph[ElementIndex][ie]);
+					LocalIncidentParticles.Add(DynamicGraph[ElementIndex][ElementLocalIdx]);
 				}
 			}
 			if (ParticleIndex < ExtraIncidentElements.Num())
 			{
-				for (int32 k = 0; k < ExtraIncidentElements[ParticleIndex].Num(); k++)
+				for (int32 IncidentIdx = 0; IncidentIdx < ExtraIncidentElements[ParticleIndex].Num(); IncidentIdx++)
 				{
-					int32 ElementIndex = ExtraIncidentElements[ParticleIndex][k];
-					for (int32 ie = 0; ie < ExtraGraph[ElementIndex].Num(); ie++)
+					int32 ElementIndex = ExtraIncidentElements[ParticleIndex][IncidentIdx];
+					for (int32 ElementLocalIdx = 0; ElementLocalIdx < ExtraGraph[ElementIndex].Num(); ElementLocalIdx++)
 					{
-						LocalIncidentParticles.Add(ExtraGraph[ElementIndex][ie]);
+						LocalIncidentParticles.Add(ExtraGraph[ElementIndex][ElementLocalIdx]);
 					}
 				}
 			}
@@ -447,6 +270,17 @@ static bool VerifyExtraNodalColoring(const Chaos::TDynamicParticles<T, 3>& InPar
 	return true;
 }
 
+// VerifyExtraNodalColoringImpl (2-graph) — delegates to 3-graph with empty DynamicGraph
+static bool VerifyExtraNodalColoringImpl(int32 NumParticles, const TFunctionRef<bool(int32)>& IsDynamicParticle, const TArray<TArray<int32>>& Graph, const TArray<TArray<int32>>& ExtraGraph, const TArray<TArray<int32>>& IncidentElements, const TArray<TArray<int32>>& ExtraIncidentElements, const TArray<int32>& ParticleColors, const TArray<TArray<int32>>& ParticlesPerColor)
+{
+	// EmptyGraph doesn't need sizing — it's only indexed by constraint indices from DynamicIncidentElements, which are all empty
+	const TArray<TArray<int32>> EmptyGraph;
+	// Size must match NumParticles so the 3-graph version can safely index DynamicIncidentElements[ParticleIndex]
+	TArray<TArray<int32>> EmptyIncidentElements;
+	EmptyIncidentElements.SetNum(NumParticles);
+	return VerifyExtraNodalColoringImpl(NumParticles, IsDynamicParticle, Graph, EmptyGraph, ExtraGraph, IncidentElements, EmptyIncidentElements, ExtraIncidentElements, ParticleColors, ParticlesPerColor);
+}
+
 template<typename DynamicParticlesType, int32 N, bool bAllDynamic>
 TArray<TArray<int32>> Chaos::FGraphColoring::ComputeGraphColoringParticlesOrRange(const TArray<TVector<int32, N>>& Graph, const DynamicParticlesType& InParticles, const int32 GraphParticlesStart, const int32 GraphParticlesEnd)
 {
@@ -456,7 +290,7 @@ TArray<TArray<int32>> Chaos::FGraphColoring::ComputeGraphColoringParticlesOrRang
 
 	TArray<TArray<int32>> ColorGraph;
 
-	int32 MaxColor = -1;
+	int32 MaxColor = INDEX_NONE;
 	TArray<FColorSet> NodeUsedColorsSubArray;
 	NodeUsedColorsSubArray.SetNum(GraphParticlesEnd - GraphParticlesStart);
 	TArrayView<FColorSet> NodeUsedColors(NodeUsedColorsSubArray.GetData() - GraphParticlesStart, GraphParticlesEnd); // Only nodes starting with GraphParticlesStart are valid to access
@@ -524,21 +358,20 @@ void Chaos::ComputeGridBasedGraphSubColoringPointer(const TArray<TArray<int32>>&
 	}
 	ElementsPerSubColors.SetNum(ElementsPerColor.Num());
 
-	//for (int32 i = 0; i < ElementsPerColor.Num(); i++) {
 	PhysicsParallelFor(ElementsPerColor.Num(), [bHaveInitialGuess, &ElementsPerColor, &Grid, GridSize, &PreviousColoring, &ConstraintsNodesSet, &ElementsPerSubColors](const int32 Color)
 		{
 			const TArray<int32>& ColorElements = ElementsPerColor[Color];
 			TArray<int32>& PreviousColoringElements = PreviousColoring->operator[](Color);
 			if (!bHaveInitialGuess)
 			{
-				PreviousColoringElements.Init(-1, ColorElements.Num());
+				PreviousColoringElements.Init(INDEX_NONE, ColorElements.Num());
 			}
 			int32 NumNodes = GridSize;
 			TArray<int32> ElementSubColors;
-			ElementSubColors.Init(-1, ColorElements.Num());
+			ElementSubColors.Init(INDEX_NONE, ColorElements.Num());
 			TArray<FColorSet> UsedColors;
 			UsedColors.SetNum(NumNodes);
-			int32 MaxColor = -1;
+			int32 MaxColor = INDEX_NONE;
 			for (int32 ElementIndex = 0; ElementIndex < ColorElements.Num(); ElementIndex++) 
 			{
 				int32 ColorToUse = 0;
@@ -569,7 +402,7 @@ void Chaos::ComputeGridBasedGraphSubColoringPointer(const TArray<TArray<int32>>&
 						ElementSubColors[ElementIndex] = ColorToUse;
 					}
 				}
-				if (ElementSubColors[ElementIndex] == -1)
+				if (ElementSubColors[ElementIndex] == INDEX_NONE)
 				{
 					while (true) 
 					{
@@ -582,7 +415,8 @@ void Chaos::ComputeGridBasedGraphSubColoringPointer(const TArray<TArray<int32>>&
 								break;
 							}
 						}
-						if (!bColorFound) {
+						if (!bColorFound)
+						{
 							break;
 						}
 						ColorToUse++;
@@ -614,14 +448,67 @@ void Chaos::ComputeGridBasedGraphSubColoringPointer(const TArray<TArray<int32>>&
 
 template<typename T>
 void Chaos::ComputeWeakConstraintsColoring(const TArray<TArray<int32>>& Indices, const TArray<TArray<int32>>& SecondIndices, const Chaos::TDynamicParticles<T, 3>& InParticles, TArray<TArray<int32>>& ConstraintsPerColor)
+{
+	Chaos::ComputeWeakConstraintsColoring(Indices, SecondIndices, (int32)InParticles.Size(), ConstraintsPerColor);
+}
+
+// Non-template verify helper (particles are unused in verification)
+static bool VerifyWeakConstraintsColoringByCount(const TArray<TArray<int32>>& Indices, const TArray<TArray<int32>>& SecondIndices, const TArray<TArray<int32>>& ConstraintsPerColor)
+{
+	TArray<bool> ConstraintIsIncluded;
+	ConstraintIsIncluded.Init(false, Indices.Num());
+	for (int32 ColorIndex = 0; ColorIndex < ConstraintsPerColor.Num(); ColorIndex++)
+	{
+		for (int32 ConstraintIdx = 0; ConstraintIdx < ConstraintsPerColor[ColorIndex].Num(); ConstraintIdx++)
+		{
+			ConstraintIsIncluded[ConstraintsPerColor[ColorIndex][ConstraintIdx]] = true;
+		}
+	}
+	for (int32 ConstraintIndex = 0; ConstraintIndex < Indices.Num(); ConstraintIndex++)
+	{
+		if (!ConstraintIsIncluded[ConstraintIndex])
+		{
+			return false;
+		}
+	}
+	for (int32 ColorIndex = 0; ColorIndex < ConstraintsPerColor.Num(); ColorIndex++)
+	{
+		TSet<int32> CoveredParticles;
+		for (int32 ConstraintIdx = 0; ConstraintIdx < ConstraintsPerColor[ColorIndex].Num(); ConstraintIdx++)
+		{
+			for (int32 Node : Indices[ConstraintsPerColor[ColorIndex][ConstraintIdx]])
+			{
+				if (CoveredParticles.Contains(Node))
+				{
+					return false;
+				}
+				CoveredParticles.Emplace(Node);
+			}
+			if (SecondIndices.Num() > 0)
+			{
+				for (int32 Node : SecondIndices[ConstraintsPerColor[ColorIndex][ConstraintIdx]])
+				{
+					if (CoveredParticles.Contains(Node))
+					{
+						return false;
+					}
+					CoveredParticles.Emplace(Node);
+				}
+			}
+		}
+	}
+	return true;
+}
+
+void Chaos::ComputeWeakConstraintsColoring(const TArray<TArray<int32>>& Indices, const TArray<TArray<int32>>& SecondIndices, int32 NumParticles, TArray<TArray<int32>>& ConstraintsPerColor)
 {
 	TArray<TSet<int32>> UsedColors;
-	UsedColors.SetNum(InParticles.Size());
+	UsedColors.SetNum(NumParticles);
 
 	ensure(Indices.Num() == SecondIndices.Num() || SecondIndices.Num() == 0);
 
 	TArray<int32> ConstraintColors;
-	ConstraintColors.Init(-1, Indices.Num());
+	ConstraintColors.Init(INDEX_NONE, Indices.Num());
 
 	if (SecondIndices.Num() == 0)
 	{
@@ -631,21 +518,24 @@ void Chaos::ComputeWeakConstraintsColoring(const TArray<TArray<int32>>& Indices,
 			while (true)
 			{
 				bool ColorFound = false;
-				for (int32 node : Indices[ConstraintIndex]) {
-					if (UsedColors[node].Contains(ColorToUse)) {
+				for (int32 Node : Indices[ConstraintIndex])
+				{
+					if (UsedColors[Node].Contains(ColorToUse))
+					{
 						ColorFound = true;
 						break;
 					}
 				}
-				if (!ColorFound) {
+				if (!ColorFound)
+				{
 					break;
 				}
 				ColorToUse++;
 			}
 			ConstraintColors[ConstraintIndex] = ColorToUse;
-			for (int32 node : Indices[ConstraintIndex])
+			for (int32 Node : Indices[ConstraintIndex])
 			{
-				UsedColors[node].Emplace(ColorToUse);
+				UsedColors[Node].Emplace(ColorToUse);
 			}
 		}
 	}
@@ -657,29 +547,36 @@ void Chaos::ComputeWeakConstraintsColoring(const TArray<TArray<int32>>& Indices,
 			while (true)
 			{
 				bool ColorFound = false;
-				for (int32 node : Indices[ConstraintIndex]) {
-					if (UsedColors[node].Contains(ColorToUse)) {
+				for (int32 Node : Indices[ConstraintIndex])
+				{
+					if (UsedColors[Node].Contains(ColorToUse))
+					{
 						ColorFound = true;
 						break;
 					}
 				}
-				for (int32 node : SecondIndices[ConstraintIndex]) {
-					if (UsedColors[node].Contains(ColorToUse)) {
+				for (int32 Node : SecondIndices[ConstraintIndex])
+				{
+					if (UsedColors[Node].Contains(ColorToUse))
+					{
 						ColorFound = true;
 						break;
 					}
 				}
-				if (!ColorFound) {
+				if (!ColorFound)
+				{
 					break;
 				}
 				ColorToUse++;
 			}
 			ConstraintColors[ConstraintIndex] = ColorToUse;
-			for (int32 node : Indices[ConstraintIndex]) {
-				UsedColors[node].Emplace(ColorToUse);
+			for (int32 Node : Indices[ConstraintIndex])
+			{
+				UsedColors[Node].Emplace(ColorToUse);
 			}
-			for (int32 node : SecondIndices[ConstraintIndex]) {
-				UsedColors[node].Emplace(ColorToUse);
+			for (int32 Node : SecondIndices[ConstraintIndex])
+			{
+				UsedColors[Node].Emplace(ColorToUse);
 			}
 		}
 	}
@@ -688,25 +585,24 @@ void Chaos::ComputeWeakConstraintsColoring(const TArray<TArray<int32>>& Indices,
 	ConstraintsPerColor.Empty();
 	ConstraintsPerColor.SetNum(NumColors + 1);
 
-	for (int32 j = 0; j < Indices.Num(); j++) 
+	for (int32 ConstraintIndex = 0; ConstraintIndex < Indices.Num(); ConstraintIndex++)
 	{
-		ConstraintsPerColor[ConstraintColors[j]].Emplace(j);
+		ConstraintsPerColor[ConstraintColors[ConstraintIndex]].Emplace(ConstraintIndex);
 	}
 
-	checkSlow(VerifyWeakConstraintsColoring<T>(Indices, SecondIndices, InParticles, ConstraintsPerColor));
+	checkSlow(VerifyWeakConstraintsColoringByCount(Indices, SecondIndices, ConstraintsPerColor));
 }
 
-template<typename T>
-void Chaos::ComputeWeakConstraintsColoring(const TArray<TArray<int32>>& Indices, const TArray<TArray<int32>>& SecondIndices, const Chaos::TDynamicParticles<T, 3>& InParticles, const Chaos::TDynamicParticles<T, 3>& InSecondParticles, TArray<TArray<int32>>& ConstraintsPerColor)
+void Chaos::ComputeWeakConstraintsColoring(const TArray<TArray<int32>>& Indices, const TArray<TArray<int32>>& SecondIndices, int32 NumParticles, int32 NumSecondParticles, TArray<TArray<int32>>& ConstraintsPerColor)
 {
 	TArray<TSet<int32>> ParticleUsedColors, SecondParticleUsedColors;
-	ParticleUsedColors.SetNum(InParticles.Size());
-	SecondParticleUsedColors.SetNum(InSecondParticles.Size());
+	ParticleUsedColors.SetNum(NumParticles);
+	SecondParticleUsedColors.SetNum(NumSecondParticles);
 
 	ensure(Indices.Num() == SecondIndices.Num() || SecondIndices.Num() == 0);
 
 	TArray<int32> ConstraintColors;
-	ConstraintColors.Init(-1, Indices.Num());
+	ConstraintColors.Init(INDEX_NONE, Indices.Num());
 
 	if (SecondIndices.Num() == 0)
 	{
@@ -761,7 +657,8 @@ void Chaos::ComputeWeakConstraintsColoring(const TArray<TArray<int32>>& Indices,
 						break;
 					}
 				}
-				if (!ColorFound) {
+				if (!ColorFound)
+				{
 					break;
 				}
 				ColorToUse++;
@@ -788,47 +685,45 @@ void Chaos::ComputeWeakConstraintsColoring(const TArray<TArray<int32>>& Indices,
 		ConstraintsPerColor[ConstraintColors[ConstraintIndex]].Emplace(ConstraintIndex);
 	}
 
-	checkSlow(VerifyWeakConstraintsColoring<T>(Indices, SecondIndices, InParticles, InSecondParticles, ConstraintsPerColor));
+	checkSlow(VerifyWeakConstraintsColoringByCount(Indices, SecondIndices, ConstraintsPerColor));
 }
 
-template<typename T>
-TArray<TArray<int32>> Chaos::ComputeNodalColoring(const TArray<TVec4<int32>>& Graph, const Chaos::TDynamicParticles<T, 3>& InParticles, const int32 GraphParticlesStart, const int32 GraphParticlesEnd, const TArray<TArray<int32>>& IncidentElements, const TArray<TArray<int32>>& IncidentElementsLocalIndex)
+// ComputeNodalColoringImpl (TVec4 graph)
+static TArray<TArray<int32>> ComputeNodalColoringImpl(const TArray<Chaos::TVec4<int32>>& Graph, int32 NumParticles, const TFunctionRef<bool(int32)>& IsDynamicParticle, const int32 GraphParticlesStart, const int32 GraphParticlesEnd, const TArray<TArray<int32>>& IncidentElements, const TArray<TArray<int32>>& IncidentElementsLocalIndex)
 {
-	using namespace Chaos;
-
 	checkSlow(GraphParticlesStart <= GraphParticlesEnd);
-	checkSlow(GraphParticlesEnd <= (int32)InParticles.Size());
+	checkSlow(GraphParticlesEnd <= NumParticles);
 	TArray<TArray<int32>> ParticlesPerColor;
 
 	TArray<int32> Particle2Incident;
-	Particle2Incident.Init(-1, GraphParticlesEnd - GraphParticlesStart);
+	Particle2Incident.Init(INDEX_NONE, GraphParticlesEnd - GraphParticlesStart);
 	//Assuming that offset of Graph is GraphParticlesStart
-	for (int32 i = 0; i < IncidentElements.Num(); i++) 
+	for (int32 IncidentIdx = 0; IncidentIdx < IncidentElements.Num(); IncidentIdx++)
 	{
-		if (IncidentElements[i].Num() > 0)
+		if (IncidentElements[IncidentIdx].Num() > 0)
 		{
-			Particle2Incident[Graph[IncidentElements[i][0]][IncidentElementsLocalIndex[i][0]] - GraphParticlesStart] = i;
+			Particle2Incident[Graph[IncidentElements[IncidentIdx][0]][IncidentElementsLocalIndex[IncidentIdx][0]] - GraphParticlesStart] = IncidentIdx;
 		}
 	}
 
 	TArray<TSet<int32>> ElementColorsSet;
 	ElementColorsSet.SetNum(Graph.Num());
 	TArray<int32> ParticleColors;
-	ParticleColors.Init(-1, GraphParticlesEnd - GraphParticlesStart);
+	ParticleColors.Init(INDEX_NONE, GraphParticlesEnd - GraphParticlesStart);
 
-	for (int32 p = 0; p < GraphParticlesEnd - GraphParticlesStart; p++)
+	for (int32 LocalParticleIdx = 0; LocalParticleIdx < GraphParticlesEnd - GraphParticlesStart; LocalParticleIdx++)
 	{
-		int32 ParticleIndex = p + GraphParticlesStart;
-		if (InParticles.InvM(ParticleIndex) != (T)0. && Particle2Incident[ParticleIndex] != -1)
+		int32 ParticleIndex = LocalParticleIdx + GraphParticlesStart;
+		if (IsDynamicParticle(ParticleIndex) && Particle2Incident[ParticleIndex] != INDEX_NONE)
 		{
-			int32 color_to_use = 0;
+			int32 ColorToUse = 0;
 			while (true)
 			{
 				bool ColorFound = false;
-				for (int32 j = 0; j < IncidentElements[Particle2Incident[ParticleIndex]].Num(); j++)
+				for (int32 IncidentIdx = 0; IncidentIdx < IncidentElements[Particle2Incident[ParticleIndex]].Num(); IncidentIdx++)
 				{
-					int32 e = IncidentElements[Particle2Incident[ParticleIndex]][j];
-					if (ElementColorsSet[e].Contains(color_to_use))
+					int32 ElementIndex = IncidentElements[Particle2Incident[ParticleIndex]][IncidentIdx];
+					if (ElementColorsSet[ElementIndex].Contains(ColorToUse))
 					{
 						ColorFound = true;
 						break;
@@ -836,19 +731,19 @@ TArray<TArray<int32>> Chaos::ComputeNodalColoring(const TArray<TVec4<int32>>& Gr
 				}
 				if (!ColorFound)
 				{
-					ParticleColors[p] = color_to_use;
-					for (int32 j = 0; j < IncidentElements[Particle2Incident[ParticleIndex]].Num(); j++)
+					ParticleColors[LocalParticleIdx] = ColorToUse;
+					for (int32 IncidentIdx = 0; IncidentIdx < IncidentElements[Particle2Incident[ParticleIndex]].Num(); IncidentIdx++)
 					{
-						int32 e = IncidentElements[Particle2Incident[ParticleIndex]][j];
-						ElementColorsSet[e].Emplace(color_to_use);
-						if (ElementColorsSet[e].Num() == 4)
+						int32 ElementIndex = IncidentElements[Particle2Incident[ParticleIndex]][IncidentIdx];
+						ElementColorsSet[ElementIndex].Emplace(ColorToUse);
+						if (ElementColorsSet[ElementIndex].Num() == 4)
 						{
-							ElementColorsSet[e].Empty();
+							ElementColorsSet[ElementIndex].Empty();
 						}
 					}
 					break;
 				}
-				color_to_use++;
+				ColorToUse++;
 			}
 		}
 	}
@@ -858,36 +753,50 @@ TArray<TArray<int32>> Chaos::ComputeNodalColoring(const TArray<TVec4<int32>>& Gr
 	ParticlesPerColor.SetNum(0);
 	ParticlesPerColor.Init(TArray<int32>(), SizeColors + 1);
 
-	for (int32 j = 0; j < ParticleColors.Num(); j++) 
+	for (int32 LocalParticleIdx = 0; LocalParticleIdx < ParticleColors.Num(); LocalParticleIdx++)
 	{
-		if (ParticleColors[j] != -1) 
+		if (ParticleColors[LocalParticleIdx] != INDEX_NONE)
 		{
-			ParticlesPerColor[ParticleColors[j]].Emplace(j + GraphParticlesStart);
+			ParticlesPerColor[ParticleColors[LocalParticleIdx]].Emplace(LocalParticleIdx + GraphParticlesStart);
 		}
 	}
 
-	checkSlow(VerifyNodalColoring<T>(Graph, InParticles, GraphParticlesStart, GraphParticlesEnd, IncidentElements, IncidentElementsLocalIndex, ParticlesPerColor));
+	checkSlow(VerifyNodalColoringImpl(Graph, NumParticles, IsDynamicParticle, GraphParticlesStart, GraphParticlesEnd, IncidentElements, IncidentElementsLocalIndex, ParticlesPerColor));
 
 	return ParticlesPerColor;
 }
 
+// TDynamicParticles delegator
+template<typename T>
+TArray<TArray<int32>> Chaos::ComputeNodalColoring(const TArray<TVec4<int32>>& Graph, const Chaos::TDynamicParticles<T, 3>& InParticles, const int32 GraphParticlesStart, const int32 GraphParticlesEnd, const TArray<TArray<int32>>& IncidentElements, const TArray<TArray<int32>>& IncidentElementsLocalIndex)
+{
+	return ComputeNodalColoringImpl(Graph, (int32)InParticles.Size(),
+		[&InParticles](int32 Idx) { return InParticles.InvM(Idx) != (T)0.; },
+		GraphParticlesStart, GraphParticlesEnd, IncidentElements, IncidentElementsLocalIndex);
+}
 
-template<typename T, typename ParticleType>
-TArray<TArray<int32>> Chaos::ComputeNodalColoring(const TArray<TArray<int32>>& Graph, const Chaos::TDynamicParticles<T, 3>& InParticles, const int32 GraphParticlesStart, const int32 GraphParticlesEnd, const TArray<TArray<int32>>& IncidentElements, const TArray<TArray<int32>>& IncidentElementsLocalIndex, const TPBDActiveView<ParticleType>* InParticleActiveView, TArray<int32>* ParticleColorsOut)
+// FSolverParticlesRange overload
+TArray<TArray<int32>> Chaos::ComputeNodalColoring(const TArray<TVec4<int32>>& Graph, const Softs::FSolverParticlesRange& InParticles, const int32 GraphParticlesStart, const int32 GraphParticlesEnd, const TArray<TArray<int32>>& IncidentElements, const TArray<TArray<int32>>& IncidentElementsLocalIndex)
 {
-	using namespace Chaos;
+	return ComputeNodalColoringImpl(Graph, InParticles.Size(),
+		[&InParticles](int32 Idx) { return InParticles.InvM(Idx) != (Softs::FSolverReal)0.; },
+		GraphParticlesStart, GraphParticlesEnd, IncidentElements, IncidentElementsLocalIndex);
+}
+
 
+// ComputeNodalColoringImpl (TArray graph, ActiveView)
+template<typename ParticleType>
+static TArray<TArray<int32>> ComputeNodalColoringImpl(const TArray<TArray<int32>>& Graph, int32 NumParticles, const TFunctionRef<bool(int32)>& IsDynamicParticle, const int32 GraphParticlesStart, const int32 GraphParticlesEnd, const TArray<TArray<int32>>& IncidentElements, const TArray<TArray<int32>>& IncidentElementsLocalIndex, const Chaos::TPBDActiveView<ParticleType>* InParticleActiveView, TArray<int32>* ParticleColorsOut)
+{
 	checkSlow(GraphParticlesStart <= GraphParticlesEnd);
-	checkSlow(GraphParticlesEnd <= (int32)InParticles.Size());
-	checkSlow(InParticles.Size() == IncidentElements.Num());
+	checkSlow(GraphParticlesEnd <= NumParticles);
 	TArray<TArray<int32>> ParticlesPerColor;
 
 	TArray<int32> Particle2Incident;
 	Particle2Incident.Init(INDEX_NONE, GraphParticlesEnd - GraphParticlesStart);
-	//Assuming that offset of Graph is GraphParticlesStart
-	for (int32 i = 0; i < GraphParticlesEnd-GraphParticlesStart; i++)
+	for (int32 LocalParticleIdx = 0; LocalParticleIdx < GraphParticlesEnd - GraphParticlesStart; LocalParticleIdx++)
 	{
-		Particle2Incident[i] = i+GraphParticlesStart;
+		Particle2Incident[LocalParticleIdx] = LocalParticleIdx + GraphParticlesStart;
 	}
 
 	TArray<TSet<int32>> ElementColorsSet;
@@ -897,20 +806,20 @@ TArray<TArray<int32>> Chaos::ComputeNodalColoring(const TArray<TArray<int32>>& G
 
 	if (InParticleActiveView)
 	{
-		InParticleActiveView->SequentialFor([&ElementColorsSet, &Graph, &GraphParticlesStart, &IncidentElements, &Particle2Incident, &ParticleColors]
-		(const Chaos::TDynamicParticles<T, 3>& InParticles, int32 ParticleIndex)
+		InParticleActiveView->SequentialFor([&IsDynamicParticle, &ElementColorsSet, &Graph, &GraphParticlesStart, &IncidentElements, &Particle2Incident, &ParticleColors]
+		(const ParticleType& ActiveViewParticles, int32 ParticleIndex)
 			{
-				const int32 p = ParticleIndex - GraphParticlesStart;
-				if (InParticles.InvM(ParticleIndex) != (T)0. && Particle2Incident[ParticleIndex] != INDEX_NONE)
+				const int32 LocalParticleIdx = ParticleIndex - GraphParticlesStart;
+				if (IsDynamicParticle(ParticleIndex) && Particle2Incident[ParticleIndex] != INDEX_NONE)
 				{
-					int32 color_to_use = 0;
+					int32 ColorToUse = 0;
 					while (true)
 					{
 						bool ColorFound = false;
-						for (int32 j = 0; j < IncidentElements[Particle2Incident[ParticleIndex]].Num(); j++)
+						for (int32 IncidentIdx = 0; IncidentIdx < IncidentElements[Particle2Incident[ParticleIndex]].Num(); IncidentIdx++)
 						{
-							int32 e = IncidentElements[Particle2Incident[ParticleIndex]][j];
-							if (ElementColorsSet[e].Contains(color_to_use))
+							int32 ElementIndex = IncidentElements[Particle2Incident[ParticleIndex]][IncidentIdx];
+							if (ElementColorsSet[ElementIndex].Contains(ColorToUse))
 							{
 								ColorFound = true;
 								break;
@@ -918,38 +827,38 @@ TArray<TArray<int32>> Chaos::ComputeNodalColoring(const TArray<TArray<int32>>& G
 						}
 						if (!ColorFound)
 						{
-							ParticleColors[p] = color_to_use;
-							for (int32 j = 0; j < IncidentElements[Particle2Incident[ParticleIndex]].Num(); j++)
+							ParticleColors[LocalParticleIdx] = ColorToUse;
+							for (int32 IncidentIdx = 0; IncidentIdx < IncidentElements[Particle2Incident[ParticleIndex]].Num(); IncidentIdx++)
 							{
-								int32 e = IncidentElements[Particle2Incident[ParticleIndex]][j];
-								ElementColorsSet[e].Emplace(color_to_use);
-								if (ElementColorsSet[e].Num() == Graph[e].Num())
+								int32 ElementIndex = IncidentElements[Particle2Incident[ParticleIndex]][IncidentIdx];
+								ElementColorsSet[ElementIndex].Emplace(ColorToUse);
+								if (ElementColorsSet[ElementIndex].Num() == Graph[ElementIndex].Num())
 								{
-									ElementColorsSet[e].Empty();
+									ElementColorsSet[ElementIndex].Empty();
 								}
 							}
 							break;
 						}
-						color_to_use++;
+						ColorToUse++;
 					}
 				}
 			});
 	}
 	else
 	{
-		for (int32 p = 0; p < GraphParticlesEnd - GraphParticlesStart; p++)
+		for (int32 LocalParticleIdx = 0; LocalParticleIdx < GraphParticlesEnd - GraphParticlesStart; LocalParticleIdx++)
 		{
-			int32 ParticleIndex = p + GraphParticlesStart;
-			if (InParticles.InvM(ParticleIndex) != (T)0. && Particle2Incident[ParticleIndex] != INDEX_NONE)
+			int32 ParticleIndex = LocalParticleIdx + GraphParticlesStart;
+			if (IsDynamicParticle(ParticleIndex) && Particle2Incident[ParticleIndex] != INDEX_NONE)
 			{
-				int32 color_to_use = 0;
+				int32 ColorToUse = 0;
 				while (true)
 				{
 					bool ColorFound = false;
-					for (int32 j = 0; j < IncidentElements[Particle2Incident[ParticleIndex]].Num(); j++)
+					for (int32 IncidentIdx = 0; IncidentIdx < IncidentElements[Particle2Incident[ParticleIndex]].Num(); IncidentIdx++)
 					{
-						int32 e = IncidentElements[Particle2Incident[ParticleIndex]][j];
-						if (ElementColorsSet[e].Contains(color_to_use))
+						int32 ElementIndex = IncidentElements[Particle2Incident[ParticleIndex]][IncidentIdx];
+						if (ElementColorsSet[ElementIndex].Contains(ColorToUse))
 						{
 							ColorFound = true;
 							break;
@@ -957,19 +866,19 @@ TArray<TArray<int32>> Chaos::ComputeNodalColoring(const TArray<TArray<int32>>& G
 					}
 					if (!ColorFound)
 					{
-						ParticleColors[p] = color_to_use;
-						for (int32 j = 0; j < IncidentElements[Particle2Incident[ParticleIndex]].Num(); j++)
+						ParticleColors[LocalParticleIdx] = ColorToUse;
+						for (int32 IncidentIdx = 0; IncidentIdx < IncidentElements[Particle2Incident[ParticleIndex]].Num(); IncidentIdx++)
 						{
-							int32 e = IncidentElements[Particle2Incident[ParticleIndex]][j];
-							ElementColorsSet[e].Emplace(color_to_use);
-							if (ElementColorsSet[e].Num() == Graph[e].Num())
+							int32 ElementIndex = IncidentElements[Particle2Incident[ParticleIndex]][IncidentIdx];
+							ElementColorsSet[ElementIndex].Emplace(ColorToUse);
+							if (ElementColorsSet[ElementIndex].Num() == Graph[ElementIndex].Num())
 							{
-								ElementColorsSet[e].Empty();
+								ElementColorsSet[ElementIndex].Empty();
 							}
 						}
 						break;
 					}
-					color_to_use++;
+					ColorToUse++;
 				}
 			}
 		}
@@ -980,11 +889,11 @@ TArray<TArray<int32>> Chaos::ComputeNodalColoring(const TArray<TArray<int32>>& G
 	ParticlesPerColor.SetNum(0);
 	ParticlesPerColor.Init(TArray<int32>(), SizeColors + 1);
 
-	for (int32 j = 0; j < ParticleColors.Num(); j++) 
+	for (int32 LocalParticleIdx = 0; LocalParticleIdx < ParticleColors.Num(); LocalParticleIdx++)
 	{
-		if (ParticleColors[j] != INDEX_NONE) 
+		if (ParticleColors[LocalParticleIdx] != INDEX_NONE)
 		{
-			ParticlesPerColor[ParticleColors[j]].Emplace(j + GraphParticlesStart);
+			ParticlesPerColor[ParticleColors[LocalParticleIdx]].Emplace(LocalParticleIdx + GraphParticlesStart);
 		}
 	}
 
@@ -993,103 +902,34 @@ TArray<TArray<int32>> Chaos::ComputeNodalColoring(const TArray<TArray<int32>>& G
 		*ParticleColorsOut = ParticleColors;
 	}
 
-	checkSlow(VerifyNodalColoring<T>(Graph, InParticles, GraphParticlesStart, GraphParticlesEnd, IncidentElements, IncidentElementsLocalIndex, ParticlesPerColor));
+	checkSlow(VerifyNodalColoringImpl(Graph, NumParticles, IsDynamicParticle, GraphParticlesStart, GraphParticlesEnd, IncidentElements, IncidentElementsLocalIndex, ParticlesPerColor));
 
 	return ParticlesPerColor;
 }
 
-
-template<typename T>
-void Chaos::ComputeExtraNodalColoring(const TArray<TArray<int32>>& Graph, const TArray<TArray<int32>>& ExtraGraph, const Chaos::TDynamicParticles<T, 3>& InParticles, const TArray<TArray<int32>>& IncidentElements, const TArray<TArray<int32>>& ExtraIncidentElements, TArray<int32>& ParticleColors, TArray<TArray<int32>>& ParticlesPerColor)
+// TDynamicParticles delegator
+template<typename T, typename ParticleType>
+TArray<TArray<int32>> Chaos::ComputeNodalColoring(const TArray<TArray<int32>>& Graph, const Chaos::TDynamicParticles<T, 3>& InParticles, const int32 GraphParticlesStart, const int32 GraphParticlesEnd, const TArray<TArray<int32>>& IncidentElements, const TArray<TArray<int32>>& IncidentElementsLocalIndex, const TPBDActiveView<ParticleType>* InParticleActiveView, TArray<int32>* ParticleColorsOut)
 {
-	TBitArray ParticleIsAffected(false, IncidentElements.Num());
-	for (int32 i = 0; i < ExtraIncidentElements.Num(); i++)
-	{
-		if (ExtraIncidentElements[i].Num() > 0)
-		{
-			ParticleIsAffected[i] = true;
-		}
-	}
-
-	TSet<int32> UsedColors;
-	for (int32 i = 0; i < ExtraIncidentElements.Num(); i++)
-	{
-		if (ParticleIsAffected[i])
-		{
-			int32 OriginalColor = ParticleColors[i];
-			ParticleColors[i] = INDEX_NONE;
-			if (OriginalColor != INDEX_NONE)
-			{
-				UsedColors.Reset();
-				UsedColors.Reserve(ExtraIncidentElements[i].Num() + IncidentElements[i].Num());
-				for (int32 j = 0; j < IncidentElements[i].Num(); j++)
-				{
-					int32 ConstraintIndex = IncidentElements[i][j];
-					for (int32 ie = 0; ie < Graph[ConstraintIndex].Num(); ie++)
-					{
-						UsedColors.Add(ParticleColors[Graph[ConstraintIndex][ie]]);
-					}
-				}
-				for (int32 j = 0; j < ExtraIncidentElements[i].Num(); j++)
-				{
-					int32 ConstraintIndex = ExtraIncidentElements[i][j];
-					for (int32 ie = 0; ie < ExtraGraph[ConstraintIndex].Num(); ie++)
-					{
-						UsedColors.Add(ParticleColors[ExtraGraph[ConstraintIndex][ie]]);
-					}
-				}
-				if (UsedColors.Contains(OriginalColor))
-				{
-					OriginalColor = 0;
-					while (UsedColors.Contains(OriginalColor))
-					{
-						OriginalColor++;
-					}
-					ParticleColors[i] = OriginalColor;
-				}
-				else
-				{
-					ParticleColors[i] = OriginalColor;
-					ParticleIsAffected[i] = false;
-				}
-			}
-		}
-	}
-
-
-	PhysicsParallelFor(ParticlesPerColor.Num(), [&](const int32 i)
-		{
-			int32 CurrentIndex = 0;
-			for (int32 j = 0; j < ParticlesPerColor[i].Num(); j++)
-			{
-				if (!ParticleIsAffected[ParticlesPerColor[i][j]])
-				{
-					if (CurrentIndex != j)
-					{
-						ParticlesPerColor[i][CurrentIndex] = ParticlesPerColor[i][j];
-					}
-					CurrentIndex += 1;
-				}
-			}
-			ParticlesPerColor[i].SetNum(CurrentIndex);
-		}, ParticlesPerColor.Num() > 0 && ParticlesPerColor[0].Num() < 1000);
-
-	ParticlesPerColor.SetNum(FMath::Max<int32>(ParticleColors) + 1);
-	for (int32 i = 0; i < ExtraIncidentElements.Num(); i++)
-	{
-		if (ParticleIsAffected[i] && ParticleColors[i] != INDEX_NONE)
-		{
-			ParticlesPerColor[ParticleColors[i]].Emplace(i);
-		}
-	}
+	return ComputeNodalColoringImpl(Graph, (int32)InParticles.Size(),
+		[&InParticles](int32 Idx) { return InParticles.InvM(Idx) != (T)0.; },
+		GraphParticlesStart, GraphParticlesEnd, IncidentElements, IncidentElementsLocalIndex, InParticleActiveView, ParticleColorsOut);
+}
 
-	checkSlow(VerifyExtraNodalColoring<T>(InParticles, Graph, ExtraGraph, IncidentElements, ExtraIncidentElements, ParticleColors, ParticlesPerColor));
+// FSolverParticlesRange overload
+template<typename ParticleType>
+TArray<TArray<int32>> Chaos::ComputeNodalColoring(const TArray<TArray<int32>>& Graph, const Softs::FSolverParticlesRange& InParticles, const int32 GraphParticlesStart, const int32 GraphParticlesEnd, const TArray<TArray<int32>>& IncidentElements, const TArray<TArray<int32>>& IncidentElementsLocalIndex, const TPBDActiveView<ParticleType>* InParticleActiveView, TArray<int32>* ParticleColorsOut)
+{
+	return ComputeNodalColoringImpl(Graph, InParticles.Size(),
+		[&InParticles](int32 Idx) { return InParticles.InvM(Idx) != (Softs::FSolverReal)0.; },
+		GraphParticlesStart, GraphParticlesEnd, IncidentElements, IncidentElementsLocalIndex, InParticleActiveView, ParticleColorsOut);
 }
 
-template<typename T>
-void Chaos::ComputeExtraNodalColoring(const TArray<TArray<int32>>& StaticGraph, const TArray<TArray<int32>>& DynamicGraph, const TArray<TArray<int32>>& ExtraGraph, const Chaos::TDynamicParticles<T, 3>& InParticles, const TArray<TArray<int32>>& StaticIncidentElements, const TArray<TArray<int32>>& DynamicIncidentElements, const TArray<TArray<int32>>& ExtraIncidentElements, TArray<int32>& ParticleColors, TArray<TArray<int32>>& ParticlesPerColor)
+
+// ComputeExtraNodalColoringImpl (3-graph)
+static void ComputeExtraNodalColoringImpl(const TArray<TArray<int32>>& StaticGraph, const TArray<TArray<int32>>& DynamicGraph, const TArray<TArray<int32>>& ExtraGraph, int32 NumParticles, const TFunctionRef<bool(int32)>& IsDynamicParticle, const TArray<TArray<int32>>& StaticIncidentElements, const TArray<TArray<int32>>& DynamicIncidentElements, const TArray<TArray<int32>>& ExtraIncidentElements, TArray<int32>& ParticleColors, TArray<TArray<int32>>& ParticlesPerColor)
 {
-	TBitArray ParticleIsAffected(false, InParticles.Size());
+	TBitArray ParticleIsAffected(false, NumParticles);
 	for (int32 i = 0; i < ExtraIncidentElements.Num(); i++)
 	{
 		if (ExtraIncidentElements[i].Num() > 0)
@@ -1152,7 +992,7 @@ void Chaos::ComputeExtraNodalColoring(const TArray<TArray<int32>>& StaticGraph,
 	}
 
 
-	PhysicsParallelFor(ParticlesPerColor.Num(), [&](const int32 i)
+	Chaos::PhysicsParallelFor(ParticlesPerColor.Num(), [&](const int32 i)
 		{
 			int32 CurrentIndex = 0;
 			for (int32 j = 0; j < ParticlesPerColor[i].Num(); j++)
@@ -1178,7 +1018,52 @@ void Chaos::ComputeExtraNodalColoring(const TArray<TArray<int32>>& StaticGraph,
 		}
 	}
 
-	checkSlow(VerifyExtraNodalColoring<T>(InParticles, StaticGraph, DynamicGraph, ExtraGraph, StaticIncidentElements, DynamicIncidentElements, ExtraIncidentElements, ParticleColors, ParticlesPerColor));
+	checkSlow(VerifyExtraNodalColoringImpl(NumParticles, IsDynamicParticle, StaticGraph, DynamicGraph, ExtraGraph, StaticIncidentElements, DynamicIncidentElements, ExtraIncidentElements, ParticleColors, ParticlesPerColor));
+}
+
+// TDynamicParticles delegator
+template<typename T>
+void Chaos::ComputeExtraNodalColoring(const TArray<TArray<int32>>& StaticGraph, const TArray<TArray<int32>>& DynamicGraph, const TArray<TArray<int32>>& ExtraGraph, const Chaos::TDynamicParticles<T, 3>& InParticles, const TArray<TArray<int32>>& StaticIncidentElements, const TArray<TArray<int32>>& DynamicIncidentElements, const TArray<TArray<int32>>& ExtraIncidentElements, TArray<int32>& ParticleColors, TArray<TArray<int32>>& ParticlesPerColor)
+{
+	ComputeExtraNodalColoringImpl(StaticGraph, DynamicGraph, ExtraGraph, (int32)InParticles.Size(),
+		[&InParticles](int32 Idx) { return InParticles.InvM(Idx) != (T)0.; },
+		StaticIncidentElements, DynamicIncidentElements, ExtraIncidentElements, ParticleColors, ParticlesPerColor);
+}
+
+// FSolverParticlesRange overload
+void Chaos::ComputeExtraNodalColoring(const TArray<TArray<int32>>& StaticGraph, const TArray<TArray<int32>>& DynamicGraph, const TArray<TArray<int32>>& ExtraGraph, const Softs::FSolverParticlesRange& InParticles, const TArray<TArray<int32>>& StaticIncidentElements, const TArray<TArray<int32>>& DynamicIncidentElements, const TArray<TArray<int32>>& ExtraIncidentElements, TArray<int32>& ParticleColors, TArray<TArray<int32>>& ParticlesPerColor)
+{
+	ComputeExtraNodalColoringImpl(StaticGraph, DynamicGraph, ExtraGraph, InParticles.Size(),
+		[&InParticles](int32 Idx) { return InParticles.InvM(Idx) != (Softs::FSolverReal)0.; },
+		StaticIncidentElements, DynamicIncidentElements, ExtraIncidentElements, ParticleColors, ParticlesPerColor);
+}
+
+// ComputeExtraNodalColoringImpl (2-graph) — delegates to 3-graph with empty DynamicGraph
+static void ComputeExtraNodalColoringImpl(const TArray<TArray<int32>>& Graph, const TArray<TArray<int32>>& ExtraGraph, int32 NumParticles, const TFunctionRef<bool(int32)>& IsDynamicParticle, const TArray<TArray<int32>>& IncidentElements, const TArray<TArray<int32>>& ExtraIncidentElements, TArray<int32>& ParticleColors, TArray<TArray<int32>>& ParticlesPerColor)
+{
+	// EmptyGraph doesn't need sizing — it's only indexed by constraint indices from DynamicIncidentElements, which are all empty
+	const TArray<TArray<int32>> EmptyGraph;
+	// Size must match NumParticles so the 3-graph version can safely index DynamicIncidentElements[ParticleIndex]
+	TArray<TArray<int32>> EmptyIncidentElements;
+	EmptyIncidentElements.SetNum(NumParticles);
+	ComputeExtraNodalColoringImpl(Graph, EmptyGraph, ExtraGraph, NumParticles, IsDynamicParticle, IncidentElements, EmptyIncidentElements, ExtraIncidentElements, ParticleColors, ParticlesPerColor);
+}
+
+// TDynamicParticles delegator (2-graph)
+template<typename T>
+void Chaos::ComputeExtraNodalColoring(const TArray<TArray<int32>>& Graph, const TArray<TArray<int32>>& ExtraGraph, const Chaos::TDynamicParticles<T, 3>& InParticles, const TArray<TArray<int32>>& IncidentElements, const TArray<TArray<int32>>& ExtraIncidentElements, TArray<int32>& ParticleColors, TArray<TArray<int32>>& ParticlesPerColor)
+{
+	ComputeExtraNodalColoringImpl(Graph, ExtraGraph, (int32)InParticles.Size(),
+		[&InParticles](int32 Idx) { return InParticles.InvM(Idx) != (T)0.; },
+		IncidentElements, ExtraIncidentElements, ParticleColors, ParticlesPerColor);
+}
+
+// FSolverParticlesRange overload (2-graph)
+void Chaos::ComputeExtraNodalColoring(const TArray<TArray<int32>>& Graph, const TArray<TArray<int32>>& ExtraGraph, const Softs::FSolverParticlesRange& InParticles, const TArray<TArray<int32>>& IncidentElements, const TArray<TArray<int32>>& ExtraIncidentElements, TArray<int32>& ParticleColors, TArray<TArray<int32>>& ParticlesPerColor)
+{
+	ComputeExtraNodalColoringImpl(Graph, ExtraGraph, InParticles.Size(),
+		[&InParticles](int32 Idx) { return InParticles.InvM(Idx) != (Softs::FSolverReal)0.; },
+		IncidentElements, ExtraIncidentElements, ParticleColors, ParticlesPerColor);
 }
 
 #define UE_COMPUTE_GRAPH_COLORING_PARTICLES_OR_RANGE_N_HELPER(N, bAllDynamic) \
@@ -1199,9 +1084,6 @@ template CHAOS_API void Chaos::ComputeGridBasedGraphSubColoringPointer(const TAr
 template CHAOS_API void Chaos::ComputeWeakConstraintsColoring<Chaos::FRealSingle>(const TArray<TArray<int32>>& Indices, const TArray<TArray<int32>>& SecondIndices, const Chaos::TDynamicParticles<Chaos::FRealSingle, 3>& InParticles, TArray<TArray<int32>>& ConstraintsPerColor);
 template CHAOS_API void Chaos::ComputeWeakConstraintsColoring<Chaos::FRealDouble>(const TArray<TArray<int32>>& Indices, const TArray<TArray<int32>>& SecondIndices, const Chaos::TDynamicParticles<Chaos::FRealDouble, 3>& InParticles, TArray<TArray<int32>>& ConstraintsPerColor);
 
-template CHAOS_API void Chaos::ComputeWeakConstraintsColoring<Chaos::FRealSingle>(const TArray<TArray<int32>>& Indices, const TArray<TArray<int32>>& SecondIndices, const Chaos::TDynamicParticles<Chaos::FRealSingle, 3>& InParticles, const Chaos::TDynamicParticles<FRealSingle, 3>& InSecondParticles, TArray<TArray<int32>>& ConstraintsPerColor);
-template CHAOS_API void Chaos::ComputeWeakConstraintsColoring<Chaos::FRealDouble>(const TArray<TArray<int32>>& Indices, const TArray<TArray<int32>>& SecondIndices, const Chaos::TDynamicParticles<Chaos::FRealDouble, 3>& InParticles, const Chaos::TDynamicParticles<FRealDouble, 3>& InSecondParticles, TArray<TArray<int32>>& ConstraintsPerColor);
-
 template CHAOS_API TArray<TArray<int32>> Chaos::ComputeNodalColoring<Chaos::FRealSingle>(const TArray<TVec4<int32>>& Graph, const Chaos::TDynamicParticles<Chaos::FRealSingle, 3>& InParticles, const int32 GraphParticlesStart, const int32 GraphParticlesEnd, const TArray<TArray<int32>>& IncidentElements, const TArray<TArray<int32>>& IncidentElementsLocalIndex);
 template CHAOS_API TArray<TArray<int32>> Chaos::ComputeNodalColoring<Chaos::FRealDouble>(const TArray<TVec4<int32>>& Graph, const Chaos::TDynamicParticles<Chaos::FRealDouble, 3>& InParticles, const int32 GraphParticlesStart, const int32 GraphParticlesEnd, const TArray<TArray<int32>>& IncidentElements, const TArray<TArray<int32>>& IncidentElementsLocalIndex);
 
@@ -1212,4 +1094,7 @@ template CHAOS_API void Chaos::ComputeExtraNodalColoring<Chaos::FRealSingle>(con
 template CHAOS_API void Chaos::ComputeExtraNodalColoring<Chaos::FRealDouble>(const TArray<TArray<int32>>& Graph, const TArray<TArray<int32>>& ExtraGraph, const Chaos::TDynamicParticles<Chaos::FRealDouble, 3>& InParticles, const TArray<TArray<int32>>& IncidentElements, const TArray<TArray<int32>>& ExtraIncidentElements, TArray<int32>& ParticleColors, TArray<TArray<int32>>& ParticlesPerColor);
 
 template CHAOS_API void Chaos::ComputeExtraNodalColoring<Chaos::FRealSingle>(const TArray<TArray<int32>>& StaticGraph, const TArray<TArray<int32>>& DynamicGraph, const TArray<TArray<int32>>& ExtraGraph, const Chaos::TDynamicParticles<Chaos::FRealSingle, 3>& InParticles, const TArray<TArray<int32>>& StaticIncidentElements, const TArray<TArray<int32>>& DynamicIncidentElements, const TArray<TArray<int32>>& ExtraIncidentElements, TArray<int32>& ParticleColors, TArray<TArray<int32>>& ParticlesPerColor);
-template CHAOS_API void Chaos::ComputeExtraNodalColoring<Chaos::FRealDouble>(const TArray<TArray<int32>>& StaticGraph, const TArray<TArray<int32>>& DynamicGraph, const TArray<TArray<int32>>& ExtraGraph, const Chaos::TDynamicParticles<Chaos::FRealDouble, 3>& InParticles, const TArray<TArray<int32>>& StaticIncidentElements, const TArray<TArray<int32>>& DynamicIncidentElements, const TArray<TArray<int32>>& ExtraIncidentElements, TArray<int32>& ParticleColors, TArray<TArray<int32>>& ParticlesPerColor);
\ No newline at end of file
+template CHAOS_API void Chaos::ComputeExtraNodalColoring<Chaos::FRealDouble>(const TArray<TArray<int32>>& StaticGraph, const TArray<TArray<int32>>& DynamicGraph, const TArray<TArray<int32>>& ExtraGraph, const Chaos::TDynamicParticles<Chaos::FRealDouble, 3>& InParticles, const TArray<TArray<int32>>& StaticIncidentElements, const TArray<TArray<int32>>& DynamicIncidentElements, const TArray<TArray<int32>>& ExtraIncidentElements, TArray<int32>& ParticleColors, TArray<TArray<int32>>& ParticlesPerColor);
+
+// FSolverParticlesRange explicit instantiation (only needed for the ActiveView templated overload)
+template CHAOS_API TArray<TArray<int32>> Chaos::ComputeNodalColoring<Chaos::Softs::FSolverParticles>(const TArray<TArray<int32>>& Graph, const Chaos::Softs::FSolverParticlesRange& InParticles, const int32 GraphParticlesStart, const int32 GraphParticlesEnd, const TArray<TArray<int32>>& IncidentElements, const TArray<TArray<int32>>& IncidentElementsLocalIndex, const TPBDActiveView<Chaos::Softs::FSolverParticles>* InParticleActiveView, TArray<int32>* ParticleColorsOut);
\ No newline at end of file
diff --git a/Engine/Source/Runtime/Experimental/Chaos/Private/Chaos/XPBDAnisotropicBendingConstraints.cpp b/Engine/Source/Runtime/Experimental/Chaos/Private/Chaos/XPBDAnisotropicBendingConstraints.cpp
--- a/Engine/Source/Runtime/Experimental/Chaos/Private/Chaos/XPBDAnisotropicBendingConstraints.cpp
+++ b/Engine/Source/Runtime/Experimental/Chaos/Private/Chaos/XPBDAnisotropicBendingConstraints.cpp
@@ -1227,7 +1227,7 @@ void FXPBDAnisotropicBendingConstraints::ComputeGradTheta(const FSolverVec3& X0,
 		}
 }
 
-void FXPBDAnisotropicBendingConstraints::AddAnisotropicBendingResidualAndHessian(const FSolverParticles& Particles, const int32 ConstraintIndex, const int32 ConstraintIndexLocal, const FSolverReal Dt, TVec3<FSolverReal>& ParticleResidual, Chaos::PMatrix<FSolverReal, 3, 3>& ParticleHessian)
+void FXPBDAnisotropicBendingConstraints::AddAnisotropicBendingResidualAndHessian(const FSolverParticlesRange& Particles, const int32 ConstraintIndex, const int32 ConstraintIndexLocal, const FSolverReal Dt, TVec3<FSolverReal>& ParticleResidual, Chaos::PMatrix<FSolverReal, 3, 3>& ParticleHessian)
 {
 	const FSolverVec3 ExpStiffnessValue((FSolverReal)StiffnessWeft, (FSolverReal)StiffnessWarp, (FSolverReal)StiffnessBias);
 	const FSolverVec3 ExpBucklingValue((FSolverReal)BucklingStiffnessWeft, (FSolverReal)BucklingStiffnessWarp, (FSolverReal)BucklingStiffnessBias);
@@ -1267,16 +1267,17 @@ void FXPBDAnisotropicBendingConstraints::AddAnisotropicBendingResidualAndHessian
 	}
 }
 
+void FXPBDAnisotropicBendingConstraints::AddAnisotropicBendingResidualAndHessian(const FSolverParticles& Particles, const int32 ConstraintIndex, const int32 ConstraintIndexLocal, const FSolverReal Dt, TVec3<FSolverReal>& ParticleResidual, Chaos::PMatrix<FSolverReal, 3, 3>& ParticleHessian)
+{
+	AddAnisotropicBendingResidualAndHessian(FSolverParticlesRange(Particles), ConstraintIndex, ConstraintIndexLocal, Dt, ParticleResidual, ParticleHessian);
+}
+
 
-void FXPBDAnisotropicBendingConstraints::AddInternalForceDifferential(const FSolverParticles& InParticles, const TArray<TVector<FSolverReal, 3>>& DeltaParticles, TArray<TVector<FSolverReal, 3>>& ndf)
+void FXPBDAnisotropicBendingConstraints::AddInternalForceDifferential(const FSolverParticlesRange& InParticles, const TArray<TVector<FSolverReal, 3>>& DeltaParticles, TArray<TVector<FSolverReal, 3>>& Ndf)
 {
 	const FSolverVec3 ExpStiffnessValue((FSolverReal)StiffnessWeft, (FSolverReal)StiffnessWarp, (FSolverReal)StiffnessBias);
 	const FSolverVec3 ExpBucklingValue((FSolverReal)BucklingStiffnessWeft, (FSolverReal)BucklingStiffnessWarp, (FSolverReal)BucklingStiffnessBias);
-	const FSolverReal DampingRatioValue = (FSolverReal)DampingRatio;
-
-	int32 ParticleStart = 0;
-	int32 ParticleNum = InParticles.Size();
-	ensure(ndf.Num() == InParticles.Size());
+	ensure(Ndf.Num() == InParticles.Size());
 
 	const int32 ConstraintColorNum = ConstraintsPerColorStartIndex.Num() - 1;
 
@@ -1288,11 +1289,6 @@ void FXPBDAnisotropicBendingConstraints::AddInternalForceDifferential(const FSol
 			{
 				const int32 ConstraintIndex = ColorStart + Index;
 				const TVec4<int32>& Constraint = Constraints[ConstraintIndex];
-				const int32 i1 = Constraint[0];
-				const int32 i2 = Constraint[1];
-				const int32 i3 = Constraint[2];
-				const int32 i4 = Constraint[3];
-
 				const FSolverVec3& P1 = InParticles.P(Constraint[0]);
 				const FSolverVec3& P2 = InParticles.P(Constraint[1]);
 				const FSolverVec3& P3 = InParticles.P(Constraint[2]);
@@ -1300,14 +1296,14 @@ void FXPBDAnisotropicBendingConstraints::AddInternalForceDifferential(const FSol
 				const FSolverVec3& WarpWeftBiasBaseMultiplier = WarpWeftBiasBaseMultipliers[ConstraintIndex];
 				FSolverReal Theta = 0.f;
 
-				TArray<FSolverVec3> Alldthetadx;
-				Alldthetadx.Init(FSolverVec3(0.f), 4);
+				TArray<FSolverVec3> AllDThetaDx;
+				AllDThetaDx.Init(FSolverVec3(0.f), 4);
 
 				constexpr int32 LocalIndexMap[] = { 2, 1, 0, 3 };
-				for (int32 i = 0; i < 4; i++)
+				for (int32 LocalIdx = 0; LocalIdx < 4; LocalIdx++)
 				{
-					int32 ActualConstraintIndexLocal = LocalIndexMap[i];
-					ComputeGradTheta(P3, P2, P1, P4, ActualConstraintIndexLocal, Alldthetadx[i], Theta);
+					const int32 ActualConstraintIndexLocal = LocalIndexMap[LocalIdx];
+					ComputeGradTheta(P3, P2, P1, P4, ActualConstraintIndexLocal, AllDThetaDx[LocalIdx], Theta);
 				}
 
 				const FSolverReal BiphasicStiffnessValue = IsBuckled[ConstraintIndex] ?
@@ -1315,23 +1311,28 @@ void FXPBDAnisotropicBendingConstraints::AddInternalForceDifferential(const FSol
 
 				FSolverReal DeltaC = 0.f;
 
-				for (int32 i = 0; i < 4; i++)
+				for (int32 LocalIdx = 0; LocalIdx < 4; LocalIdx++)
 				{
-					for (int32 j = 0; j < 3; j++)
+					for (int32 Alpha = 0; Alpha < 3; Alpha++)
 					{
-						DeltaC += Alldthetadx[i][j] * DeltaParticles[Constraint[i]][j];
+						DeltaC += AllDThetaDx[LocalIdx][Alpha] * DeltaParticles[Constraint[LocalIdx]][Alpha];
 					}
 				}
 
-				for (int32 i = 0; i < 4; i++)
+				for (int32 LocalIdx = 0; LocalIdx < 4; LocalIdx++)
 				{
-					ndf[Constraint[i]] += BiphasicStiffnessValue * DeltaC * Alldthetadx[i];
+					Ndf[Constraint[LocalIdx]] += BiphasicStiffnessValue * DeltaC * AllDThetaDx[LocalIdx];
 				}
 
 		});
 	}
 }
 
+void FXPBDAnisotropicBendingConstraints::AddInternalForceDifferential(const FSolverParticles& InParticles, const TArray<TVector<FSolverReal, 3>>& DeltaParticles, TArray<TVector<FSolverReal, 3>>& Ndf)
+{
+	AddInternalForceDifferential(FSolverParticlesRange(InParticles), DeltaParticles, Ndf);
+}
+
 
 
 
diff --git a/Engine/Source/Runtime/Experimental/Chaos/Private/Chaos/XPBDBendingConstraints.cpp b/Engine/Source/Runtime/Experimental/Chaos/Private/Chaos/XPBDBendingConstraints.cpp
--- a/Engine/Source/Runtime/Experimental/Chaos/Private/Chaos/XPBDBendingConstraints.cpp
+++ b/Engine/Source/Runtime/Experimental/Chaos/Private/Chaos/XPBDBendingConstraints.cpp
@@ -608,40 +608,31 @@ FSolverReal FXPBDBendingConstraints::ComputeTotalEnergy(const FSolverParticles&
 	for (int32 ConstraintIndex = 0; ConstraintIndex < Constraints.Num(); ConstraintIndex++)
 	{
 		const TVec4<int32>& Constraint = Constraints[ConstraintIndex];
-		const int32 i1 = Constraint[0];
-		const int32 i2 = Constraint[1];
-		const int32 i3 = Constraint[2];
-		const int32 i4 = Constraint[3];
-		//const FSolverVec3& WarpWeftBiasBaseMultiplier = WarpWeftBiasBaseMultipliers[ConstraintIndex];
 		const FSolverVec3& P1 = InParticles.P(Constraint[0]);
 		const FSolverVec3& P2 = InParticles.P(Constraint[1]);
 		const FSolverVec3& P3 = InParticles.P(Constraint[2]);
 		const FSolverVec3& P4 = InParticles.P(Constraint[3]);
 
-		//FSolverVec3 g0(0.f), g1(0.f);
+		const FSolverVec3 Y1 = P2 - P1;
+		const FSolverVec3 Y2 = P3 - P1;
+		const FSolverVec3 Y3 = P4 - P1;
+		const FSolverVec3 T1 = Y1.GetSafeNormal();
 
-		const FSolverVec3 y1 = P2 - P1;
-		const FSolverVec3 y2 = P3 - P1;
-		const FSolverVec3 y3 = P4 - P1;
-		const FSolverVec3 t1 = y1.GetSafeNormal();
-		const FSolverVec3 t2 = y2.GetSafeNormal();
-		const FSolverVec3 t3 = y3.GetSafeNormal();
+		const FSolverVec3 Z0 = Y2 - FSolverVec3::DotProduct(Y2, T1) * T1;
+		const FSolverVec3 Z1 = Y3 - FSolverVec3::DotProduct(Y3, T1) * T1;
 
-		const FSolverVec3 z0 = y2 - FSolverVec3::DotProduct(y2, t1) * t1;
-		const FSolverVec3 z1 = y3 - FSolverVec3::DotProduct(y3, t1) * t1;
+		const FSolverVec3 Z0Normalized = Z0.GetSafeNormal();
+		const FSolverVec3 Z1Normalized = Z1.GetSafeNormal();
 
-		FSolverVec3 z0Normalized = z0.GetSafeNormal();
-		FSolverVec3 z1Normalized = z1.GetSafeNormal();
+		const FSolverVec3 Tb = FSolverVec3::CrossProduct(Z0, Z1).GetSafeNormal();
 
-		FSolverVec3 tb = FSolverVec3::CrossProduct(z0, z1).GetSafeNormal();
+		const FSolverReal S = FSolverVec3::DotProduct(Tb, T1);
 
-		FSolverReal S = FSolverVec3::DotProduct(tb, t1);
-
-		FSolverReal Theta = FMath::Acos(FSolverVec3::DotProduct(z0Normalized, z1Normalized));
+		FSolverReal Theta = FMath::Acos(FSolverVec3::DotProduct(Z0Normalized, Z1Normalized));
 
 		Theta = S * (PI - Theta);
 
-		FSolverReal CAngle = Theta - RestAngles[ConstraintIndex];
+		const FSolverReal CAngle = Theta - RestAngles[ConstraintIndex];
 
 
 		Energy += CAngle * CAngle * StiffnessValue / 2.f;
@@ -690,17 +681,12 @@ void ComputeGradTheta(const FSolverVec3& X0, const FSolverVec3& X1, const FSolve
 }
 
 
-void FXPBDBendingConstraints::AddBendingResidualAndHessian(const FSolverParticles& Particles, const int32 ConstraintIndex, const int32 ConstraintIndexLocal, const FSolverReal Dt, TVec3<FSolverReal>& ParticleResidual, Chaos::PMatrix<FSolverReal, 3, 3>& ParticleHessian)
+void FXPBDBendingConstraints::AddBendingResidualAndHessian(const FSolverParticlesRange& Particles, const int32 ConstraintIndex, const int32 ConstraintIndexLocal, const FSolverReal Dt, TVec3<FSolverReal>& ParticleResidual, Chaos::PMatrix<FSolverReal, 3, 3>& ParticleHessian)
 {
-	FSolverReal ExpStiffnessValue = (FSolverReal)XPBDStiffness;
-	FSolverReal ExpBucklingValue = (FSolverReal)XPBDBucklingStiffness;
-	const FSolverReal DampingRatioValue = (FSolverReal)DampingRatio;
+	const FSolverReal ExpStiffnessValue = (FSolverReal)XPBDStiffness;
+	const FSolverReal ExpBucklingValue = (FSolverReal)XPBDBucklingStiffness;
 
 	const TVec4<int32>& Constraint = Constraints[ConstraintIndex];
-	const int32 i1 = Constraint[0];
-	const int32 i2 = Constraint[1];
-	const int32 i3 = Constraint[2];
-	const int32 i4 = Constraint[3];
 	const FSolverVec3& P1 = Particles.P(Constraint[0]);
 	const FSolverVec3& P2 = Particles.P(Constraint[1]);
 	const FSolverVec3& P3 = Particles.P(Constraint[2]);
@@ -726,27 +712,22 @@ void FXPBDBendingConstraints::AddBendingResidualAndHessian(const FSolverParticle
 	}
 }
 
-
-void FXPBDBendingConstraints::AddInternalForceDifferential(const FSolverParticles& InParticles, const TArray<TVector<FSolverReal, 3>>& DeltaParticles, TArray<TVector<FSolverReal, 3>>& ndf)
+void FXPBDBendingConstraints::AddBendingResidualAndHessian(const FSolverParticles& Particles, const int32 ConstraintIndex, const int32 ConstraintIndexLocal, const FSolverReal Dt, TVec3<FSolverReal>& ParticleResidual, Chaos::PMatrix<FSolverReal, 3, 3>& ParticleHessian)
 {
-	FSolverReal ExpStiffnessValue = (FSolverReal)XPBDStiffness;
-	FSolverReal ExpBucklingValue = (FSolverReal)XPBDBucklingStiffness;
-	const FSolverReal DampingRatioValue = (FSolverReal)DampingRatio;
+	AddBendingResidualAndHessian(FSolverParticlesRange(Particles), ConstraintIndex, ConstraintIndexLocal, Dt, ParticleResidual, ParticleHessian);
+}
 
 
-	int32 ParticleStart = 0;
-	int32 ParticleNum = InParticles.Size();
-	ensure(ndf.Num() == InParticles.Size());
+void FXPBDBendingConstraints::AddInternalForceDifferential(const FSolverParticlesRange& InParticles, const TArray<TVector<FSolverReal, 3>>& DeltaParticles, TArray<TVector<FSolverReal, 3>>& Ndf)
+{
+	const FSolverReal ExpStiffnessValue = (FSolverReal)XPBDStiffness;
+	const FSolverReal ExpBucklingValue = (FSolverReal)XPBDBucklingStiffness;
 
+	ensure(Ndf.Num() == InParticles.Size());
 
 	for (int32 ConstraintIndex = 0; ConstraintIndex < Constraints.Num(); ConstraintIndex++)
 	{
 		const TVec4<int32>& Constraint = Constraints[ConstraintIndex];
-		const int32 i1 = Constraint[0];
-		const int32 i2 = Constraint[1];
-		const int32 i3 = Constraint[2];
-		const int32 i4 = Constraint[3];
-
 		const FSolverVec3& P1 = InParticles.P(Constraint[0]);
 		const FSolverVec3& P2 = InParticles.P(Constraint[1]);
 		const FSolverVec3& P3 = InParticles.P(Constraint[2]);
@@ -757,31 +738,36 @@ void FXPBDBendingConstraints::AddInternalForceDifferential(const FSolverParticle
 		TArray<FSolverVec3> Alldthetadx;
 		Alldthetadx.Init(FSolverVec3(0.f), 4);
 
-		TArray<int32> LocalIndexMap = { 2, 1, 0, 3 };
-		for (int32 i = 0; i < 4; i++)
+		const TArray<int32> LocalIndexMap = { 2, 1, 0, 3 };
+		for (int32 Index = 0; Index < 4; Index++)
 		{
-			int32 ActualConstraintIndexLocal = LocalIndexMap[i];
-			ComputeGradTheta(P3, P2, P1, P4, ActualConstraintIndexLocal, Alldthetadx[i], Theta);
+			const int32 ActualConstraintIndexLocal = LocalIndexMap[Index];
+			ComputeGradTheta(P3, P2, P1, P4, ActualConstraintIndexLocal, Alldthetadx[Index], Theta);
 		}
 
-		FSolverReal BiphasicStiffnessValue = IsBuckled[ConstraintIndex] ? ExpBucklingValue : ExpStiffnessValue;
+		const FSolverReal BiphasicStiffnessValue = IsBuckled[ConstraintIndex] ? ExpBucklingValue : ExpStiffnessValue;
 
 		FSolverReal DeltaC = 0.f;
 
-		for (int32 i = 0; i < 4; i++)
+		for (int32 Index = 0; Index < 4; Index++)
 		{
-			for (int32 j = 0; j < 3; j++)
+			for (int32 Alpha = 0; Alpha < 3; Alpha++)
 			{
-				DeltaC += Alldthetadx[i][j] * DeltaParticles[Constraint[i]][j];
+				DeltaC += Alldthetadx[Index][Alpha] * DeltaParticles[Constraint[Index]][Alpha];
 			}
 		}
 
-		for (int32 i = 0; i < 4; i++)
+		for (int32 Index = 0; Index < 4; Index++)
 		{
-			ndf[Constraint[i]] += BiphasicStiffnessValue * DeltaC * Alldthetadx[i];
+			Ndf[Constraint[Index]] += BiphasicStiffnessValue * DeltaC * Alldthetadx[Index];
 		}
 
 	}
 }
 
+void FXPBDBendingConstraints::AddInternalForceDifferential(const FSolverParticles& InParticles, const TArray<TVector<FSolverReal, 3>>& DeltaParticles, TArray<TVector<FSolverReal, 3>>& Ndf)
+{
+	AddInternalForceDifferential(FSolverParticlesRange(InParticles), DeltaParticles, Ndf);
+}
+
 }  // End namespace Chaos::Softs
