ChaosDeformable: Constraint hierarchy refactor and FSolverParticlesRange migration

diff --git a/Engine/Source/Runtime/Experimental/Chaos/Private/Chaos/XPBDStretchBiasElementConstraints.cpp b/Engine/Source/Runtime/Experimental/Chaos/Private/Chaos/XPBDStretchBiasElementConstraints.cpp
--- a/Engine/Source/Runtime/Experimental/Chaos/Private/Chaos/XPBDStretchBiasElementConstraints.cpp
+++ b/Engine/Source/Runtime/Experimental/Chaos/Private/Chaos/XPBDStretchBiasElementConstraints.cpp
@@ -783,18 +783,18 @@ template<typename SolverParticlesOrRange>
 void FXPBDStretchBiasElementConstraints::ApplyHelper(SolverParticlesOrRange& Particles, const FSolverReal Dt, const int32 ConstraintIndex, const FSolverVec3& ExpStiffnessValue, const FSolverReal DampingRatioValue, const FSolverReal WarpScaleValue, const FSolverReal WeftScaleValue) const
 {
 	const TVec3<int32>& Constraint = Constraints[ConstraintIndex];
-	const int32 i0 = Constraint[0];
-	const int32 i1 = Constraint[1];
-	const int32 i2 = Constraint[2];
-	if (Particles.InvM(i0) == (FSolverReal)0. && Particles.InvM(i1) == (FSolverReal)0. && Particles.InvM(i2) == (FSolverReal)0.)
+	const int32 Idx0 = Constraint[0];
+	const int32 Idx1 = Constraint[1];
+	const int32 Idx2 = Constraint[2];
+	if (Particles.InvM(Idx0) == (FSolverReal)0. && Particles.InvM(Idx1) == (FSolverReal)0. && Particles.InvM(Idx2) == (FSolverReal)0.)
 	{
 		return;
 	}
 
-	const FSolverReal CombinedInvMass = Particles.InvM(i0) + Particles.InvM(i1) + Particles.InvM(i2);
+	const FSolverReal CombinedInvMass = Particles.InvM(Idx0) + Particles.InvM(Idx1) + Particles.InvM(Idx2);
 
 	FSolverVec3 DXDu, DXDv;
-	CalculateUVStretch(ConstraintIndex, Particles.P(i0), Particles.P(i1), Particles.P(i2), DXDu, DXDv);
+	CalculateUVStretch(ConstraintIndex, Particles.P(Idx0), Particles.P(Idx1), Particles.P(Idx2), DXDu, DXDv);
 
 	const FSolverReal DXDuLength = DXDu.Length();
 	const FSolverReal DXDvLength = DXDv.Length();
@@ -829,9 +829,9 @@ void FXPBDStretchBiasElementConstraints::ApplyHelper(SolverParticlesOrRange& Par
 	const FSolverVec3 GradCs2 = (DXDvNormalized * OneOverDXDuLen * DeltaUVInv.M[1] + DXDuNormalized * OneOverDXDvLen * DeltaUVInv.M[3]);
 	const FSolverVec3 GradCs0 = -GradCs1 - GradCs2;
 
-	const FSolverVec3 V0TimesDt = Particles.GetP(i0) - Particles.GetX(i0);
-	const FSolverVec3 V1TimesDt = Particles.GetP(i1) - Particles.GetX(i1);
-	const FSolverVec3 V2TimesDt = Particles.GetP(i2) - Particles.GetX(i2);
+	const FSolverVec3 V0TimesDt = Particles.GetP(Idx0) - Particles.GetX(Idx0);
+	const FSolverVec3 V1TimesDt = Particles.GetP(Idx1) - Particles.GetX(Idx1);
+	const FSolverVec3 V2TimesDt = Particles.GetP(Idx2) - Particles.GetX(Idx2);
 
 	// These scale factors make everything resolution independent
 	const FSolverVec3 FinalStiffnesses = StiffnessScales[ConstraintIndex] * ExpStiffnessValue;
@@ -841,49 +841,54 @@ void FXPBDStretchBiasElementConstraints::ApplyHelper(SolverParticlesOrRange& Par
 
 	FSolverVec3& Lambda = Lambdas[ConstraintIndex];
 	FSolverVec3 DLambda;
-	DLambda[0] = CalcDLambda(FinalStiffnesses[0], Dt, Lambda[0], Damping[0], Particles.InvM(i0), Particles.InvM(i1), Particles.InvM(i2),
+	DLambda[0] = CalcDLambda(FinalStiffnesses[0], Dt, Lambda[0], Damping[0], Particles.InvM(Idx0), Particles.InvM(Idx1), Particles.InvM(Idx2),
 		V0TimesDt, V1TimesDt, V2TimesDt, Cu, GradCu0, GradCu1, GradCu2);
-	DLambda[1] = CalcDLambda(FinalStiffnesses[1], Dt, Lambda[1], Damping[1], Particles.InvM(i0), Particles.InvM(i1), Particles.InvM(i2),
+	DLambda[1] = CalcDLambda(FinalStiffnesses[1], Dt, Lambda[1], Damping[1], Particles.InvM(Idx0), Particles.InvM(Idx1), Particles.InvM(Idx2),
 		V0TimesDt, V1TimesDt, V2TimesDt, Cv, GradCv0, GradCv1, GradCv2);
-	DLambda[2] = CalcDLambda(FinalStiffnesses[2], Dt, Lambda[2], Damping[2], Particles.InvM(i0), Particles.InvM(i1), Particles.InvM(i2),
+	DLambda[2] = CalcDLambda(FinalStiffnesses[2], Dt, Lambda[2], Damping[2], Particles.InvM(Idx0), Particles.InvM(Idx1), Particles.InvM(Idx2),
 		V0TimesDt, V1TimesDt, V2TimesDt, Cs, GradCs0, GradCs1, GradCs2);
 
-	Particles.P(i0) -= Particles.InvM(i0) * (DLambda[0] * GradCu0 + DLambda[1] * GradCv0 + DLambda[2] * GradCs0);
-	Particles.P(i1) -= Particles.InvM(i1) * (DLambda[0] * GradCu1 + DLambda[1] * GradCv1 + DLambda[2] * GradCs1);
-	Particles.P(i2) -= Particles.InvM(i2) * (DLambda[0] * GradCu2 + DLambda[1] * GradCv2 + DLambda[2] * GradCs2);
+	Particles.P(Idx0) -= Particles.InvM(Idx0) * (DLambda[0] * GradCu0 + DLambda[1] * GradCv0 + DLambda[2] * GradCs0);
+	Particles.P(Idx1) -= Particles.InvM(Idx1) * (DLambda[0] * GradCu1 + DLambda[1] * GradCv1 + DLambda[2] * GradCs1);
+	Particles.P(Idx2) -= Particles.InvM(Idx2) * (DLambda[0] * GradCu2 + DLambda[1] * GradCv2 + DLambda[2] * GradCs2);
 
 	Lambda += DLambda;
 }
 
 
-void FXPBDStretchBiasElementConstraints::InitializeDmInvAndMeasures(const FSolverParticles& Particles)
+void FXPBDStretchBiasElementConstraints::InitializeDmInvAndMeasures(const FSolverParticlesRange& Particles)
 {
 	Measure.SetNum(Constraints.Num());
 	DmInverse.Init(PMatrix<FSolverReal, 2, 2>(0.f, 0.f, 0.f), Constraints.Num());
 	DmArray.Init(PMatrix<FSolverReal, 2, 2>(0.f, 0.f, 0.f),Constraints.Num());
 	RestDmArray.Init(PMatrix<FSolverReal, 3, 2>(0.f),Constraints.Num());
-	for (int32 e = 0; e < Constraints.Num(); e++)
+	for (int32 ElemIdx = 0; ElemIdx < Constraints.Num(); ElemIdx++)
 	{
-		const TVec3<FSolverReal> X1X0 = Particles.GetX(Constraints[e][1]) - Particles.GetX(Constraints[e][0]);
-		const TVec3<FSolverReal> X2X0 = Particles.GetX(Constraints[e][2]) - Particles.GetX(Constraints[e][0]);
+		const TVec3<FSolverReal> X1X0 = Particles.GetX(Constraints[ElemIdx][1]) - Particles.GetX(Constraints[ElemIdx][0]);
+		const TVec3<FSolverReal> X2X0 = Particles.GetX(Constraints[ElemIdx][2]) - Particles.GetX(Constraints[ElemIdx][0]);
 		PMatrix<FSolverReal, 3, 2> RestDm(0.f);
 		PMatrix<FSolverReal, 2, 2> Dm(0.f, 0.f, 0.f);
 		Dm.M[0] = X1X0.Size();
 		Dm.M[2] = X1X0.Dot(X2X0) / Dm.M[0];
 		Dm.M[3] = Chaos::TVector<FSolverReal, 3>::CrossProduct(X1X0, X2X0).Size() / Dm.M[0];
-		Measure[e] = Chaos::TVector<FSolverReal, 3>::CrossProduct(X1X0, X2X0).Size() / 2.f;
-		ensureMsgf(Measure[e] > 0.f, TEXT("Degenerate triangle detected"));
+		Measure[ElemIdx] = Chaos::TVector<FSolverReal, 3>::CrossProduct(X1X0, X2X0).Size() / 2.f;
+		ensureMsgf(Measure[ElemIdx] > 0.f, TEXT("Degenerate triangle detected"));
 
-		PMatrix<FSolverReal, 2, 2> DmInv = Dm.Inverse();
-		DmInverse[e] = DmInv;
-		DmArray[e] = Dm;
-		RestDmArray[e] = Chaos::PMatrix<FSolverReal, 3, 2>(X1X0, X2X0);
+		const PMatrix<FSolverReal, 2, 2> DmInv = Dm.Inverse();
+		DmInverse[ElemIdx] = DmInv;
+		DmArray[ElemIdx] = Dm;
+		RestDmArray[ElemIdx] = Chaos::PMatrix<FSolverReal, 3, 2>(X1X0, X2X0);
 	}
 
 	bDmInitialized = true;
 }
 
-void FXPBDStretchBiasElementConstraints::AddStretchBiasElementResidualAndHessian(const FSolverParticles& Particles, const int32 ConstraintIndex, const int32 ConstraintIndexLocal, const FSolverReal Dt, TVec3<FSolverReal>& ParticleResidual, Chaos::PMatrix<FSolverReal, 3, 3>& ParticleHessian)
+void FXPBDStretchBiasElementConstraints::InitializeDmInvAndMeasures(const FSolverParticles& Particles)
+{
+	InitializeDmInvAndMeasures(FSolverParticlesRange(Particles));
+}
+
+void FXPBDStretchBiasElementConstraints::AddStretchBiasElementResidualAndHessian(const FSolverParticlesRange& Particles, const int32 ConstraintIndex, const int32 ConstraintIndexLocal, const FSolverReal Dt, TVec3<FSolverReal>& ParticleResidual, Chaos::PMatrix<FSolverReal, 3, 3>& ParticleHessian)
 {
 	const bool StiffnessWarpHasWeightMap = StiffnessWarp.HasWeightMap();
 	const bool StiffnessWeftHasWeightMap = StiffnessWeft.HasWeightMap();
@@ -903,16 +908,16 @@ void FXPBDStretchBiasElementConstraints::AddStretchBiasElementResidualAndHessian
 	const FSolverMatrix22& DeltaUVInv = DeltaUVInverse[ConstraintIndex];
 
 	const TVec3<int32>& Constraint = Constraints[ConstraintIndex];
-	const int32 i0 = Constraint[0];
-	const int32 i1 = Constraint[1];
-	const int32 i2 = Constraint[2];
-	if (Particles.InvM(i0) == (FSolverReal)0. && Particles.InvM(i1) == (FSolverReal)0. && Particles.InvM(i2) == (FSolverReal)0.)
+	const int32 Idx0 = Constraint[0];
+	const int32 Idx1 = Constraint[1];
+	const int32 Idx2 = Constraint[2];
+	if (Particles.InvM(Idx0) == (FSolverReal)0. && Particles.InvM(Idx1) == (FSolverReal)0. && Particles.InvM(Idx2) == (FSolverReal)0.)
 	{
 		return;
 	}
 
 	FSolverVec3 DXDu, DXDv;
-	CalculateUVStretch(ConstraintIndex, Particles.P(i0), Particles.P(i1), Particles.P(i2), DXDu, DXDv);
+	CalculateUVStretch(ConstraintIndex, Particles.P(Idx0), Particles.P(Idx1), Particles.P(Idx2), DXDu, DXDv);
 
 	const FSolverReal DXDuLength = DXDu.Length();
 	const FSolverReal DXDvLength = DXDv.Length();
@@ -1022,10 +1027,10 @@ void FXPBDStretchBiasElementConstraints::AddStretchBiasElementResidualAndHessian
 		{
 			ResContribution[Alpha] += Coeffs * GradCs[Alpha];
 			//original uv constraint residuals:
-			for (int32 s = 0; s < 2; s++)
+			for (int32 Col = 0; Col < 2; Col++)
 			{
-				ResContribution[Alpha] -= Coeffu * DXDu[Alpha] * DeltaUVInv.M[s];
-				ResContribution[Alpha] -= Coeffv * DXDv[Alpha] * DeltaUVInv.M[s + 2];
+				ResContribution[Alpha] -= Coeffu * DXDu[Alpha] * DeltaUVInv.M[Col];
+				ResContribution[Alpha] -= Coeffv * DXDv[Alpha] * DeltaUVInv.M[Col + 2];
 			}
 		}
 	}
@@ -1056,14 +1061,14 @@ void FXPBDStretchBiasElementConstraints::AddStretchBiasElementResidualAndHessian
 			HessianContribution.SetRow(Alpha, HessianContribution.GetRow(Alpha) + HessianCoeffs * GradCs[Alpha] * GradCs);
 		}
 		//code for stiffness padding:
-		FSolverReal DmInvsum = 0.f;
-		for (int32 nu = 0; nu < 2; nu++)
+		FSolverReal DmInvSum = 0.f;
+		for (int32 Col = 0; Col < 2; Col++)
 		{
-			DmInvsum += DmInvMat.GetAt(ConstraintIndexLocal - 1, nu) * DmInvMat.GetAt(ConstraintIndexLocal - 1, nu);
+			DmInvSum += DmInvMat.GetAt(ConstraintIndexLocal - 1, Col) * DmInvMat.GetAt(ConstraintIndexLocal - 1, Col);
 		}
-		for (int32 alpha = 0; alpha < 3; alpha++)
+		for (int32 Alpha = 0; Alpha < 3; Alpha++)
 		{
-			HessianContribution.SetAt(alpha, alpha, HessianContribution.GetAt(alpha, alpha) + PaddingCoeff * StiffnessPadding * DmInvsum);
+			HessianContribution.SetAt(Alpha, Alpha, HessianContribution.GetAt(Alpha, Alpha) + PaddingCoeff * StiffnessPadding * DmInvSum);
 		}
 	}
 	else
@@ -1072,34 +1077,34 @@ void FXPBDStretchBiasElementConstraints::AddStretchBiasElementResidualAndHessian
 		{
 			for (int32 Beta = 0; Beta < 3; Beta++)
 			{
-				for (int32 s = 0; s < 2; s++)
+				for (int32 Col = 0; Col < 2; Col++)
 				{
-					for (int32 n = 0; n < 2; n++)
+					for (int32 Row = 0; Row < 2; Row++)
 					{
-						//constraint u (definitely could be simplified by adding u[s] first):
-						HessianContribution.SetAt(Alpha, Beta, HessianContribution.GetAt(Alpha, Beta) + RankOneuCoeff * DeltaUVInv.M[s] * DeltaUVInv.M[n] * DXDu[Alpha] * DXDu[Beta]);
-						
-						//constraint v (definitely could be simplified by adding v[s] first):
-						HessianContribution.SetAt(Alpha, Beta, HessianContribution.GetAt(Alpha, Beta) + RankOnevCoeff * DeltaUVInv.M[s+2] * DeltaUVInv.M[n+2] * DXDv[Alpha] * DXDv[Beta]);
+						//constraint u (definitely could be simplified by adding u[Col] first):
+						HessianContribution.SetAt(Alpha, Beta, HessianContribution.GetAt(Alpha, Beta) + RankOneuCoeff * DeltaUVInv.M[Col] * DeltaUVInv.M[Row] * DXDu[Alpha] * DXDu[Beta]);
+
+						//constraint v (definitely could be simplified by adding v[Col] first):
+						HessianContribution.SetAt(Alpha, Beta, HessianContribution.GetAt(Alpha, Beta) + RankOnevCoeff * DeltaUVInv.M[Col+2] * DeltaUVInv.M[Row+2] * DXDv[Alpha] * DXDv[Beta]);
 					}
 				}
 			}
 			HessianContribution.SetRow(Alpha, HessianContribution.GetRow(Alpha) + HessianCoeffs * GradCs[Alpha] * GradCs);
 		}
 		//code for stiffness padding:
-		FSolverReal DmInvsum = 0.f;
-		for (int32 nu = 0; nu < 2; nu++) 
+		FSolverReal DmInvSum = 0.f;
+		for (int32 Row = 0; Row < 2; Row++)
 		{
-			FSolverReal localDmsum = 0.f;
-			for (int32 k = 0; k < 2; k++) 
+			FSolverReal LocalDmSum = 0.f;
+			for (int32 Col = 0; Col < 2; Col++)
 			{
-				localDmsum += DmInvMat.M[nu * 2 + k];
+				LocalDmSum += DmInvMat.M[Row * 2 + Col];
 			}
-			DmInvsum += localDmsum * localDmsum;
+			DmInvSum += LocalDmSum * LocalDmSum;
 		}
-		for (int32 alpha = 0; alpha < 3; alpha++) 
+		for (int32 Alpha = 0; Alpha < 3; Alpha++)
 		{
-			HessianContribution.SetAt(alpha, alpha, HessianContribution.GetAt(alpha, alpha) + PaddingCoeff * StiffnessPadding * DmInvsum);
+			HessianContribution.SetAt(Alpha, Alpha, HessianContribution.GetAt(Alpha, Alpha) + PaddingCoeff * StiffnessPadding * DmInvSum);
 		}
 
 	}
@@ -1108,8 +1113,13 @@ void FXPBDStretchBiasElementConstraints::AddStretchBiasElementResidualAndHessian
 	ParticleHessian += HessianContribution;
 }
 
+void FXPBDStretchBiasElementConstraints::AddStretchBiasElementResidualAndHessian(const FSolverParticles& Particles, const int32 ConstraintIndex, const int32 ConstraintIndexLocal, const FSolverReal Dt, TVec3<FSolverReal>& ParticleResidual, Chaos::PMatrix<FSolverReal, 3, 3>& ParticleHessian)
+{
+	AddStretchBiasElementResidualAndHessian(FSolverParticlesRange(Particles), ConstraintIndex, ConstraintIndexLocal, Dt, ParticleResidual, ParticleHessian);
+}
+
 
-void FXPBDStretchBiasElementConstraints::AddInternalForceDifferential(const FSolverParticles& InParticles, const TArray<TVector<FSolverReal, 3>>& DeltaParticles, TArray<TVector<FSolverReal, 3>>& ndf)
+void FXPBDStretchBiasElementConstraints::AddInternalForceDifferential(const FSolverParticlesRange& InParticles, const TArray<TVector<FSolverReal, 3>>& DeltaParticles, TArray<TVector<FSolverReal, 3>>& Ndf)
 {
 	TRACE_CPUPROFILER_EVENT_SCOPE(STAT_ChaosCGAnisoBiasElementAddInternalForceDifferential);
 	const FSolverReal StiffnessWarpNoMap = (FSolverReal)StiffnessWarp;
@@ -1120,10 +1130,9 @@ void FXPBDStretchBiasElementConstraints::AddInternalForceDifferential(const FSol
 	const FSolverReal WeftScaleNoMap = (FSolverReal)WeftScale;
 	const FSolverVec3 ExpStiffnessValue(StiffnessWeftNoMap, StiffnessWarpNoMap, StiffnessBiasNoMap);
 
-	int32 ParticleStart = 0; 
-	int32 ParticleNum = InParticles.Size();
+	const int32 ParticleStart = 0;
 
-	check(ndf.Num() == InParticles.Size());
+	check(Ndf.Num() == InParticles.Size());
 
 	const int32 ConstraintColorNum = ConstraintsPerColorStartIndex.Num() - 1;
 
@@ -1135,17 +1144,17 @@ void FXPBDStretchBiasElementConstraints::AddInternalForceDifferential(const FSol
 			{
 				const int32 ConstraintIndex = ColorStart + Index;
 				const TVec3<int32>& Constraint = Constraints[ConstraintIndex];
-				const int32 i0 = Constraint[0];
-				const int32 i1 = Constraint[1];
-				const int32 i2 = Constraint[2];
+				const int32 Idx0 = Constraint[0];
+				const int32 Idx1 = Constraint[1];
+				const int32 Idx2 = Constraint[2];
 
-				if (InParticles.InvM(i0) == (FSolverReal)0. && InParticles.InvM(i1) == (FSolverReal)0. && InParticles.InvM(i2) == (FSolverReal)0.)
+				if (InParticles.InvM(Idx0) == (FSolverReal)0. && InParticles.InvM(Idx1) == (FSolverReal)0. && InParticles.InvM(Idx2) == (FSolverReal)0.)
 				{
 					return;
 				}
 
 				FSolverVec3 DXDu, DXDv;
-				CalculateUVStretch(ConstraintIndex, InParticles.P(i0), InParticles.P(i1), InParticles.P(i2), DXDu, DXDv);
+				CalculateUVStretch(ConstraintIndex, InParticles.P(Idx0), InParticles.P(Idx1), InParticles.P(Idx2), DXDu, DXDv);
 
 				const FSolverReal DXDuLength = DXDu.Length();
 				const FSolverReal DXDvLength = DXDv.Length();
@@ -1183,22 +1192,22 @@ void FXPBDStretchBiasElementConstraints::AddInternalForceDifferential(const FSol
 
 				FSolverReal DeltaCu = 0.f, DeltaCv = 0.f, DeltaCs = 0.f;
 
-				const FSolverVec3 DeltaP01 = DeltaParticles[i1 - ParticleStart] - DeltaParticles[i0 - ParticleStart];
-				const FSolverVec3 DeltaP02 = DeltaParticles[i2 - ParticleStart] - DeltaParticles[i0 - ParticleStart];
-				for (int32 i = 0; i < 3; i++)
+				const FSolverVec3 DeltaP01 = DeltaParticles[Idx1 - ParticleStart] - DeltaParticles[Idx0 - ParticleStart];
+				const FSolverVec3 DeltaP02 = DeltaParticles[Idx2 - ParticleStart] - DeltaParticles[Idx0 - ParticleStart];
+				for (int32 Alpha = 0; Alpha < 3; Alpha++)
 				{
-					DeltaCu += DXDuNormalized[i] * DeltaP01[i] * DeltaUVInv.M[0];
-					DeltaCu += DXDuNormalized[i] * DeltaP02[i] * DeltaUVInv.M[1];
-					DeltaCv += DXDvNormalized[i] * DeltaP01[i] * DeltaUVInv.M[2];
-					DeltaCv += DXDvNormalized[i] * DeltaP02[i] * DeltaUVInv.M[3];
-					DeltaCs += GradCs1[i] * DeltaP01[i];
-					DeltaCs += GradCs2[i] * DeltaP02[i];
+					DeltaCu += DXDuNormalized[Alpha] * DeltaP01[Alpha] * DeltaUVInv.M[0];
+					DeltaCu += DXDuNormalized[Alpha] * DeltaP02[Alpha] * DeltaUVInv.M[1];
+					DeltaCv += DXDvNormalized[Alpha] * DeltaP01[Alpha] * DeltaUVInv.M[2];
+					DeltaCv += DXDvNormalized[Alpha] * DeltaP02[Alpha] * DeltaUVInv.M[3];
+					DeltaCs += GradCs1[Alpha] * DeltaP01[Alpha];
+					DeltaCs += GradCs2[Alpha] * DeltaP02[Alpha];
 				}
 
 				const FSolverVec3 FinalStiffnesses = StiffnessScales[ConstraintIndex] * ExpStiffnessValue;
 
 				FSolverVec3 DeltaDXDu, DeltaDXDv;
-				CalculateUVStretch(ConstraintIndex, DeltaParticles[i0 - ParticleStart], DeltaParticles[i1 - ParticleStart], DeltaParticles[i2 - ParticleStart], DeltaDXDu, DeltaDXDv);
+				CalculateUVStretch(ConstraintIndex, DeltaParticles[Idx0 - ParticleStart], DeltaParticles[Idx1 - ParticleStart], DeltaParticles[Idx2 - ParticleStart], DeltaDXDu, DeltaDXDv);
 				const FSolverVec3 DeltaDXDuOverDxDuNorm = DeltaDXDu * OneOverDXDuLen;
 				const FSolverVec3 DeltaDXDvOverDxDvNorm = DeltaDXDv * OneOverDXDvLen;
 				FSolverVec3 DeltaGradCu0 = DeltaDXDuOverDxDuNorm * DeltaUVInv.M[0];
@@ -1209,18 +1218,23 @@ void FXPBDStretchBiasElementConstraints::AddInternalForceDifferential(const FSol
 				FSolverVec3 DeltaGradCv2 = -DeltaGradCv0 - DeltaGradCv1;
 
 
-				ndf[i0] += FinalStiffnesses[0] * DeltaCu * GradCu0 + FinalStiffnesses[0] * Cu * DeltaGradCu0;
-				ndf[i0] += FinalStiffnesses[1] * DeltaCv * GradCv0 + FinalStiffnesses[0] * Cu * DeltaGradCu1;
-				ndf[i0] += FinalStiffnesses[2] * DeltaCs * GradCs0 + FinalStiffnesses[0] * Cu * DeltaGradCu2;
-				ndf[i1] += FinalStiffnesses[0] * DeltaCu * GradCu1 + FinalStiffnesses[1] * Cv * DeltaGradCv0;
-				ndf[i1] += FinalStiffnesses[1] * DeltaCv * GradCv1 + FinalStiffnesses[1] * Cv * DeltaGradCv1;
-				ndf[i1] += FinalStiffnesses[2] * DeltaCs * GradCs1 + FinalStiffnesses[1] * Cv * DeltaGradCv2;
-				ndf[i2] += FinalStiffnesses[0] * DeltaCu * GradCu2;
-				ndf[i2] += FinalStiffnesses[1] * DeltaCv * GradCv2;
-				ndf[i2] += FinalStiffnesses[2] * DeltaCs * GradCs2;
+				Ndf[Idx0] += FinalStiffnesses[0] * DeltaCu * GradCu0 + FinalStiffnesses[0] * Cu * DeltaGradCu0;
+				Ndf[Idx0] += FinalStiffnesses[1] * DeltaCv * GradCv0 + FinalStiffnesses[0] * Cu * DeltaGradCu1;
+				Ndf[Idx0] += FinalStiffnesses[2] * DeltaCs * GradCs0 + FinalStiffnesses[0] * Cu * DeltaGradCu2;
+				Ndf[Idx1] += FinalStiffnesses[0] * DeltaCu * GradCu1 + FinalStiffnesses[1] * Cv * DeltaGradCv0;
+				Ndf[Idx1] += FinalStiffnesses[1] * DeltaCv * GradCv1 + FinalStiffnesses[1] * Cv * DeltaGradCv1;
+				Ndf[Idx1] += FinalStiffnesses[2] * DeltaCs * GradCs1 + FinalStiffnesses[1] * Cv * DeltaGradCv2;
+				Ndf[Idx2] += FinalStiffnesses[0] * DeltaCu * GradCu2;
+				Ndf[Idx2] += FinalStiffnesses[1] * DeltaCv * GradCv2;
+				Ndf[Idx2] += FinalStiffnesses[2] * DeltaCs * GradCs2;
 			});
 	}
 }
 
+void FXPBDStretchBiasElementConstraints::AddInternalForceDifferential(const FSolverParticles& InParticles, const TArray<TVector<FSolverReal, 3>>& DeltaParticles, TArray<TVector<FSolverReal, 3>>& Ndf)
+{
+	AddInternalForceDifferential(FSolverParticlesRange(InParticles), DeltaParticles, Ndf);
+}
+
 
 }  // End namespace Chaos::Softs
diff --git a/Engine/Source/Runtime/Experimental/Chaos/Public/Chaos/Deformable/GaussSeidelCorotatedCodimensionalConstraints.h b/Engine/Source/Runtime/Experimental/Chaos/Public/Chaos/Deformable/GaussSeidelCorotatedCodimensionalConstraints.h
--- a/Engine/Source/Runtime/Experimental/Chaos/Public/Chaos/Deformable/GaussSeidelCorotatedCodimensionalConstraints.h
+++ b/Engine/Source/Runtime/Experimental/Chaos/Public/Chaos/Deformable/GaussSeidelCorotatedCodimensionalConstraints.h
@@ -8,6 +8,7 @@
 #include "Chaos/Framework/Parallel.h"
 #include "Chaos/Deformable/ChaosDeformableSolverTypes.h"
 #include "Chaos/Vector.h"
+#include "Chaos/SoftsSolverParticlesRange.h"
 
 
 namespace Chaos::Softs
@@ -17,158 +18,55 @@ namespace Chaos::Softs
 	class FGaussSeidelCorotatedTriangleConstraints
 	{
 	public:
+		CHAOS_API FGaussSeidelCorotatedTriangleConstraints(
+			const FSolverParticlesRange& InParticles,
+			const TArray<TVector<int32, 3>>& InMesh,
+			const bool bRecordMetricIn = true,
+			const FSolverReal& EMesh = (FSolverReal)10.0,
+			const FSolverReal& NuMesh = (FSolverReal).3);
+
 		FGaussSeidelCorotatedTriangleConstraints(
 			const FSolverParticles& InParticles,
 			const TArray<TVector<int32, 3>>& InMesh,
 			const bool bRecordMetricIn = true,
 			const FSolverReal& EMesh = (FSolverReal)10.0,
-			const FSolverReal& NuMesh = (FSolverReal).3
-		)
-			: bRecordMetric(bRecordMetricIn), MeshConstraints(InMesh)
-		{
-			Measure.Init((FSolverReal)0., MeshConstraints.Num());
-			Lambda = EMesh * NuMesh / (((FSolverReal)1. + NuMesh) * ((FSolverReal)1. - (FSolverReal)2. * NuMesh));
-			Mu = EMesh / ((FSolverReal)2. * ((FSolverReal)1. + NuMesh));
-
-			MuElementArray.Init(Mu, MeshConstraints.Num());
-			LambdaElementArray.Init(Lambda, MeshConstraints.Num());
-
-			InitializeCodimensionData(InParticles);
-		}
+			const FSolverReal& NuMesh = (FSolverReal).3)
+			: FGaussSeidelCorotatedTriangleConstraints(
+				FSolverParticlesRange(InParticles),
+				InMesh, bRecordMetricIn, EMesh, NuMesh)
+		{}
+
+		CHAOS_API FGaussSeidelCorotatedTriangleConstraints(
+			const FSolverParticlesRange& InParticles,
+			const TArray<TVector<int32, 3>>& InMesh,
+			const TArray<FSolverReal>& EMeshArray,
+			const FSolverReal& NuMesh = (FSolverReal).3);
 
 		FGaussSeidelCorotatedTriangleConstraints(
 			const FSolverParticles& InParticles,
 			const TArray<TVector<int32, 3>>& InMesh,
 			const TArray<FSolverReal>& EMeshArray,
-			const FSolverReal& NuMesh = (FSolverReal).3
-		)
-			: MeshConstraints(InMesh)
-		{
-			ensureMsgf(EMeshArray.Num() == InMesh.Num(), TEXT("Input Young Modulus Array Size is wrong"));
-			Measure.Init((FSolverReal)0., MeshConstraints.Num());
-
-			InitializeCodimensionData(InParticles);
-			LambdaElementArray.Init((FSolverReal)0., MeshConstraints.Num());
-			MuElementArray.Init((FSolverReal)0., MeshConstraints.Num());
-			
-			for (int32 e = 0; e < InMesh.Num(); e++)
-			{
-				LambdaElementArray[e] = EMeshArray[e] * NuMesh / (((FSolverReal)1. + NuMesh) * ((FSolverReal)1. - (FSolverReal)2. * NuMesh));
-				MuElementArray[e] = EMeshArray[e] / ((FSolverReal)2. * ((FSolverReal)1. + NuMesh));
-			}
-		}
+			const FSolverReal& NuMesh = (FSolverReal).3)
+			: FGaussSeidelCorotatedTriangleConstraints(
+				FSolverParticlesRange(InParticles),
+				InMesh, EMeshArray, NuMesh)
+		{}
 
 		virtual ~FGaussSeidelCorotatedTriangleConstraints() = default;
 
-		PMatrix<FSolverReal, 3, 3> DsInit(const int32 E, const FSolverParticles& InParticles)  
-		{
-			PMatrix<FSolverReal, 3, 3> Result((FSolverReal)0.);
-			for (int32 i = 0; i < 3; i++) 
-			{
-				for (int32 c = 0; c < 3; c++) 
-				{
-					Result.SetAt(c, i, InParticles.X(MeshConstraints[E][i + 1])[c] - InParticles.X(MeshConstraints[E][0])[c]);
-				}
-			}
-			return Result;
-		}
-
-		PMatrix<FSolverReal, 3, 2> Ds(const int32 E, const FSolverParticles& InParticles) const 
-		{
-			if (INDEX_NONE < E && E< MeshConstraints.Num()
-				&& INDEX_NONE < MeshConstraints[E][0] && MeshConstraints[E][0] < (int32)InParticles.Size()
-				&& INDEX_NONE < MeshConstraints[E][1] && MeshConstraints[E][1] < (int32)InParticles.Size()
-				&& INDEX_NONE < MeshConstraints[E][2] && MeshConstraints[E][2] < (int32)InParticles.Size())
-			{
-				const TVec3<FSolverReal> P1P0 = InParticles.P(MeshConstraints[E][1]) - InParticles.P(MeshConstraints[E][0]);
-				const TVec3<FSolverReal> P2P0 = InParticles.P(MeshConstraints[E][2]) - InParticles.P(MeshConstraints[E][0]);
-
-				return PMatrix<FSolverReal, 3, 2>(
-					P1P0[0], P1P0[1], P1P0[2],
-					P2P0[0], P2P0[1], P2P0[2]);
-			}
-			else
-			{
-				return PMatrix<FSolverReal, 3, 2>(
-					(FSolverReal)0., (FSolverReal)0, (FSolverReal)0,
-					(FSolverReal)0, (FSolverReal)0, (FSolverReal)0);
-			}
-		
-		}
+		CHAOS_API TArray<TArray<int32>> GetConstraintsArray() const;
 
-		PMatrix<FSolverReal, 3, 2> F(const int32 E, const FSolverParticles& InParticles) const 
-		{
-			if (INDEX_NONE < E && E < DmInverse.Num())
-			{
-				return Ds(E, InParticles) * DmInverse[E];
-			}
-			else
-			{
-				return PMatrix<FSolverReal, 3, 2>(
-					(FSolverReal)0., (FSolverReal)0, (FSolverReal)0,
-					(FSolverReal)0, (FSolverReal)0, (FSolverReal)0);
-			}
-		}
+		CHAOS_API void AddHyperelasticResidualAndHessian(const FSolverParticlesRange& Particles, const int32 ElementIndex, const int32 ElementIndexLocal, const FSolverReal Dt, TVec3<FSolverReal>& ParticleResidual, Chaos::PMatrix<FSolverReal, 3, 3>& ParticleHessian);
 
-		TArray<TArray<int32>> GetConstraintsArray() const
-		{
-			TArray<TArray<int32>> Constraints;
-			Constraints.SetNum(MeshConstraints.Num());
-			for (int32 i = 0; i < MeshConstraints.Num(); i++)
-			{
-				Constraints[i].SetNum(3);
-				for (int32 j = 0; j < 3; j++)
-				{
-					Constraints[i][j] = MeshConstraints[i][j];
-				}
-			}
-			return Constraints;
-		}
+	protected:
 
-		void ComputeNodalMass(const FSolverReal InDensity, const int32 NumParticles, TArray<FSolverReal>& NodalMass) const
-		{
-			NodalMass.Init((FSolverReal)0., NumParticles);
-			for (int32 e = 0; e < Measure.Num(); e++)
-			{
-				const FSolverReal TotalMass = InDensity * Measure[e];
-				for (int32 j = 0; j < 3; j++)
-				{
-					NodalMass[MeshConstraints[e][j]] += TotalMass / (FSolverReal)3.;
-				}
-			}
-		}
+		PMatrix<FSolverReal, 3, 2> Ds(const int32 ElemIdx, const FSolverParticlesRange& InParticles) const;
 
+		PMatrix<FSolverReal, 3, 2> F(const int32 ElemIdx, const FSolverParticlesRange& InParticles) const;
 
+		CHAOS_API void ComputeNodalMass(const FSolverReal InDensity, const int32 NumParticles, TArray<FSolverReal>& NodalMass) const;
 
-	protected:
-
-		void InitializeCodimensionData(const FSolverParticles& Particles)
-		{
-			Measure.Init((FSolverReal)0., MeshConstraints.Num());
-			DmInverse.Init(PMatrix<FSolverReal, 2, 2>(0.f, 0.f, 0.f), MeshConstraints.Num());
-			for (int32 e = 0; e < MeshConstraints.Num(); e++)
-			{
-				check(MeshConstraints[e][0] < (int32)Particles.Size() && MeshConstraints[e][0] > INDEX_NONE);
-				check(MeshConstraints[e][1] < (int32)Particles.Size() && MeshConstraints[e][1] > INDEX_NONE);
-				check(MeshConstraints[e][2] < (int32)Particles.Size() && MeshConstraints[e][2] > INDEX_NONE);
-				if (MeshConstraints[e][0] < (int32)Particles.Size() && MeshConstraints[e][0] > INDEX_NONE
-					&& MeshConstraints[e][1] < (int32)Particles.Size() && MeshConstraints[e][1] > INDEX_NONE
-					&& MeshConstraints[e][2] < (int32)Particles.Size() && MeshConstraints[e][2] > INDEX_NONE)
-				{
-					const TVec3<FSolverReal> X1X0 = Particles.GetX(MeshConstraints[e][1]) - Particles.GetX(MeshConstraints[e][0]);
-					const TVec3<FSolverReal> X2X0 = Particles.GetX(MeshConstraints[e][2]) - Particles.GetX(MeshConstraints[e][0]);
-					PMatrix<FSolverReal, 2, 2> Dm((FSolverReal)0., (FSolverReal)0., (FSolverReal)0.);
-					Dm.M[0] = X1X0.Size();
-					Dm.M[2] = X1X0.Dot(X2X0) / Dm.M[0];
-					Dm.M[3] = Chaos::TVector<FSolverReal, 3>::CrossProduct(X1X0, X2X0).Size() / Dm.M[0];
-					Measure[e] = Chaos::TVector<FSolverReal, 3>::CrossProduct(X1X0, X2X0).Size() / (FSolverReal)2.;
-					ensureMsgf(Measure[e] > (FSolverReal)0., TEXT("Degenerate triangle detected"));
-
-					PMatrix<FSolverReal, 2, 2> DmInv = Dm.Inverse();
-					DmInverse[e] = DmInv;
-				}
-			}
-		}
+		void InitializeCodimensionData(const FSolverParticlesRange& Particles);
 
 		static FSolverReal SafeRecip(const FSolverReal Len, const FSolverReal Fallback)
 		{
@@ -179,78 +77,56 @@ namespace Chaos::Softs
 			return Fallback;
 		}
 
-		static inline void SymSchur2D(const FSolverReal Aqq, const FSolverReal App, const FSolverReal Apq, FSolverReal& c, FSolverReal& s)
+		static inline void SymSchur2D(const FSolverReal Aqq, const FSolverReal App, const FSolverReal Apq, FSolverReal& C, FSolverReal& S)
 		{
 			// A is an n x n matrix
 			// (p, q) is an index pair satisfying 1 <= p < q <= n
 			//
-			// This function computes a (c, s) pair such that
+			// This function computes a (C, S) pair such that
 			//
-			// B = [ c, s]^T [a_pp, a_pq] [ c, s]
-			//     [-s, c]   [a_pq, a_qq] [-s, c]
+			// B = [ C, S]^T [a_pp, a_pq] [ C, S]
+			//     [-S, C]   [a_pq, a_qq] [-S, C]
 			//
 			// is diagonal
 
-			if (Apq != 0) 
+			if (Apq != 0)
 			{
 				FSolverReal Tau = FSolverReal(0.5) * (Aqq - App) / Apq;
-				FSolverReal t;
+				FSolverReal T;
 				if (Tau >= 0)
-					t = FSolverReal(1) / (Tau + FMath::Sqrt((FSolverReal)1. + Tau * Tau));
+				{
+					T = FSolverReal(1) / (Tau + FMath::Sqrt((FSolverReal)1. + Tau * Tau));
+				}
 				else
-					t = FSolverReal(-1) / (-Tau + FMath::Sqrt((FSolverReal)1. + Tau * Tau));
-				c = FSolverReal(1) / FMath::Sqrt((FSolverReal)1. + t * t);
-				s = t * c;
+				{
+					T = FSolverReal(-1) / (-Tau + FMath::Sqrt((FSolverReal)1. + Tau * Tau));
+				}
+				C = FSolverReal(1) / FMath::Sqrt((FSolverReal)1. + T * T);
+				S = T * C;
 			}
-			else 
+			else
 			{
-				c = (FSolverReal)1.;
-				s = (FSolverReal)0.;
+				C = (FSolverReal)1.;
+				S = (FSolverReal)0.;
 			}
 		}
 
-
 		static inline void Jacobi(const PMatrix<FSolverReal, 2, 2>& B, PMatrix<FSolverReal, 2, 2>& D, PMatrix<FSolverReal, 2, 2>& V)
 		{
-			FSolverReal c, s;
-			SymSchur2D(B.M[3], B.M[0], B.M[1], c, s);
+			FSolverReal C, S;
+			SymSchur2D(B.M[3], B.M[0], B.M[1], C, S);
 
-			V.M[0] = c;
-			V.M[1] = -s;
-			V.M[2] = s;
-			V.M[3] = c;
+			V.M[0] = C;
+			V.M[1] = -S;
+			V.M[2] = S;
+			V.M[3] = C;
 
 			D = V.GetTransposed() * B * V;
 		}
 
-		static PMatrix<FSolverReal, 3, 2> ComputeR(const PMatrix<FSolverReal, 3, 2>& Fe) 
-		{
-			PMatrix<FSolverReal, 2, 2> FTF = ComputeFTF(Fe), D((FSolverReal)0., (FSolverReal)0., (FSolverReal)0.), U((FSolverReal)0., (FSolverReal)0., (FSolverReal)0.);
-			Jacobi(FTF, D, U);
-			PMatrix<FSolverReal, 2, 2> SDiag(FMath::Sqrt(D.M[0]), (FSolverReal)0., (FSolverReal)0., FMath::Sqrt(D.M[3]));
-			PMatrix<FSolverReal, 2, 2> S = U * SDiag * U.GetTransposed(), SInv = S.Inverse();;
-			return Fe * S.Inverse();
-		}
+		CHAOS_API static PMatrix<FSolverReal, 3, 2> ComputeR(const PMatrix<FSolverReal, 3, 2>& Fe);
 
-		static PMatrix<FSolverReal, 3, 2> ComputeRSimple(const PMatrix<FSolverReal, 3, 2>& InputMatrix) 
-		{
-			TVec3<FSolverReal> Col1(InputMatrix.M[0], InputMatrix.M[1], InputMatrix.M[2]), Col2(InputMatrix.M[3], InputMatrix.M[4], InputMatrix.M[5]), a(FSolverReal(0)), b(FSolverReal(0));
-			if (FMath::Abs(Col1[0]) < (FSolverReal)UE_SMALL_NUMBER && FMath::Abs(Col1[1]) < (FSolverReal)UE_SMALL_NUMBER)
-			{
-				a[0] = FSolverReal(1);
-			}
-			else
-			{
-				a[0] = -Col1[1];
-				a[1] = Col1[0];
-				const FSolverReal OneOveraNorm = SafeRecip(a.Length(), 0.f);
-				a *= OneOveraNorm;
-			}
-			b = TVec3<FSolverReal>::CrossProduct(a, Col2);
-			const FSolverReal OneOverbNorm = SafeRecip(b.Length(), 0.f);
-			b *= OneOverbNorm;
-			return Chaos::PMatrix<FSolverReal, 3, 2>(b[0], b[1], b[2], a[0], a[1], a[2]);
-		}
+		CHAOS_API static PMatrix<FSolverReal, 3, 2> ComputeRSimple(const PMatrix<FSolverReal, 3, 2>& InputMatrix);
 
 		//Does Gram Schmidt on a 3*2 matrix and returns an orthogonal one
 		static Chaos::PMatrix<FSolverReal, 3, 2> GramSchmidt(const Chaos::PMatrix<FSolverReal, 3, 2>& InputMatrix)
@@ -291,127 +167,6 @@ namespace Chaos::Softs
 			}
 		}
 
-	public:
-		void AddHyperelasticResidualAndHessian(const FSolverParticles& Particles, const int32 ElementIndex, const int32 ElementIndexLocal, const FSolverReal Dt, TVec3<FSolverReal>& ParticleResidual, Chaos::PMatrix<FSolverReal, 3, 3>& ParticleHessian)
-		{
-			check(ElementIndex < DmInverse.Num() && ElementIndex > INDEX_NONE);
-			check(ElementIndex < MuElementArray.Num());
-			check(ElementIndex < Measure.Num());
-			check(ElementIndex < LambdaElementArray.Num());
-			check(ElementIndexLocal < 3 && ElementIndexLocal > INDEX_NONE);
-			if (ElementIndexLocal < 3 
-				&& ElementIndexLocal > INDEX_NONE
-				&& ElementIndex < Measure.Num()
-				&& ElementIndex < MuElementArray.Num()
-				&& ElementIndex < LambdaElementArray.Num()
-				&& ElementIndex < DmInverse.Num() 
-				&& ElementIndex > INDEX_NONE)
-			{
-				const Chaos::PMatrix<FSolverReal, 2, 2> DmInvT = DmInverse[ElementIndex].GetTransposed(), DmInv = DmInverse[ElementIndex]; 
-				const Chaos::PMatrix<FSolverReal, 3, 2> Fe = F(ElementIndex, Particles);
-
-				const PMatrix<FSolverReal, 3, 2> Re = ComputeR(Fe);
-				const PMatrix<FSolverReal, 2, 2> FTF = ComputeFTF(Fe);
-				const FSolverReal J = FMath::Sqrt(FTF.Determinant());
-
-				PMatrix<FSolverReal, 3, 2> Pe(TVec3<FSolverReal>((FSolverReal)0.), TVec3<FSolverReal>((FSolverReal)0.)), ForceTerm(TVec3<FSolverReal>((FSolverReal)0.), TVec3<FSolverReal>((FSolverReal)0.));
-				Pe = FSolverReal(2) * MuElementArray[ElementIndex] * (Fe - Re) + LambdaElementArray[ElementIndex] * (J - FSolverReal(1)) * J * Fe * ComputeFTF(Fe).Inverse(); 
-
-				ForceTerm = -Measure[ElementIndex] * Pe * DmInverse[ElementIndex].GetTransposed();
-
-				Chaos::TVector<FSolverReal, 3> Dx((FSolverReal)0.);
-				if (ElementIndexLocal > 0)
-				{
-					for (int32 c = 0; c < 3; c++)
-					{
-						Dx[c] += ForceTerm.M[ElementIndexLocal * 3 - 3 + c];
-					}
-				}
-				else
-				{
-					for (int32 c = 0; c < 3; c++)
-					{
-						for (int32 h = 0; h < 2; h++)
-						{
-							Dx[c] -= ForceTerm.M[h * 3 + c];
-						}
-					}
-				}
-
-				Dx *= Dt * Dt;
-
-				ParticleResidual -= Dx;
-
-				PMatrix<FSolverReal, 3, 2> dJ(TVec3<FSolverReal>((FSolverReal)0.), TVec3<FSolverReal>((FSolverReal)0.));
-
-				dJdF32(Fe, dJ);
-
-				FSolverReal Coeff = Dt * Dt * Measure[ElementIndex];
-
-				if (ElementIndexLocal == 0) 
-				{
-					FSolverReal DmInvsum = FSolverReal(0);
-					for (int32 nu = 0; nu < 2; nu++) 
-					{
-						FSolverReal localDmsum = FSolverReal(0);
-						for (int32 k = 0; k < 2; k++) 
-						{
-							localDmsum += DmInv.M[nu * 2 + k];
-						}
-						DmInvsum += localDmsum * localDmsum;
-					}
-					for (int32 alpha = 0; alpha < 3; alpha++) 
-					{
-						ParticleHessian.SetAt(alpha, alpha, ParticleHessian.GetAt(alpha, alpha) + Coeff * FSolverReal(2) * MuElementArray[ElementIndex] * DmInvsum);
-					}
-
-					const PMatrix<FSolverReal, 3, 2> dJDmInvT = dJ * DmInvT; 
-					Chaos::TVector<FSolverReal, 3> L((FSolverReal)0.);
-					for (int32 alpha = 0; alpha < 3; alpha++) 
-					{
-						for (int32 k = 0; k < 2; k++) 
-						{
-							L[alpha] += dJDmInvT.M[3 * k + alpha]; 
-						}
-					}
-					for (int32 alpha = 0; alpha < 3; alpha++)
-					{
-						ParticleHessian.SetRow(alpha, ParticleHessian.GetRow(alpha) + Coeff * LambdaElementArray[ElementIndex] * L[alpha] * L);
-					}
-				}
-				else 
-				{
-					FSolverReal DmInvsum = FSolverReal(0);
-					for (int32 nu = 0; nu < 2; nu++)
-					{
-						DmInvsum += DmInv.GetAt(ElementIndexLocal - 1, nu) * DmInv.GetAt(ElementIndexLocal - 1, nu);
-					}
-					for (int32 alpha = 0; alpha < 3; alpha++)
-					{
-						ParticleHessian.SetAt(alpha, alpha, ParticleHessian.GetAt(alpha, alpha) + Dt * Dt * Measure[ElementIndex] * FSolverReal(2) * MuElementArray[ElementIndex] * DmInvsum);
-					}
-
-					const PMatrix<FSolverReal, 3, 2> dJDmInvT = dJ * DmInvT; 
-
-					Chaos::TVector<FSolverReal, 3> L((FSolverReal)0.);
-					for (int32 alpha = 0; alpha < 3; alpha++) 
-					{
-						const int32 IndexVisited = ElementIndexLocal* 3 - 3 + alpha;
-						if (IndexVisited < 6 && IndexVisited > INDEX_NONE)
-						{
-							L[alpha] = dJDmInvT.M[IndexVisited];
-						}
-					}
-					for (int32 alpha = 0; alpha < 3; alpha++)
-					{
-						ParticleHessian.SetRow(alpha, ParticleHessian.GetRow(alpha) + Dt * Dt * Measure[ElementIndex] * LambdaElementArray[ElementIndex] * L[alpha] * L);
-					}
-				}
-			}
-		}
-
-
-
 	protected:
 		mutable TArray<FSolverMatrix22> DmInverse;
 
diff --git a/Engine/Source/Runtime/Experimental/Chaos/Public/Chaos/Deformable/GaussSeidelCorotatedConstraints.h b/Engine/Source/Runtime/Experimental/Chaos/Public/Chaos/Deformable/GaussSeidelCorotatedConstraints.h
--- a/Engine/Source/Runtime/Experimental/Chaos/Public/Chaos/Deformable/GaussSeidelCorotatedConstraints.h
+++ b/Engine/Source/Runtime/Experimental/Chaos/Public/Chaos/Deformable/GaussSeidelCorotatedConstraints.h
@@ -28,7 +28,24 @@ namespace Chaos::Softs
 		using Base::F;
 
 	public:
-		FGaussSeidelCorotatedTetrahedralConstraints(
+		CHAOS_API FGaussSeidelCorotatedTetrahedralConstraints(
+			const FSolverParticlesRange& InParticles,
+			const TArray<TVector<int32, 4>>& InMesh,
+			const TArray<FSolverReal>& EMeshArray,
+			const TArray<FSolverReal>& NuMeshArray,
+			TArray<FSolverReal>&& AlphaJMeshArray,
+			TArray<TArray<int32>>&& IncidentElementsIn,
+			TArray<TArray<int32>>&& IncidentElementsLocalIn,
+			const int32 ParticleStartIndexIn,
+			const int32 ParticleEndIndexIn,
+			const bool bDoQuasistaticsIn = false,
+			const bool bDoSORIn = true,
+			const FSolverReal InOmegaSOR = (FSolverReal)1.6,
+			const FDeformableXPBDCorotatedParams& InParams = FDeformableXPBDCorotatedParams(),
+			const FSolverReal& NuMesh = (FSolverReal).3,
+			const bool bRecordMetricIn = false);
+
+		CHAOS_API FGaussSeidelCorotatedTetrahedralConstraints(
 			const FSolverParticles& InParticles,
 			const TArray<TVector<int32, 4>>& InMesh,
 			const TArray<FSolverReal>& EMeshArray,
@@ -43,361 +60,47 @@ namespace Chaos::Softs
 			const FSolverReal InOmegaSOR = (FSolverReal)1.6,
 			const FDeformableXPBDCorotatedParams& InParams = FDeformableXPBDCorotatedParams(),
 			const FSolverReal& NuMesh = (FSolverReal).3,
-			const bool bRecordMetricIn = false
-		)
-			: Base(InParticles, InMesh, EMeshArray, NuMeshArray, MoveTemp(AlphaJMeshArray), InParams, NuMesh, bRecordMetricIn, false), IncidentElements(IncidentElementsIn), IncidentElementsLocal(IncidentElementsLocalIn),
-			ParticleStartIndex(ParticleStartIndexIn), ParticleEndIndex(ParticleEndIndexIn), bDoQuasistatics(bDoQuasistaticsIn), bDoSOR(bDoSORIn), OmegaSOR(InOmegaSOR)
-		{
-			Base::LambdaArray.SetNum(0);
-			Particle2Incident.Init(INDEX_NONE, ParticleEndIndexIn - ParticleStartIndexIn);
-			for (int32 i = 0; i < IncidentElements.Num(); i++)
-			{
-				Particle2Incident[i] = i;
-			}
-			for (int32 e = 0; e < InMesh.Num(); e++)
-			{
-				for (int32 r = 0; r < 3; r++) {
-					for (int32 c = 0; c < 3; c++) {
-						DmInverse[(3 * 3) * e + 3 * r + c] /= Base::AlphaJArray[e];
-					}
-				}
-			}
-			if (!bDoQuasistatics)
-			{
-				xtilde.Init(TVector<FSolverReal, 3>(FSolverReal(0.)), ParticleEndIndex - ParticleStartIndex);
-			}
-			if (bDoSOR)
-			{
-				X_k_1.Init(TVector<FSolverReal, 3>(FSolverReal(0.)), ParticleEndIndex - ParticleStartIndex);
-				X_k.Init(TVector<FSolverReal, 3>(FSolverReal(0.)), ParticleEndIndex - ParticleStartIndex);
-			}
-			InitColor(InParticles);
-			InitializeCorotatedLambdas();
-		}
+			const bool bRecordMetricIn = false);
 
 		virtual ~FGaussSeidelCorotatedTetrahedralConstraints() = default;
 
+		CHAOS_API void Apply(FSolverParticlesRange& Particles, const FSolverReal Dt) const;
 
-		void Apply(FSolverParticlesRange& Particles, const FSolverReal Dt) const 
-		{
-			TRACE_CPUPROFILER_EVENT_SCOPE(STAT_ChaosGaussSeidelApply);
-			for (int32 i = 0; i < ParticlesPerColor.Num(); i++)
-			{
-				int32 NumBatch = ParticlesPerColor[i].Num() / CorotatedParams.XPBDCorotatedBatchSize;
-				if (ParticlesPerColor[i].Num() % CorotatedParams.XPBDCorotatedBatchSize != 0)
-				{
-					NumBatch += 1;
-				}
-
-				PhysicsParallelFor(NumBatch, [&](const int32 BatchIndex)
-					{
-						for (int32 BatchSubIndex = 0; BatchSubIndex < CorotatedParams.XPBDCorotatedBatchSize; BatchSubIndex++) {
-							int32 TaskIndex = CorotatedParams.XPBDCorotatedBatchSize * BatchIndex + BatchSubIndex;
-							if (TaskIndex < ParticlesPerColor[i].Num())
-							{
-								int32 ParticleIndex = ParticlesPerColor[i][TaskIndex];
-								Chaos::TVector<FSolverReal, 3> Dx = ComputeDeltax(ParticleIndex, Particle2Incident[ParticleIndex - ParticleStartIndex], Particles,
-									Dt);
-								Particles.P(ParticleIndex) += Dx;
-							}
-						}
-					}, NumBatch < CorotatedParams.XPBDCorotatedBatchThreshold);
+		CHAOS_API void Apply(FSolverParticles& Particles, const FSolverReal Dt) const;
 
-			}
-			ApplySOR(Particles, Dt);
-			CurrentIt += 1;
-		}
+		CHAOS_API void Init(const FSolverReal Dt, const FSolverParticlesRange& Particles) const;
 
-		void Apply(FSolverParticles& Particles, const FSolverReal Dt) const
-		{
-			FSolverParticlesRange Range(&Particles, /*InOffset = */ 0, /*InRange = */ Particles.Size());
-			Apply(Range, Dt);
-		}
-
-		void Init(const FSolverReal Dt, const FSolverParticles& Particles) const 
-		{
-			if (!bDoQuasistatics)
-			{
-				for (int32 i = 0; i < ParticleEndIndex - ParticleStartIndex; i++)
-				{
-					xtilde[i] = Particles.X(ParticleStartIndex + i) + Dt * Particles.V(ParticleStartIndex + i);
-				}
-			}
-			CurrentIt = 0;
-		}
+		CHAOS_API void Init(const FSolverReal Dt, const FSolverParticles& Particles) const;
 
 		void Init() const override {}
 
 		TArray<TArray<int32>>& GetIncidentElements() { return IncidentElements; }
 		TArray<TArray<int32>>& GetIncidentElementsLocal() { return IncidentElementsLocal; }
 		TArray<TVector<int32, 4>> GetMeshConstraints() const { return MeshConstraints; }
-		void SetParticlesPerColor(TArray<TArray<int32>>&& InParticlesPerColor) {ParticlesPerColor = MoveTemp(InParticlesPerColor); }
-
-		TArray<TArray<int32>> GetMeshArray() const 
-		{
-			TArray<TArray<int32>> MeshArray;
-			MeshArray.SetNum(MeshConstraints.Num());
-			for (int32 i = 0; i < MeshConstraints.Num(); i++)
-			{
-				MeshArray[i].SetNum(4);
-				for (int32 ie = 0; ie < 4; ie++)
-				{
-					MeshArray[i][ie] = MeshConstraints[i][ie];
-				}
-			}
-			return MeshArray;
-		}
-
-		// Apply Successive Over Relaxation (SOR) to accelerate convergence rate
-		void ApplySOR(FSolverParticlesRange& Particles, const FSolverReal Dt) const
-		{
-			if (bDoSOR)
-			{
-				PhysicsParallelFor(ParticleEndIndex - ParticleStartIndex, [&](const int32 i)
-					{
-						int32 ParticleIndex = i + ParticleStartIndex;
-						if (Particles.InvM(ParticleIndex) != FSolverReal(0) && CurrentIt > 3)
-						{
-							Particles.P(ParticleIndex) = OmegaSOR * (Particles.P(ParticleIndex) - X_k_1[i]) + X_k_1[i];
-						}
-						X_k_1[i] = X_k[i];
-						X_k[i] = Particles.P(ParticleIndex);
-					}, ParticleEndIndex - ParticleStartIndex < 1000);
-			}
-		}
-
-		Chaos::TVector<FSolverReal, 3> ComputePerParticleResidual(const int32 p, const int32 IncidentIndex, const FSolverParticlesRange& Particles,
-			const FSolverReal Dt, const bool AddMass = true) const
-		{
-
-			Chaos::TVector<FSolverReal, 3> res(FSolverReal(0));
-			if (AddMass) {
-				for (int32 alpha = 0; alpha < 3; alpha++) {
-					res[alpha] = Particles.M(p) * (Particles.P(p)[alpha] - xtilde[p - ParticleStartIndex][alpha]);
-				}
-			}
-			for (int32 j = 0; j < IncidentElements[IncidentIndex].Num(); j++) {
-				int32 local_index = IncidentElementsLocal[IncidentIndex][j];
-				int32 e = IncidentElements[IncidentIndex][j];
-				Chaos::PMatrix<FSolverReal, 3, 3> DmInvT = Base::ElementDmInv(e).GetTransposed();
-				Chaos::PMatrix<FSolverReal, 3, 3> Fe = Base::F(e, Particles);
-				Chaos::PMatrix<FSolverReal, 3, 3> P((FSolverReal)0.);
-
-				ComputeStress(Fe, MuElementArray[e], LambdaElementArray[e], P);
-
-				Chaos::PMatrix<FSolverReal, 3, 3> force_term = -Measure[e] * DmInvT * P;
-				Chaos::TVector<FSolverReal, 3> dx((FSolverReal)0.);
-				if (local_index > 0)
-				{
-					for (int32 c = 0; c < 3; c++)
-					{
-						dx[c] += force_term.GetAt(c, local_index - 1);
-					}
-				}
-				else
-				{
-					for (int32 c = 0; c < 3; c++)
-					{
-						for (int32 h = 0; h < 3; h++)
-						{
-							dx[c] -= force_term.GetAt(c, h);
-						}
-					}
-				}
-
-				for (int32 ll = 0; ll < 3; ll++) {
-					dx[ll] *= Dt * Dt;
-				}
-
-				for (int32 alpha = 0; alpha < 3; alpha++) {
-					res[alpha] -= dx[alpha];
-				}
-			}
-			return res;
-		}
+		void SetParticlesPerColor(TArray<TArray<int32>>&& InParticlesPerColor) { ParticlesPerColor = MoveTemp(InParticlesPerColor); }
 
+		CHAOS_API TArray<TArray<int32>> GetMeshArray() const;
 
-		Chaos::PMatrix<FSolverReal, 3, 3> ComputePerParticleCorotatedHessianSimple(const int32 p, const int32 IncidentIndex, const FSolverParticlesRange& Particles,
-			const FSolverReal Dt, const bool AddMass = true) const
-		{
-			Chaos::PMatrix<FSolverReal, 3, 3> final_hessian = Chaos::PMatrix<FSolverReal, 3, 3>::Zero;
-			if (AddMass) {
-				for (int32 alpha = 0; alpha < 3; alpha++) {
-					final_hessian.SetAt(alpha, alpha, Particles.M(p));
-				}
-			}
-
-			for (int32 j = 0; j < IncidentElements[IncidentIndex].Num(); j++) {
-				int32 local_index = IncidentElementsLocal[IncidentIndex][j];
-				int32 ElementIndex = IncidentElements[IncidentIndex][j];
-				Chaos::PMatrix<FSolverReal, 3, 3> DmInv = Base::ElementDmInv(ElementIndex);
-				Chaos::PMatrix<FSolverReal, 3, 3> Fe = Base::F(ElementIndex, Particles);
-
-				FSolverReal Coeff = Dt * Dt * Measure[ElementIndex];
-
-				ComputeHessianHelper(Fe, DmInv, MuElementArray[ElementIndex], LambdaElementArray[ElementIndex], local_index, Coeff, final_hessian);
-
-			}
-			return final_hessian;
-		}
-
-
-		void AddHyperelasticResidualAndHessian (const FSolverParticlesRange& Particles, const int32 ElementIndex, const int32 ElementIndexLocal, const FSolverReal Dt, TVec3<FSolverReal>& ParticleResidual, Chaos::PMatrix<FSolverReal, 3, 3>& ParticleHessian)
-		{
-			Chaos::PMatrix<FSolverReal, 3, 3> DmInvT = Base::ElementDmInv(ElementIndex).GetTransposed();
-			Chaos::PMatrix<FSolverReal, 3, 3> Fe = F(ElementIndex, Particles);
-			Chaos::PMatrix<FSolverReal, 3, 3> P((FSolverReal)0.);
-
-			this->ComputeStress(Fe, MuElementArray[ElementIndex], LambdaElementArray[ElementIndex], P);
-
-			Chaos::PMatrix<FSolverReal, 3, 3> force_term = -Measure[ElementIndex] * DmInvT * P;
-			Chaos::TVector<FSolverReal, 3> dx((FSolverReal)0.);
-			if (ElementIndexLocal > 0)
-			{
-				for (int32 c = 0; c < 3; c++)
-				{
-					dx[c] += force_term.GetAt(c, ElementIndexLocal - 1);
-				}
-			}
-			else
-			{
-				for (int32 c = 0; c < 3; c++)
-				{
-					for (int32 h = 0; h < 3; h++)
-					{
-						dx[c] -= force_term.GetAt(c, h);
-					}
-				}
-			}
-
-			for (int32 ll = 0; ll < 3; ll++) {
-				dx[ll] *= Dt * Dt;
-			}
-
-			for (int32 alpha = 0; alpha < 3; alpha++) {
-				ParticleResidual[alpha] -= dx[alpha];
-			}
-
-			ComputeHessianHelper(Fe, Base::ElementDmInv(ElementIndex), MuElementArray[ElementIndex], LambdaElementArray[ElementIndex], ElementIndexLocal, Dt * Dt * Measure[ElementIndex], ParticleHessian);
-
-		}
-
-		void AddHyperelasticResidualAndHessian(const FSolverParticles& Particles, const int32 ElementIndex, const int32 ElementIndexLocal, const FSolverReal Dt, TVec3<FSolverReal>& ParticleResidual, Chaos::PMatrix<FSolverReal, 3, 3>& ParticleHessian)
-		{
-			AddHyperelasticResidualAndHessian(FSolverParticlesRange(const_cast<FSolverParticles*>(&Particles), /*InOffset = */ 0, /*InRange = */ Particles.Size()), 
-				ElementIndex, ElementIndexLocal, Dt, ParticleResidual, ParticleHessian);
-		}
+		CHAOS_API void AddHyperelasticResidualAndHessian(const FSolverParticlesRange& Particles, const int32 ElementIndex, const int32 ElementIndexLocal, const FSolverReal Dt, TVec3<FSolverReal>& ParticleResidual, Chaos::PMatrix<FSolverReal, 3, 3>& ParticleHessian);
 
 	protected:
-		 void InitColor(const FSolverParticles& Particles)
-		{
-			ParticlesPerColor = ComputeNodalColoring(MeshConstraints, Particles, ParticleStartIndex, ParticleEndIndex, IncidentElements, IncidentElementsLocal);
-		}
-
-		 void InitializeCorotatedLambdas()
-		 {
-			 ComputeStress = [](const Chaos::PMatrix<FSolverReal, 3, 3>& Fe, const FSolverReal mu, const FSolverReal lambda, Chaos::PMatrix<FSolverReal, 3, 3>& P)
-			 {
-				 PCorotated(Fe, mu, lambda, P);
-			 };
-			 ComputeHessianHelper = [](const Chaos::PMatrix<FSolverReal, 3, 3>& Fe, const Chaos::PMatrix<FSolverReal, 3, 3>& DmInv, const FSolverReal mu, const FSolverReal lambda, const int32 local_index, const FSolverReal Coeff, Chaos::PMatrix<FSolverReal, 3, 3>& final_hessian) 
-			 {
-				 Chaos::PMatrix<FSolverReal, 3, 3> JFinvT((FSolverReal)0.);
-				 JFinvT.SetAt(0, 0, Fe.GetAt(1, 1) * Fe.GetAt(2, 2) - Fe.GetAt(2, 1) * Fe.GetAt(1, 2));
-				 JFinvT.SetAt(0, 1, Fe.GetAt(2, 0) * Fe.GetAt(1, 2) - Fe.GetAt(1, 0) * Fe.GetAt(2, 2));
-				 JFinvT.SetAt(0, 2, Fe.GetAt(1, 0) * Fe.GetAt(2, 1) - Fe.GetAt(2, 0) * Fe.GetAt(1, 1));
-				 JFinvT.SetAt(1, 0, Fe.GetAt(2, 1) * Fe.GetAt(0, 2) - Fe.GetAt(0, 1) * Fe.GetAt(2, 2));
-				 JFinvT.SetAt(1, 1, Fe.GetAt(0, 0) * Fe.GetAt(2, 2) - Fe.GetAt(2, 0) * Fe.GetAt(0, 2));
-				 JFinvT.SetAt(1, 2, Fe.GetAt(2, 0) * Fe.GetAt(0, 1) - Fe.GetAt(0, 0) * Fe.GetAt(2, 1));
-				 JFinvT.SetAt(2, 0, Fe.GetAt(0, 1) * Fe.GetAt(1, 2) - Fe.GetAt(1, 1) * Fe.GetAt(0, 2));
-				 JFinvT.SetAt(2, 1, Fe.GetAt(1, 0) * Fe.GetAt(0, 2) - Fe.GetAt(0, 0) * Fe.GetAt(1, 2));
-				 JFinvT.SetAt(2, 2, Fe.GetAt(0, 0) * Fe.GetAt(1, 1) - Fe.GetAt(1, 0) * Fe.GetAt(0, 1));
-
-				 Chaos::PMatrix<FSolverReal, 3, 3> JFinv = JFinvT.GetTransposed();
+		CHAOS_API void InitColor(const FSolverParticlesRange& Particles);
+		CHAOS_API void InitColor(const FSolverParticles& Particles);
+		CHAOS_API void InitializeCorotatedLambdas();
 
-				 if (local_index == 0) {
-					 FSolverReal DmInvsum = FSolverReal(0);
-					 for (int32 nu = 0; nu < 3; nu++) {
-						 FSolverReal localDmsum = FSolverReal(0);
-						 for (int32 k = 0; k < 3; k++) {
-							 localDmsum += DmInv.GetAt(k, nu);
-						 }
-						 DmInvsum += localDmsum * localDmsum;
-					 }
-					 for (int32 alpha = 0; alpha < 3; alpha++) {
-						 final_hessian.SetAt(alpha, alpha, final_hessian.GetAt(alpha, alpha) + Coeff * FSolverReal(2) * mu * DmInvsum);
-					 }
+		CHAOS_API void ApplySOR(FSolverParticlesRange& Particles, const FSolverReal Dt) const;
 
-					 Chaos::PMatrix<FSolverReal, 3, 3> DmInvJFinv = JFinv * DmInv;
-					 Chaos::TVector<FSolverReal, 3> l((FSolverReal)0.);
-					 for (int32 alpha = 0; alpha < 3; alpha++) {
-						 for (int32 k = 0; k < 3; k++) {
-							 l[alpha] += DmInvJFinv.GetAt(k, alpha);
-						 }
-					 }
-					 for (int32 alpha = 0; alpha < 3; alpha++)
-					 {
-						 final_hessian.SetRow(alpha, final_hessian.GetRow(alpha) + Coeff * lambda * l[alpha] * l);
-					 }
+		CHAOS_API Chaos::TVector<FSolverReal, 3> ComputePerParticleResidual(const int32 ParticleIdx, const int32 IncidentIndex, const FSolverParticlesRange& Particles,
+			const FSolverReal Dt, const bool AddMass = true) const;
 
-				 }
-				 else {
-					 FSolverReal DmInvsum = FSolverReal(0);
-					 for (int32 nu = 0; nu < 3; nu++)
-					 {
-						 DmInvsum += DmInv.GetAt(local_index - 1, nu) * DmInv.GetAt(local_index - 1, nu);
-					 }
-					 for (int32 alpha = 0; alpha < 3; alpha++)
-					 {
-						 final_hessian.SetAt(alpha, alpha, final_hessian.GetAt(alpha, alpha) + Coeff * FSolverReal(2) * mu * DmInvsum);
-					 }
-
-					 Chaos::PMatrix<FSolverReal, 3, 3> DmInvJFinv = JFinv * DmInv;
-					 Chaos::TVector<FSolverReal, 3> l((FSolverReal)0.);
-					 for (int32 alpha = 0; alpha < 3; alpha++) {
-						 l[alpha] = DmInvJFinv.GetAt(local_index - 1, alpha);
-					 }
-					 for (int32 alpha = 0; alpha < 3; alpha++)
-					 {
-						 final_hessian.SetRow(alpha, final_hessian.GetRow(alpha) + Coeff * lambda * l[alpha] * l);
-					 }
-				 }
-			 };
-
-			 AddAdditionalRes = [](const FSolverParticlesRange& InParticles, const int32 p, const FSolverReal Dt, TVec3<FSolverReal>& res) {};
-			 AddAdditionalHessian = [](const FSolverParticlesRange& InParticles, const int32 p, const FSolverReal Dt, Chaos::PMatrix<FSolverReal, 3, 3>& hessian) {};
-		 }
+		CHAOS_API Chaos::PMatrix<FSolverReal, 3, 3> ComputePerParticleCorotatedHessianSimple(const int32 ParticleIdx, const int32 IncidentIndex, const FSolverParticlesRange& Particles,
+			const FSolverReal Dt, const bool AddMass = true) const;
 
 	private:
+		TVector<FSolverReal, 3> ComputeDeltax(const int32 ParticleIdx, const int32 IncidentIndex, const FSolverParticlesRange& Particles,
+			const FSolverReal Dt) const;
 
-		TVector<FSolverReal, 3> ComputeDeltax(const int32 p, const int32 IncidentIndex, const FSolverParticlesRange& Particles,
-			const FSolverReal Dt) const
-		{
-			Chaos::TVector<FSolverReal, 3> ParticleResidual = ComputePerParticleResidual(p, Particle2Incident[p], Particles, Dt, !bDoQuasistatics);
-
-			AddAdditionalRes(Particles, p, Dt, ParticleResidual);
-
-			if (ParticleResidual.Size() > LocalNewtonTol)
-			{
-				Chaos::PMatrix<FSolverReal, 3, 3> SimplifiedHessian = ComputePerParticleCorotatedHessianSimple(p, Particle2Incident[p], Particles, Dt, !bDoQuasistatics);
-				AddAdditionalHessian(Particles, p, Dt, SimplifiedHessian);
-				
-				FSolverReal HessianDet = SimplifiedHessian.Determinant();
-				if (HessianDet > UE_SMALL_NUMBER)
-				{
-					Chaos::PMatrix<FSolverReal, 3, 3> HessianInv = SimplifiedHessian.SymmetricCofactorMatrix();
-					HessianInv *= FSolverReal(1) / HessianDet;
-					return HessianInv.GetTransposed() * (-ParticleResidual);
-				}
-			}
-			return TVector<FSolverReal, 3>((FSolverReal)0.);
-		}
-
-		
-
-	protected:	
-
-		TArray<int32> Particle2Incident;
+	protected:
 		TArray<TArray<int32>> IncidentElements;
 		TArray<TArray<int32>> IncidentElementsLocal;
 		FSolverReal LocalNewtonTol = FSolverReal(1e-5);
@@ -409,12 +112,11 @@ namespace Chaos::Softs
 		mutable TArray<Chaos::TVector<FSolverReal, 3>> X_k;
 		bool bDoQuasistatics = false;
 		bool bDoSOR = true;
-		//Parameter to control the behavior of SOR. 1-1.7 is the usual range. Larger paramater leads to potentially more convergence but less stable behavior. 
 		FSolverReal OmegaSOR = FSolverReal(1.6);
 		mutable int32 CurrentIt = 0;
-		TFunction<void(const Chaos::PMatrix<FSolverReal, 3, 3>&, const FSolverReal, const FSolverReal, Chaos::PMatrix<FSolverReal, 3, 3>& )> ComputeStress;
+		TFunction<void(const Chaos::PMatrix<FSolverReal, 3, 3>&, const FSolverReal, const FSolverReal, Chaos::PMatrix<FSolverReal, 3, 3>&)> ComputeStress;
 		TFunction<void(const Chaos::PMatrix<FSolverReal, 3, 3>&, const Chaos::PMatrix<FSolverReal, 3, 3>&, const FSolverReal, const FSolverReal, const int32, const FSolverReal, Chaos::PMatrix<FSolverReal, 3, 3>&)> ComputeHessianHelper;
-		
+
 	public:
 		TFunction<void(const FSolverParticlesRange&, const int32, const FSolverReal, TVec3<FSolverReal>&)> AddAdditionalRes;
 		TFunction<void(const FSolverParticlesRange&, const int32, const FSolverReal, Chaos::PMatrix<FSolverReal, 3, 3>&)> AddAdditionalHessian;
diff --git a/Engine/Source/Runtime/Experimental/Chaos/Public/Chaos/Deformable/GaussSeidelLinearCodimensionalConstraints.h b/Engine/Source/Runtime/Experimental/Chaos/Public/Chaos/Deformable/GaussSeidelLinearCodimensionalConstraints.h
--- a/Engine/Source/Runtime/Experimental/Chaos/Public/Chaos/Deformable/GaussSeidelLinearCodimensionalConstraints.h
+++ b/Engine/Source/Runtime/Experimental/Chaos/Public/Chaos/Deformable/GaussSeidelLinearCodimensionalConstraints.h
@@ -8,150 +8,42 @@
 #include "Chaos/Framework/Parallel.h"
 #include "Chaos/Deformable/ChaosDeformableSolverTypes.h"
 #include "Chaos/Vector.h"
+#include "Chaos/SoftsSolverParticlesRange.h"
 
 
 namespace Chaos::Softs
-{	
+{
 	// Triangle constraint using linear constitutive model
 	// 3D triangles are projected to the 2D space for deformation gradient (3*2 matrix) computation
 	class FGaussSeidelLinearTriangleConstraints
 	{
 
 	public:
-		FGaussSeidelLinearTriangleConstraints(
+		CHAOS_API FGaussSeidelLinearTriangleConstraints(
 			const FSolverParticles& InParticles,
 			const TArray<TVector<int32, 3>>& InMesh,
 			const FSolverReal& EMesh = (FSolverReal)10.0,
-			const FSolverReal& NuMesh = (FSolverReal).3
-		)
-			: MeshConstraints(InMesh)
-		{
-			Measure.Init((FSolverReal)0., MeshConstraints.Num());
-			Lambda = EMesh * NuMesh / (((FSolverReal)1. + NuMesh) * ((FSolverReal)1. - (FSolverReal)2. * NuMesh));
-			Mu = EMesh / ((FSolverReal)2. * ((FSolverReal)1. + NuMesh));
-
-			MuElementArray.Init(Mu, MeshConstraints.Num());
-			LambdaElementArray.Init(Lambda, MeshConstraints.Num());
-
-			InitializeCodimensionData(InParticles);
-		}
+			const FSolverReal& NuMesh = (FSolverReal).3);
 
-		FGaussSeidelLinearTriangleConstraints(
+		CHAOS_API FGaussSeidelLinearTriangleConstraints(
 			const FSolverParticles& InParticles,
 			const TArray<TVector<int32, 3>>& InMesh,
 			const TArray<FSolverReal>& EMeshArray,
-			const FSolverReal& NuMesh = (FSolverReal).3
-		)
-			: MeshConstraints(InMesh)
-		{
-			Measure.Init((FSolverReal)0., MeshConstraints.Num());
-
-			InitializeCodimensionData(InParticles);
-			LambdaElementArray.Init((FSolverReal)0., MeshConstraints.Num());
-			MuElementArray.Init((FSolverReal)0., MeshConstraints.Num());
-			
-			for (int32 e = 0; e < InMesh.Num(); e++)
-			{
-				LambdaElementArray[e] = EMeshArray[e] * NuMesh / (((FSolverReal)1. + NuMesh) * ((FSolverReal)1. - (FSolverReal)2. * NuMesh));
-				MuElementArray[e] = EMeshArray[e] / ((FSolverReal)2. * ((FSolverReal)1. + NuMesh));
-			}
-		}
+			const FSolverReal& NuMesh = (FSolverReal).3);
 
 		virtual ~FGaussSeidelLinearTriangleConstraints() = default;
 
-		PMatrix<FSolverReal, 3, 3> DsInit(const int32 E, const FSolverParticles& InParticles)  
-		{
-			PMatrix<FSolverReal, 3, 3> Result((FSolverReal)0.);
-			for (int32 i = 0; i < 3; i++) 
-			{
-				for (int32 c = 0; c < 3; c++) 
-				{
-					Result.SetAt(c, i, InParticles.X(MeshConstraints[E][i + 1])[c] - InParticles.X(MeshConstraints[E][0])[c]);
-				}
-			}
-			return Result;
-		}
-
-		PMatrix<FSolverReal, 3, 2> Ds(const int32 E, const FSolverParticles& InParticles) const 
-		{
-			if (INDEX_NONE < E && E< MeshConstraints.Num()
-				&& INDEX_NONE < MeshConstraints[E][0] && MeshConstraints[E][0] < (int32)InParticles.Size()
-				&& INDEX_NONE < MeshConstraints[E][1] && MeshConstraints[E][1] < (int32)InParticles.Size()
-				&& INDEX_NONE < MeshConstraints[E][2] && MeshConstraints[E][2] < (int32)InParticles.Size())
-			{
-				const TVec3<FSolverReal> P1P0 = InParticles.P(MeshConstraints[E][1]) - InParticles.P(MeshConstraints[E][0]);
-				const TVec3<FSolverReal> P2P0 = InParticles.P(MeshConstraints[E][2]) - InParticles.P(MeshConstraints[E][0]);
-
-				return PMatrix<FSolverReal, 3, 2>(
-					P1P0[0], P1P0[1], P1P0[2],
-					P2P0[0], P2P0[1], P2P0[2]);
-			}
-			else
-			{
-				return PMatrix<FSolverReal, 3, 2>(
-					(FSolverReal)0., (FSolverReal)0, (FSolverReal)0,
-					(FSolverReal)0, (FSolverReal)0, (FSolverReal)0);
-			}
-		}
-
-		PMatrix<FSolverReal, 3, 2> F(const int32 E, const FSolverParticles& InParticles) const 
-		{
-			if (INDEX_NONE < E && E < DmInverse.Num())
-			{
-				return Ds(E, InParticles) * DmInverse[E];
-			}
-			else
-			{
-				return PMatrix<FSolverReal, 3, 2>(
-					(FSolverReal)0., (FSolverReal)0, (FSolverReal)0,
-					(FSolverReal)0, (FSolverReal)0, (FSolverReal)0);
-			}
-		}
+		CHAOS_API TArray<TArray<int32>> GetConstraintsArray() const;
 
-		TArray<TArray<int32>> GetConstraintsArray() const
-		{
-			TArray<TArray<int32>> Constraints;
-			Constraints.SetNum(MeshConstraints.Num());
-			for (int32 i = 0; i < MeshConstraints.Num(); i++)
-			{
-				Constraints[i].SetNum(3);
-				for (int32 j = 0; j < 3; j++)
-				{
-					Constraints[i][j] = MeshConstraints[i][j];
-				}
-			}
-			return Constraints;
-		}
+		CHAOS_API void AddHyperelasticResidualAndHessian(const FSolverParticlesRange& Particles, const int32 ElementIndex, const int32 ElementIndexLocal, const FSolverReal Dt, TVec3<FSolverReal>& ParticleResidual, Chaos::PMatrix<FSolverReal, 3, 3>& ParticleHessian);
 
 	protected:
 
-		void InitializeCodimensionData(const FSolverParticles& Particles)
-		{
-			Measure.Init((FSolverReal)0., MeshConstraints.Num());
-			DmInverse.Init(PMatrix<FSolverReal, 2, 2>(0.f, 0.f, 0.f), MeshConstraints.Num());
-			for (int32 e = 0; e < MeshConstraints.Num(); e++)
-			{
-				check(MeshConstraints[e][0] < (int32)Particles.Size() && MeshConstraints[e][0] > INDEX_NONE);
-				check(MeshConstraints[e][1] < (int32)Particles.Size() && MeshConstraints[e][1] > INDEX_NONE);
-				check(MeshConstraints[e][2] < (int32)Particles.Size() && MeshConstraints[e][2] > INDEX_NONE);
-				if (MeshConstraints[e][0] < (int32)Particles.Size() && MeshConstraints[e][0] > INDEX_NONE
-					&& MeshConstraints[e][1] < (int32)Particles.Size() && MeshConstraints[e][1] > INDEX_NONE
-					&& MeshConstraints[e][2] < (int32)Particles.Size() && MeshConstraints[e][2] > INDEX_NONE)
-				{
-					const TVec3<FSolverReal> X1X0 = Particles.GetX(MeshConstraints[e][1]) - Particles.GetX(MeshConstraints[e][0]);
-					const TVec3<FSolverReal> X2X0 = Particles.GetX(MeshConstraints[e][2]) - Particles.GetX(MeshConstraints[e][0]);
-					PMatrix<FSolverReal, 2, 2> Dm((FSolverReal)0., (FSolverReal)0., (FSolverReal)0.);
-					Dm.M[0] = X1X0.Size();
-					Dm.M[2] = X1X0.Dot(X2X0) / Dm.M[0];
-					Dm.M[3] = Chaos::TVector<FSolverReal, 3>::CrossProduct(X1X0, X2X0).Size() / Dm.M[0];
-					Measure[e] = Chaos::TVector<FSolverReal, 3>::CrossProduct(X1X0, X2X0).Size() / (FSolverReal)2.;
-					ensureMsgf(Measure[e] > (FSolverReal)0., TEXT("Degenerate triangle detected"));
+		PMatrix<FSolverReal, 3, 2> Ds(const int32 ElemIdx, const FSolverParticlesRange& InParticles) const;
 
-					PMatrix<FSolverReal, 2, 2> DmInv = Dm.Inverse();
-					DmInverse[e] = DmInv;
-				}
-			}
-		}
+		PMatrix<FSolverReal, 3, 2> F(const int32 ElemIdx, const FSolverParticlesRange& InParticles) const;
+
+		void InitializeCodimensionData(const FSolverParticlesRange& Particles);
 
 		static FSolverReal SafeRecip(const FSolverReal Len, const FSolverReal Fallback)
 		{
@@ -162,85 +54,6 @@ namespace Chaos::Softs
 			return Fallback;
 		}
 
-	public:
-		void AddHyperelasticResidualAndHessian(const FSolverParticles& Particles, const int32 ElementIndex, const int32 ElementIndexLocal, const FSolverReal Dt, TVec3<FSolverReal>& ParticleResidual, Chaos::PMatrix<FSolverReal, 3, 3>& ParticleHessian)
-		{
-			check(ElementIndex < DmInverse.Num() && ElementIndex > INDEX_NONE);
-			check(ElementIndex < MuElementArray.Num());
-			check(ElementIndex < Measure.Num());
-			check(ElementIndex < LambdaElementArray.Num());
-			check(ElementIndexLocal < 3 && ElementIndexLocal > INDEX_NONE);
-			if (ElementIndexLocal < 3 
-				&& ElementIndexLocal > INDEX_NONE
-				&& ElementIndex < Measure.Num()
-				&& ElementIndex < MuElementArray.Num()
-				&& ElementIndex < LambdaElementArray.Num()
-				&& ElementIndex < DmInverse.Num() 
-				&& ElementIndex > INDEX_NONE)
-			{
-				const Chaos::PMatrix<FSolverReal, 2, 2> DmInvFSolverReal = DmInverse[ElementIndex].GetTransposed(), DmInv = DmInverse[ElementIndex]; 
-				const Chaos::PMatrix<FSolverReal, 3, 2> Fe = F(ElementIndex, Particles);
-
-				PMatrix<FSolverReal, 3, 2> Pe(TVec3<FSolverReal>((FSolverReal)0.), TVec3<FSolverReal>((FSolverReal)0.)), ForceTerm(TVec3<FSolverReal>((FSolverReal)0.), TVec3<FSolverReal>((FSolverReal)0.));
-				Pe = FSolverReal(2) * MuElementArray[ElementIndex] * Fe; 
-
-				ForceTerm = -Measure[ElementIndex] * Pe * DmInverse[ElementIndex].GetTransposed();
-
-				Chaos::TVector<FSolverReal, 3> Dx((FSolverReal)0.);
-				if (ElementIndexLocal > 0)
-				{
-					for (int32 c = 0; c < 3; c++)
-					{
-						Dx[c] += ForceTerm.M[ElementIndexLocal * 3 - 3 + c];
-					}
-				}
-				else
-				{
-					for (int32 c = 0; c < 3; c++)
-					{
-						for (int32 h = 0; h < 2; h++)
-						{
-							Dx[c] -= ForceTerm.M[h * 3 + c];
-						}
-					}
-				}
-
-				Dx *= Dt * Dt;
-				ParticleResidual -= Dx;
-				FSolverReal Coeff = Dt * Dt * Measure[ElementIndex];
-
-				if (ElementIndexLocal == 0) 
-				{
-					FSolverReal DmInvsum = FSolverReal(0);
-					for (int32 nu = 0; nu < 2; nu++) 
-					{
-						FSolverReal localDmsum = FSolverReal(0);
-						for (int32 k = 0; k < 2; k++) 
-						{
-							localDmsum += DmInv.M[nu * 2 + k];
-						}
-						DmInvsum += localDmsum * localDmsum;
-					}
-					for (int32 alpha = 0; alpha < 3; alpha++) 
-					{
-						ParticleHessian.SetAt(alpha, alpha, ParticleHessian.GetAt(alpha, alpha) + Coeff * FSolverReal(2) * MuElementArray[ElementIndex] * DmInvsum);
-					}
-				}
-				else 
-				{
-					FSolverReal DmInvsum = FSolverReal(0);
-					for (int32 nu = 0; nu < 2; nu++)
-					{
-						DmInvsum += DmInv.GetAt(ElementIndexLocal - 1, nu) * DmInv.GetAt(ElementIndexLocal - 1, nu);
-					}
-					for (int32 alpha = 0; alpha < 3; alpha++)
-					{
-						ParticleHessian.SetAt(alpha, alpha, ParticleHessian.GetAt(alpha, alpha) + Dt * Dt * Measure[ElementIndex] * FSolverReal(2) * MuElementArray[ElementIndex] * DmInvsum);
-					}
-				}
-			}
-		}
-
 	protected:
 		mutable TArray<FSolverMatrix22> DmInverse;
 
diff --git a/Engine/Source/Runtime/Experimental/Chaos/Public/Chaos/Deformable/GaussSeidelMainConstraint.h b/Engine/Source/Runtime/Experimental/Chaos/Public/Chaos/Deformable/GaussSeidelMainConstraint.h
--- a/Engine/Source/Runtime/Experimental/Chaos/Public/Chaos/Deformable/GaussSeidelMainConstraint.h
+++ b/Engine/Source/Runtime/Experimental/Chaos/Public/Chaos/Deformable/GaussSeidelMainConstraint.h
@@ -11,6 +11,7 @@
 #include "Chaos/Deformable/GaussSeidelWeakConstraints.h"
 #include "Chaos/Deformable/GaussSeidelCorotatedConstraints.h"
 #include "Chaos/CollectionPropertyFacade.h"
+#include "Chaos/SoftsSolverParticlesRange.h"
 
 
 #define PERF_SCOPE(X) SCOPE_CYCLE_COUNTER(X); TRACE_CPUPROFILER_EVENT_SCOPE(X);
@@ -26,20 +27,29 @@ namespace Chaos::Softs
 {
 	// TFunction type aliases for Gauss-Seidel constraint lambdas
 	// (Particles, ParticleIdx, Dt, OutResidual, OutHessian)
-	using FGSInitialResidualHessianFunc = TFunction<void(const FSolverParticles&, const int32, const FSolverReal, TVec3<FSolverReal>&, Chaos::PMatrix<FSolverReal, 3, 3>&)>;
+	using FGSInitialResidualHessianFunc = TFunction<void(const FSolverParticlesRange&, const int32, const FSolverReal, TVec3<FSolverReal>&, Chaos::PMatrix<FSolverReal, 3, 3>&)>;
 	// (Particles, ConstraintIndex, LocalIndex, Dt, OutResidual, OutHessian)
-	using FGSConstraintResidualHessianFunc = TFunction<void(const FSolverParticles&, const int32, const int32, const FSolverReal, TVec3<FSolverReal>&, Chaos::PMatrix<FSolverReal, 3, 3>&)>;
+	using FGSConstraintResidualHessianFunc = TFunction<void(const FSolverParticlesRange&, const int32, const int32, const FSolverReal, TVec3<FSolverReal>&, Chaos::PMatrix<FSolverReal, 3, 3>&)>;
 	// (SourceParticles, TargetParticles, ConstraintIndex, LocalIndex, Dt, OutResidual, OutHessian)
-	using FGSCrossBodyResidualHessianFunc = TFunction<void(const FSolverParticles&, const FSolverParticles&, const int32, const int32, const FSolverReal, TVec3<FSolverReal>&, Chaos::PMatrix<FSolverReal, 3, 3>&)>;
+	using FGSCrossBodyResidualHessianFunc = TFunction<void(const FSolverParticlesRange&, const FSolverParticlesRange&, const int32, const int32, const FSolverReal, TVec3<FSolverReal>&, Chaos::PMatrix<FSolverReal, 3, 3>&)>;
 	// (ParticleIdx, Dt, OutHessian)
 	using FGSPerNodeHessianFunc = TFunction<void(const int32, const FSolverReal, Chaos::PMatrix<FSolverReal, 3, 3>&)>;
 	// (Particles, Dx, OutForceDifferential)
-	using FGSForceDifferentialFunc = TFunction<void(const FSolverParticles&, const TArray<TVec3<FSolverReal>>&, TArray<TVec3<FSolverReal>>&)>;
+	using FGSForceDifferentialFunc = TFunction<void(const FSolverParticlesRange&, const TArray<TVec3<FSolverReal>>&, TArray<TVec3<FSolverReal>>&)>;
 
 	class FGaussSeidelMainConstraints
 	{
 
 	public:
+		CHAOS_API FGaussSeidelMainConstraints(
+			const FSolverParticlesRange& InParticles,
+			const bool bDoQuasistaticsIn = false,
+			const bool bDoSORIn = true,
+			const FSolverReal InOmegaSOR = (FSolverReal)1.6,
+			const int32 ParallelMaxIn = 1000,
+			const FSolverReal MaxDxRatioIn = FSolverReal(1),
+			const FDeformableXPBDCorotatedParams& InParams = FDeformableXPBDCorotatedParams());
+
 		CHAOS_API FGaussSeidelMainConstraints(
 			const FSolverParticles& InParticles,
 			const bool bDoQuasistaticsIn = false,
@@ -49,6 +59,8 @@ namespace Chaos::Softs
 			const FSolverReal MaxDxRatioIn = FSolverReal(1),
 			const FDeformableXPBDCorotatedParams& InParams = FDeformableXPBDCorotatedParams());
 
+		CHAOS_API FGaussSeidelMainConstraints(const FSolverParticlesRange& InParticles, const FCollectionPropertyConstFacade& Property, const FSolverReal MaxDxRatioIn);
+
 		CHAOS_API FGaussSeidelMainConstraints(const FSolverParticles& InParticles, const FCollectionPropertyConstFacade& Property, const FSolverReal MaxDxRatioIn);
 
 		virtual ~FGaussSeidelMainConstraints() = default;
@@ -90,28 +102,35 @@ namespace Chaos::Softs
 
 		CHAOS_API int32 AddAddInternalForceDifferentialsRange(int32 NumConstraints);
 
-		CHAOS_API void AddStaticConstraints(const TArray<TArray<int32>>& ExtraConstraints, TArray<TArray<int32>>& ExtraIncidentElements, TArray<TArray<int32>>& ExtraIncidentElementsLocal);
+		CHAOS_API void AddStaticConstraints(const TArray<TArray<int32>>& ExtraConstraints, const TArray<TArray<int32>>& ExtraIncidentElements, const TArray<TArray<int32>>& ExtraIncidentElementsLocal);
 
-		CHAOS_API void AddTransientConstraints(const TArray<TArray<int32>>& ExtraConstraints, TArray<TArray<int32>>& ExtraIncidentElements, TArray<TArray<int32>>& ExtraIncidentElementsLocal, bool CheckIncidentElements = false);
+		CHAOS_API void AddTransientConstraints(const TArray<TArray<int32>>& ExtraConstraints, const TArray<TArray<int32>>& ExtraIncidentElements, const TArray<TArray<int32>>& ExtraIncidentElementsLocal, bool CheckIncidentElements = false);
 
-		CHAOS_API void AddDynamicConstraints(const TArray<TArray<int32>>& ExtraConstraints, TArray<TArray<int32>>& ExtraIncidentElements, TArray<TArray<int32>>& ExtraIncidentElementsLocal, bool CheckIncidentElements = false);
+		CHAOS_API void AddDynamicConstraints(const TArray<TArray<int32>>& ExtraConstraints, const TArray<TArray<int32>>& ExtraIncidentElements, const TArray<TArray<int32>>& ExtraIncidentElementsLocal, bool CheckIncidentElements = false);
 
 		CHAOS_API void ResetDynamicConstraints();
 		CHAOS_API void ResetTransientConstraints();
 		CHAOS_API void ResetTransientConstraintsAndRules();
 
+		CHAOS_API void Apply(FSolverParticlesRange& Particles, const FSolverReal Dt, const int32 MaxWriteIters = 10, const bool Write2File = false, const TPBDActiveView<FSolverParticles>* InParticleActiveView = nullptr);
 		CHAOS_API void Apply(FSolverParticles& Particles, const FSolverReal Dt, const int32 MaxWriteIters = 10, const bool Write2File = false, const TPBDActiveView<FSolverParticles>* InParticleActiveView = nullptr);
 
+		CHAOS_API void InitStaticColor(const FSolverParticlesRange& Particles, const TPBDActiveView<FSolverParticles>* InParticleActiveView = nullptr);
 		CHAOS_API void InitStaticColor(const FSolverParticles& Particles, const TPBDActiveView<FSolverParticles>* InParticleActiveView = nullptr);
 
+		CHAOS_API void InitTransientColor(const FSolverParticlesRange& Particles);
 		CHAOS_API void InitTransientColor(const FSolverParticles& Particles);
 
+		CHAOS_API void InitDynamicColor(const FSolverParticlesRange& Particles);
 		CHAOS_API void InitDynamicColor(const FSolverParticles& Particles);
 
+		CHAOS_API void Init(const FSolverReal Dt, const FSolverParticlesRange& Particles);
 		CHAOS_API void Init(const FSolverReal Dt, const FSolverParticles& Particles);
 
+		CHAOS_API TArray<TVec3<FSolverReal>> ComputeNewtonResiduals(const FSolverParticlesRange& Particles, const FSolverReal Dt, const bool Write2File = false, TArray<PMatrix<FSolverReal, 3, 3>>* AllParticleHessian = nullptr);
 		CHAOS_API TArray<TVec3<FSolverReal>> ComputeNewtonResiduals(const FSolverParticles& Particles, const FSolverReal Dt, const bool Write2File = false, TArray<PMatrix<FSolverReal, 3, 3>>* AllParticleHessian = nullptr);
 
+		CHAOS_API void ApplyCG(FSolverParticlesRange& Particles, const FSolverReal Dt);
 		CHAOS_API void ApplyCG(FSolverParticles& Particles, const FSolverReal Dt);
 
 		/* Adds external acceleration, e.g. gravity (0,0,-980) cm/s^2 */
@@ -126,7 +145,7 @@ namespace Chaos::Softs
 
 	private:
 
-		bool ApplySingleParticle(const int32 ParticleIdx, const FSolverReal Dt, FSolverParticles& Particles);
+		bool ApplySingleParticle(const int32 ParticleIdx, const FSolverReal Dt, FSolverParticlesRange& Particles);
 
 		void InitializeLambdas();
 
@@ -170,7 +189,7 @@ namespace Chaos::Softs
 		mutable TArray<Chaos::TVector<FSolverReal, 3>> xtilde;
 
 		//SOR variables:
-		TFunction<void(FSolverParticles&, int32)> AccelerationTechniquePerParticle;
+		TFunction<void(FSolverParticlesRange&, int32)> AccelerationTechniquePerParticle;
 		mutable TArray<Chaos::TVector<FSolverReal, 3>> X_k_1;
 		mutable TArray<Chaos::TVector<FSolverReal, 3>> X_k;
 		int32 CurrentIt = 0;
diff --git a/Engine/Source/Runtime/Experimental/Chaos/Public/Chaos/Deformable/GaussSeidelNeohookeanConstraints.h b/Engine/Source/Runtime/Experimental/Chaos/Public/Chaos/Deformable/GaussSeidelNeohookeanConstraints.h
--- a/Engine/Source/Runtime/Experimental/Chaos/Public/Chaos/Deformable/GaussSeidelNeohookeanConstraints.h
+++ b/Engine/Source/Runtime/Experimental/Chaos/Public/Chaos/Deformable/GaussSeidelNeohookeanConstraints.h
@@ -23,7 +23,6 @@ namespace Chaos::Softs
 		using GSBase::IncidentElementsLocal;
 		using GSBase::ParticleStartIndex;
 		using GSBase::LocalNewtonTol;
-		using GSBase::Particle2Incident;
 		using GSBase::bDoQuasistatics;
 		using GSBase::ComputeStress;
 		using GSBase::ComputeHessianHelper;
diff --git a/Engine/Source/Runtime/Experimental/Chaos/Public/Chaos/Deformable/GaussSeidelSphereRepulsionConstraints.h b/Engine/Source/Runtime/Experimental/Chaos/Public/Chaos/Deformable/GaussSeidelSphereRepulsionConstraints.h
--- a/Engine/Source/Runtime/Experimental/Chaos/Public/Chaos/Deformable/GaussSeidelSphereRepulsionConstraints.h
+++ b/Engine/Source/Runtime/Experimental/Chaos/Public/Chaos/Deformable/GaussSeidelSphereRepulsionConstraints.h
@@ -26,178 +26,44 @@ namespace Chaos::Softs
 	class FGaussSeidelParticleSphereRepulsionConstraints
 	{
 	public:
-		//TODO(Yushan): Add unit tests for Gauss Seidel Sphere Repulsion Constraints
-		FGaussSeidelParticleSphereRepulsionConstraints(FSolverReal InRadius, FSolverReal InStiffness, const FSolverParticles& InParticles, const FDeformableXPBDWeightedSpringConstraintParams& InParams): Radius(InRadius), Stiffness(InStiffness), DebugDrawParams(InParams)
-		{
-			const TArrayCollectionArray<FPAndInvM>& PAndInvM = InParticles.GetPAndInvM();
-			ReferencePositions.SetNum(PAndInvM.Num());
-			for (int32 i = 0; i < PAndInvM.Num(); ++i)
-			{
-				ReferencePositions[i] = PAndInvM[i].P;
-			}
-		}
-
-		virtual ~FGaussSeidelParticleSphereRepulsionConstraints()= default;
-
-		//Energy = k/2*(2r-d)^2
-		//Residual = de/dx = -force = -k*(2r-d)*dd/dx
-		//Hessian = de2/dx2 = k*dd/dx*dd/dx-k*(2r-d)*(-dd/dx*dd/dx^T+I)/d
-		void AddSphereRepulsionResidual(const FSolverParticles& InParticles, const int32 p, const FSolverReal Dt, TVec3<FSolverReal>& res)
-		{
-
-		}
-
-		void AddSphereRepulsionHessian(const int32 p, const FSolverReal Dt, Chaos::PMatrix<FSolverReal, 3, 3>& ParticleHessian)
-		{
-
-		}
-
-		void AddSphereRepulsionResidualAndHessian(const FSolverParticles& InParticles, const int32 ConstraintIndex, const int32 LocalIndex, const FSolverReal Dt, TVec3<FSolverReal>& ParticleResidual, Chaos::PMatrix<FSolverReal, 3, 3>& ParticleHessian)
-		{
-			Chaos::TVec3<FSolverReal> x0 = InParticles.P(Constraints[ConstraintIndex][0]);
-			Chaos::TVec3<FSolverReal> x1 = InParticles.P(Constraints[ConstraintIndex][1]);
-			Chaos::TVec3<FSolverReal> normal = (x1 - x0).GetSafeNormal();
-			normal = (LocalIndex ? normal : -normal); //dd/dx
-			FSolverReal dist = (x1 - x0).Size();
-			FSolverReal penetration = 2 * Radius - dist; //2r-d
+		CHAOS_API FGaussSeidelParticleSphereRepulsionConstraints(FSolverReal InRadius, FSolverReal InStiffness, const FSolverParticles& InParticles, const FDeformableXPBDWeightedSpringConstraintParams& InParams);
 
-			if (penetration > FSolverReal(0)) {
-				FSolverReal dist_inv = FSolverReal(1) / (dist+FSolverReal(1e-12));
-				Chaos::PMatrix<FSolverReal, 3, 3> outer = Chaos::PMatrix<FSolverReal, 3, 3>::OuterProduct(normal, normal);
-				Chaos::PMatrix<FSolverReal, 3, 3> A = (Chaos::PMatrix<FSolverReal, 3, 3>::Identity - outer) * dist_inv; //(-dd/dx*dd/dx^T+I)/d
-				ParticleHessian += Dt * Dt * ConstraintStiffness[ConstraintIndex] * (outer - penetration * A);
-				ParticleResidual += -Dt * Dt * penetration * ConstraintStiffness[ConstraintIndex] * normal;
-			}
-		}
+		virtual ~FGaussSeidelParticleSphereRepulsionConstraints() = default;
 
-		void VisualizeAllBindings(const FSolverParticles& InParticles, const FSolverReal Dt) const
-		{
-#if WITH_EDITOR
-			auto DoubleVert = [](Chaos::TVec3<FSolverReal> V) { return FVector3d(V.X, V.Y, V.Z); };
-			for (int32 i = 0; i < Constraints.Num(); i++)
-			{
-				FVector3d SourcePos = DoubleVert(InParticles.P(Constraints[i][0]));
-				FVector3d TargetPos = DoubleVert(InParticles.P(Constraints[i][1]));
+		CHAOS_API void AddSphereRepulsionResidualAndHessian(const FSolverParticlesRange& InParticles, const int32 ConstraintIndex, const int32 LocalIndex, const FSolverReal Dt, TVec3<FSolverReal>& ParticleResidual, Chaos::PMatrix<FSolverReal, 3, 3>& ParticleHessian);
 
-				float ParticleThickness = DebugDrawParams.DebugParticleWidth;
-				float LineThickness = DebugDrawParams.DebugLineWidth;
+		CHAOS_API void VisualizeAllBindings(const FSolverParticles& InParticles, const FSolverReal Dt) const;
 
-				Chaos::FDebugDrawQueue::GetInstance().DrawDebugPoint(SourcePos, FColor::Red, false, Dt, 0, ParticleThickness);
-				Chaos::FDebugDrawQueue::GetInstance().DrawDebugPoint(TargetPos, FColor::Red, false, Dt, 0, ParticleThickness);
-				Chaos::FDebugDrawQueue::GetInstance().DrawDebugLine(SourcePos, TargetPos, FColor::Green, false, Dt, 0, LineThickness);
-			}
-#endif
-		}
+		CHAOS_API void Init(const FSolverParticles& InParticles, const FSolverReal Dt) const;
 
-		void Init(const FSolverParticles& InParticles, const FSolverReal Dt) const
-		{
-			if (DebugDrawParams.bVisualizeBindings)
-			{
-				VisualizeAllBindings(InParticles, Dt);
-			}
+		CHAOS_API void UpdateSphereRepulsionConstraints(const FSolverParticles& Particles, const TArray<int32>& SurfaceVertices, const TArray<int32>& ComponentIndex);
 
-		}
+		CHAOS_API void ReturnSphereRepulsionConstraints(TArray<TArray<int32>>& ExtraConstraints, TArray<TArray<int32>>& ExtraIncidentElements, TArray<TArray<int32>>& ExtraIncidentElementsLocal);
 
-		void UpdateSphereRepulsionConstraints(const FSolverParticles& Particles, const TArray<int32>& SurfaceVertices, const TArray<int32>& ComponentIndex)
+	protected:
+		TArray<TVec2<int32>> Constraints;
+		FSolverReal Radius = FSolverReal(0);
+		FSolverReal Stiffness = FSolverReal(0);
+		TArray<FSolverReal> ConstraintStiffness;
+	private:
+		struct FSphereSpatialEntry
 		{
-			TRACE_CPUPROFILER_EVENT_SCOPE(STAT_GaussSeidelSphereRepulsionConstraints);
-			Constraints.Reset();
-			ConstraintStiffness.Reset();
-			if (SurfaceVertices.Num() == 0)
-			{
-				return;
-			}
-
-			// Build Spatial
-			TArray<FSphereSpatialEntry> Entries;
-			const FSolverReal Diameter = 2.f * Radius;
-			TConstArrayView<FSolverVec3> Points = Particles.XArray();
-
-			Entries.Reset(SurfaceVertices.Num());
-			for (int32 Index = 0; Index < SurfaceVertices.Num(); ++Index)
-			{
-				Entries.Add({ &Points, SurfaceVertices[Index]});
-			}
-			
-			TSpatialHashGridPoints<int32, FSolverReal> SpatialHash(Diameter);
-			SpatialHash.InitializePoints(Entries);
-
-			const FSolverReal DiamSq = FMath::Square(Diameter);
-			constexpr int32 CellRadius = 1; // We set the cell size of the spatial hash such that we only need to look 1 cell away to find proximities.
-			constexpr int32 MaxNumExpectedConnectionsPerParticle = 3;
-			const int32 MaxNumExpectedConnections = MaxNumExpectedConnectionsPerParticle * Entries.Num();
+			const TConstArrayView<FSolverVec3>* Points;
+			int32 Index;
 
-			Constraints = SpatialHash.FindAllSelfProximities(CellRadius, MaxNumExpectedConnections,
-				[this, &Particles, DiamSq, &ComponentIndex](const int32 i1, const int32 i2)
-				{
-					if (ComponentIndex[i1] == ComponentIndex[i2])
-					{
-						return false;
-					}
-					const FSolverReal CombinedMass = Particles.InvM(i1) + Particles.InvM(i2);
-					if (CombinedMass < (FSolverReal)1e-7)
-					{
-						return false;
-					}
-					if (FSolverVec3::DistSquared(this->ReferencePositions[i1], this->ReferencePositions[i2]) < DiamSq)
-					{
-						return false;
-					}
-					return true;
-				}
-			);
-
-			for (const TVec2<int32>& CollisionPair : Constraints)
+			FSolverVec3 X() const
 			{
-				FSolverReal TotalMass = FSolverReal(0);
-				if (Particles.InvM(CollisionPair[0]) > (FSolverReal)1e-7)
-				{
-					TotalMass += Particles.M(CollisionPair[0]);
-				}
-				if (Particles.InvM(CollisionPair[1]) > (FSolverReal)1e-7)
-				{
-					TotalMass += Particles.M(CollisionPair[1]);
-				}
-				ConstraintStiffness.Add(Stiffness*TotalMass/FSolverReal(2));
+				return (*Points)[Index];
 			}
-		}
 
-		void ReturnSphereRepulsionConstraints(TArray<TArray<int32>>& ExtraConstraints, TArray<TArray<int32>>& ExtraIncidentElements, TArray<TArray<int32>>& ExtraIncidentElementsLocal)
-		{
-			ExtraConstraints.Init(TArray<int32>(), Constraints.Num());
-			for (int32 i = 0; i < Constraints.Num(); i++)
+			template<typename TPayloadType>
+			int32 GetPayload(int32) const
 			{
-				ExtraConstraints[i].SetNum(2);
-				ExtraConstraints[i][0] = Constraints[i][0];
-				ExtraConstraints[i][1] = Constraints[i][1];
+				return Index;
 			}
-			ExtraIncidentElements = Chaos::Utilities::ComputeIncidentElements(ExtraConstraints, &ExtraIncidentElementsLocal);
-		}
-
-		protected:
-			TArray<TVec2<int32>> Constraints;
-			FSolverReal Radius = FSolverReal(0);
-			FSolverReal Stiffness = FSolverReal(0);
-			TArray<FSolverReal> ConstraintStiffness;
-		private:
-			struct FSphereSpatialEntry
-			{
-				const TConstArrayView<FSolverVec3>* Points;
-				int32 Index;
-
-				FSolverVec3 X() const
-				{
-					return (*Points)[Index];
-				}
-
-				template<typename TPayloadType>
-				int32 GetPayload(int32) const
-				{
-					return Index;
-				}
-			};
-			TArray<FSolverVec3> ReferencePositions;
-			FDeformableXPBDWeightedSpringConstraintParams DebugDrawParams;
+		};
+		TArray<FSolverVec3> ReferencePositions;
+		FDeformableXPBDWeightedSpringConstraintParams DebugDrawParams;
 	};
 
 	template <typename T, typename ParticleType>
diff --git a/Engine/Source/Runtime/Experimental/Chaos/Public/Chaos/Deformable/GaussSeidelUnilateralTetConstraints.h b/Engine/Source/Runtime/Experimental/Chaos/Public/Chaos/Deformable/GaussSeidelUnilateralTetConstraints.h
--- a/Engine/Source/Runtime/Experimental/Chaos/Public/Chaos/Deformable/GaussSeidelUnilateralTetConstraints.h
+++ b/Engine/Source/Runtime/Experimental/Chaos/Public/Chaos/Deformable/GaussSeidelUnilateralTetConstraints.h
@@ -10,89 +10,20 @@ namespace Chaos::Softs
 	// Tetrahdral volume constraint for position-based nonlinear Gauss Seidel solver [Chen et al., 2024]
 	class FGaussSeidelUnilateralVolumeConstraints
 	{
-	public: 
-		FGaussSeidelUnilateralVolumeConstraints(
+	public:
+		CHAOS_API FGaussSeidelUnilateralVolumeConstraints(
 			const FSolverParticles& Particles,
 			TArray<TVector<int32, 4>>&& InConstraints,
-			TArray<FSolverReal> InStiffnessArray)
-			: Constraints(InConstraints), 
-			StiffnessArray(InStiffnessArray)
-		{
-			IncidentElements = Chaos::Utilities::ComputeIncidentElements(Constraints, &IncidentElementsLocal);
-			Volumes.SetNumUninitialized(Constraints.Num());
+			TArray<FSolverReal> InStiffnessArray);
 
-			for (int32 Idx = 0; Idx < Constraints.Num(); ++Idx)
-			{
-				const TVec4<int32>& Constraint = Constraints[Idx];
-				Volumes[Idx] = ComputeVolume(
-					Particles.X(Constraint[0]),
-					Particles.X(Constraint[1]),
-					Particles.X(Constraint[2]),
-					Particles.X(Constraint[3]));
-				StiffnessArray[Idx] /= Volumes[Idx];
-			}
-		};
+		CHAOS_API void AddEnergy(const FSolverParticles& InParticles, const int32 ConstraintIndex, const FSolverReal Dt, FSolverReal& Energy) const;
 
-		static FSolverReal ComputeVolume(const TVec3<FSolverReal>& P1, const TVec3<FSolverReal>& P2, const TVec3<FSolverReal>& P3, const TVec3<FSolverReal>& P4)
-		{
-			const TVec3<FSolverReal> P2P1 = P2 - P1;
-			const TVec3<FSolverReal> P4P1 = P4 - P1;
-			const TVec3<FSolverReal> P3P1 = P3 - P1;
-			return TVec3<FSolverReal>::DotProduct(TVec3<FSolverReal>::CrossProduct(P2P1, P3P1), P4P1) / (FSolverReal)6.;
-		}
+		CHAOS_API void AddResidualAndHessian(const FSolverParticles& InParticles, const int32 ConstraintIndex, const int32 LocalIndex, const FSolverReal Dt, TVec3<FSolverReal>& ParticleResidual, Chaos::PMatrix<FSolverReal, 3, 3>& ParticleHessian) const;
 
-		void AddEnergy(const FSolverParticles& InParticles, const int32 ConstraintIndex, const FSolverReal Dt, FSolverReal& Energy) const
-		{
-			const TVec4<int32>& Constraint = Constraints[ConstraintIndex];
-			FSolverReal CurrentVol = ComputeVolume(
-				InParticles.P(Constraint[0]),
-				InParticles.P(Constraint[1]),
-				InParticles.P(Constraint[2]),
-				InParticles.P(Constraint[3]));
-			Energy += Dt * Dt * StiffnessArray[ConstraintIndex] / (FSolverReal)2. * (CurrentVol - Volumes[ConstraintIndex]) * (CurrentVol - Volumes[ConstraintIndex]);
-		}
-
-		void AddResidualAndHessian(const FSolverParticles& InParticles, const int32 ConstraintIndex, const int32 LocalIndex, const FSolverReal Dt, TVec3<FSolverReal>& ParticleResidual, Chaos::PMatrix<FSolverReal, 3, 3>& ParticleHessian) const
-		{
-			const TVec4<int32>& Constraint = Constraints[ConstraintIndex];
-			const int32 Index1 = Constraint[0];
-			const int32 Index2 = Constraint[1];
-			const int32 Index3 = Constraint[2];
-			const int32 Index4 = Constraint[3];
-			
-			const TVec3<FSolverReal>& P1 = InParticles.P(Index1);
-			const TVec3<FSolverReal>& P2 = InParticles.P(Index2);
-			const TVec3<FSolverReal>& P3 = InParticles.P(Index3);
-			const TVec3<FSolverReal>& P4 = InParticles.P(Index4);
+		CHAOS_API TArray<TArray<int32>> GetStaticConstraintArrays(const TArray<TArray<int32>>*& OutIncidentElements, const TArray<TArray<int32>>*& OutIncidentElementsLocal) const;
 
-			TVec4<TVec3<FSolverReal>> Grads;
-			const TVec3<FSolverReal> P2P1 = P2 - P1;
-			const TVec3<FSolverReal> P4P1 = P4 - P1;
-			const TVec3<FSolverReal> P3P1 = P3 - P1;
-			Grads[1] = TVec3<FSolverReal>::CrossProduct(P3P1, P4P1) / (FSolverReal)6.;
-			Grads[2] = TVec3<FSolverReal>::CrossProduct(P4P1, P2P1) / (FSolverReal)6.;
-			Grads[3] = TVec3<FSolverReal>::CrossProduct(P2P1, P3P1) / (FSolverReal)6.;
-			Grads[0] = -(Grads[1] + Grads[2] + Grads[3]);
-
-			const FSolverReal Volume = ComputeVolume(P1, P2, P3, P4);
-			const FSolverReal C_Hessian = StiffnessArray[ConstraintIndex] * Dt * Dt;
-			const FSolverReal C_Residual = (Volume - Volumes[ConstraintIndex]) * C_Hessian;
-			ParticleResidual += C_Residual * Grads[LocalIndex];
-			ParticleHessian += C_Hessian * Chaos::PMatrix<FSolverReal, 3, 3>::OuterProduct(Grads[LocalIndex], Grads[LocalIndex]);
-		}
-
-		TArray<TArray<int32>> GetStaticConstraintArrays(TArray<TArray<int32>>& InIncidentElements, TArray<TArray<int32>>& InIncidentElementsLocal) const
-		{
-			InIncidentElements = IncidentElements;
-			InIncidentElementsLocal = IncidentElementsLocal;
-			TArray<TArray<int32>> NestedConstraints;
-			NestedConstraints.SetNum(Constraints.Num());
-			for (int32 Idx = 0; Idx < NestedConstraints.Num(); ++Idx)
-			{
-				NestedConstraints[Idx] = {Constraints[Idx][0], Constraints[Idx][1], Constraints[Idx][2], Constraints[Idx][3]};
-			}
-			return NestedConstraints;
-		}
+		UE_DEPRECATED(5.8, "Use GetStaticConstraintArrays with pointer output parameters instead.")
+		CHAOS_API TArray<TArray<int32>> GetStaticConstraintArrays(TArray<TArray<int32>>& IncidentElements, TArray<TArray<int32>>& IncidentElementsLocal) const;
 
 		int32 NumConstraints() const
 		{
@@ -100,6 +31,8 @@ namespace Chaos::Softs
 		}
 
 	private:
+		static FSolverReal ComputeVolume(const TVec3<FSolverReal>& P1, const TVec3<FSolverReal>& P2, const TVec3<FSolverReal>& P3, const TVec3<FSolverReal>& P4);
+
 		TArray<TVector<int32, 4>> Constraints;
 		TArray<FSolverReal> Volumes;
 		TArray<TArray<int32>> IncidentElements;
diff --git a/Engine/Source/Runtime/Experimental/Chaos/Public/Chaos/Deformable/GaussSeidelWeakConstraints.h b/Engine/Source/Runtime/Experimental/Chaos/Public/Chaos/Deformable/GaussSeidelWeakConstraints.h
--- a/Engine/Source/Runtime/Experimental/Chaos/Public/Chaos/Deformable/GaussSeidelWeakConstraints.h
+++ b/Engine/Source/Runtime/Experimental/Chaos/Public/Chaos/Deformable/GaussSeidelWeakConstraints.h
@@ -44,8 +44,8 @@ namespace Chaos::Softs
 		void AddConstraints(const int32 Num) { AddElementsHelper(Num); }
 		void RemoveConstraint(const int32 Idx) { RemoveAtSwapHelper(Idx); }
 
-		void SetSingleConstraint(const FGaussSeidelSpringConstraintSingleData& SingleData, const int32 ConstraintIndex);
-		void AddSingleConstraint(const FGaussSeidelSpringConstraintSingleData& SingleData);
+		CHAOS_API void SetSingleConstraint(const FGaussSeidelSpringConstraintSingleData& SingleData, const int32 ConstraintIndex);
+		CHAOS_API void AddSingleConstraint(const FGaussSeidelSpringConstraintSingleData& SingleData);
 
 		int32 Size() const { return static_cast<int32>(MSize); }
 		void Resize(const int32 Num) { ResizeHelper(Num); }
@@ -114,6 +114,18 @@ namespace Chaos::Softs
 			const FDeformableXPBDWeightedSpringConstraintParams& InParams
 		);
 
+		// Scalar-bool overload: applies the same bIsAnisotropic and bIsZeroRestLength to all constraints
+		CHAOS_API FGaussSeidelSpringConstraints(
+			const TArray<TArray<int32>>& InIndices,
+			const TArray<TArray<FSolverReal>>& InWeights,
+			const TArray<FSolverReal>& InStiffness,
+			const TArray<TArray<int32>>& InSecondIndices,
+			const TArray<TArray<FSolverReal>>& InSecondWeights,
+			const FDeformableXPBDWeightedSpringConstraintParams& InParams,
+			bool bIsAnisotropic,
+			bool bIsZeroRestLength
+		);
+
 		UE_DEPRECATED(5.8, "Replaced with constructor with IsAnisotropic and IsZeroRestLength.")
 		CHAOS_API FGaussSeidelSpringConstraints(
 			const TArray<TArray<int32>>& InIndices,
@@ -156,10 +168,14 @@ namespace Chaos::Softs
 		template<typename SpatialAccelerator>
 		void CollisionDetectionSpatialHashInComponent(const FSolverParticles& Particles, const TArray<int32>& SurfaceVertices, const FTriangleMesh& TriangleMesh, const TMap<int32, TSet<int32>>& ExcludeMap, const SpatialAccelerator& Spatial, float DetectRadius = 0.f, float PositionTargetStiffness = 10000.f, bool UseAnisotropicSpring = true);
 
-		void ComputeCollisionWCDataSimplified(TArray<TArray<int32>>& ExtraConstraints, TArray<TArray<int32>>& ExtraWCIncidentElements, TArray<TArray<int32>>& ExtraWCIncidentElementsLocal);
+		CHAOS_API void ComputeCollisionWCDataSimplified(TArray<TArray<int32>>& ExtraConstraints, TArray<TArray<int32>>& ExtraWCIncidentElements, TArray<TArray<int32>>& ExtraWCIncidentElementsLocal);
+
+		CHAOS_API const TArray<TArray<int32>>& GetStaticConstraintArrays(const TArray<TArray<int32>>*& OutIncidentElements, const TArray<TArray<int32>>*& OutIncidentElementsLocal) const;
 
+		UE_DEPRECATED(5.8, "Use GetStaticConstraintArrays with pointer output parameters instead.")
 		CHAOS_API const TArray<TArray<int32>>& GetStaticConstraintArrays(TArray<TArray<int32>>& IncidentElements, TArray<TArray<int32>>& IncidentElementsLocal) const;
-		TArray<TArray<int32>> GetDynamicConstraintArrays(TArray<TArray<int32>>& IncidentElements, TArray<TArray<int32>>& IncidentElementsLocal) const;
+
+		CHAOS_API TArray<TArray<int32>> GetDynamicConstraintArrays(TArray<TArray<int32>>& IncidentElements, TArray<TArray<int32>>& IncidentElementsLocal) const;
 
 		UE_DEPRECATED(5.8, "Replaced with AddWCResidual for more general cases.")
 		void AddZeroRestLengthWCResidualAndHessian(const FSolverParticles& InParticles, const int32 ConstraintIndex, const int32 LocalIndex, const FSolverReal Dt, TVec3<FSolverReal>& ParticleResidual, Chaos::PMatrix<FSolverReal, 3, 3>& ParticleHessian) const;
@@ -168,10 +184,7 @@ namespace Chaos::Softs
 
 		CHAOS_API void AddWCResidual(const FSolverParticles& InParticles, const int32 ConstraintIndex, const int32 LocalIndex, const FSolverReal Dt, TVec3<FSolverReal>& ParticleResidual, Chaos::PMatrix<FSolverReal, 3, 3>& ParticleHessian) const;
 
-		int32 GetInitialWCSize() const
-		{
-			return InitialWCSize;
-		}
+		int32 GetInitialWCSize() const { return InitialWCSize; }
 		const FGaussSeidelSpringConstraintData& GetConstraintsData() const { return ConstraintsData; }
 
 	protected:
diff --git a/Engine/Source/Runtime/Experimental/Chaos/Public/Chaos/GraphColoring.h b/Engine/Source/Runtime/Experimental/Chaos/Public/Chaos/GraphColoring.h
--- a/Engine/Source/Runtime/Experimental/Chaos/Public/Chaos/GraphColoring.h
+++ b/Engine/Source/Runtime/Experimental/Chaos/Public/Chaos/GraphColoring.h
@@ -7,6 +7,7 @@
 #include "Chaos/Vector.h"
 #include "Chaos/UniformGrid.h"
 #include "Chaos/PBDActiveView.h"
+#include "Chaos/SoftsSolverParticlesRange.h"
 
 namespace Chaos
 {
@@ -60,13 +61,15 @@ void ComputeGridBasedGraphSubColoringPointer(const TArray<TArray<int32>>& Elemen
 template<typename T>
 CHAOS_API void ComputeGridBasedGraphSubColoringPointer(const TArray<TArray<int32>>& ElementsPerColor, const TMPMGrid<T>& Grid, const int32 GridSize, TUniquePtr<TArray<TArray<int32>>>& PreviousColoring, const TArray<TArray<int32>>& ConstraintsNodesSet, TArray<TArray<TArray<int32>>>& ElementsPerSubColors);
 
-// Both Indices and SecondIndices point to the same InParticles
 template<typename T>
+UE_DEPRECATED(5.8, "Use the NumParticles overload instead of passing TDynamicParticles.")
 CHAOS_API void ComputeWeakConstraintsColoring(const TArray<TArray<int32>>& Indices, const TArray<TArray<int32>>& SecondIndices, const Chaos::TDynamicParticles<T, 3>& InParticles, TArray<TArray<int32>>& ConstraintsPerColor);
 
-// if SecondIndices point to InSecondParticles
-template<typename T>
-CHAOS_API void ComputeWeakConstraintsColoring(const TArray<TArray<int32>>& Indices, const TArray<TArray<int32>>& SecondIndices, const Chaos::TDynamicParticles<T, 3>& InParticles, const Chaos::TDynamicParticles<T, 3>& InSecondParticles, TArray<TArray<int32>>& ConstraintsPerColor);
+// Non-template overloads taking particle count directly (for use with FSolverParticlesRange)
+// Use this version if Indices and SecondIndices point to the same particles
+CHAOS_API void ComputeWeakConstraintsColoring(const TArray<TArray<int32>>& Indices, const TArray<TArray<int32>>& SecondIndices, int32 NumParticles, TArray<TArray<int32>>& ConstraintsPerColor);
+// Use this version if Indices and SecondIndices point to different particles, correspondingly
+CHAOS_API void ComputeWeakConstraintsColoring(const TArray<TArray<int32>>& Indices, const TArray<TArray<int32>>& SecondIndices, int32 NumParticles, int32 NumSecondParticles, TArray<TArray<int32>>& ConstraintsPerColor);
 
 template<typename T>
 CHAOS_API TArray<TArray<int32>> ComputeNodalColoring(const TArray<TVec4<int32>>& Graph, const Chaos::TDynamicParticles<T, 3>& InParticles, const int32 GraphParticlesStart, const int32 GraphParticlesEnd, const TArray<TArray<int32>>& IncidentElements, const TArray<TArray<int32>>& IncidentElementsLocalIndex);
@@ -80,5 +83,15 @@ CHAOS_API void ComputeExtraNodalColoring(const TArray<TArray<int32>>& Graph, con
 template<typename T>
 CHAOS_API void ComputeExtraNodalColoring(const TArray<TArray<int32>>& StaticGraph, const TArray<TArray<int32>>& DynamicGraph, const TArray<TArray<int32>>& ExtraGraph, const Chaos::TDynamicParticles<T, 3>& InParticles, const TArray<TArray<int32>>& StaticIncidentElements, const TArray<TArray<int32>>& DynamicIncidentElements, const TArray<TArray<int32>>& ExtraIncidentElements, TArray<int32>& ParticleColors, TArray<TArray<int32>>& ParticlesPerColor);
 
+// FSolverParticlesRange versions
+CHAOS_API TArray<TArray<int32>> ComputeNodalColoring(const TArray<TVec4<int32>>& Graph, const Softs::FSolverParticlesRange& InParticles, const int32 GraphParticlesStart, const int32 GraphParticlesEnd, const TArray<TArray<int32>>& IncidentElements, const TArray<TArray<int32>>& IncidentElementsLocalIndex);
+
+template<typename ParticleType>
+CHAOS_API TArray<TArray<int32>> ComputeNodalColoring(const TArray<TArray<int32>>& Graph, const Softs::FSolverParticlesRange& InParticles, const int32 GraphParticlesStart, const int32 GraphParticlesEnd, const TArray<TArray<int32>>& IncidentElements, const TArray<TArray<int32>>& IncidentElementsLocalIndex, const Chaos::TPBDActiveView<ParticleType>* InParticleActiveView = nullptr, TArray<int32>* ParticleColorsOut = nullptr);
+
+CHAOS_API void ComputeExtraNodalColoring(const TArray<TArray<int32>>& Graph, const TArray<TArray<int32>>& ExtraGraph, const Softs::FSolverParticlesRange& InParticles, const TArray<TArray<int32>>& IncidentElements, const TArray<TArray<int32>>& ExtraIncidentElements, TArray<int32>& ParticleColors, TArray<TArray<int32>>& ParticlesPerColor);
+
+CHAOS_API void ComputeExtraNodalColoring(const TArray<TArray<int32>>& StaticGraph, const TArray<TArray<int32>>& DynamicGraph, const TArray<TArray<int32>>& ExtraGraph, const Softs::FSolverParticlesRange& InParticles, const TArray<TArray<int32>>& StaticIncidentElements, const TArray<TArray<int32>>& DynamicIncidentElements, const TArray<TArray<int32>>& ExtraIncidentElements, TArray<int32>& ParticleColors, TArray<TArray<int32>>& ParticlesPerColor);
+
 
 }
diff --git a/Engine/Source/Runtime/Experimental/Chaos/Public/Chaos/SoftsSolverParticlesRange.h b/Engine/Source/Runtime/Experimental/Chaos/Public/Chaos/SoftsSolverParticlesRange.h
--- a/Engine/Source/Runtime/Experimental/Chaos/Public/Chaos/SoftsSolverParticlesRange.h
+++ b/Engine/Source/Runtime/Experimental/Chaos/Public/Chaos/SoftsSolverParticlesRange.h
@@ -22,6 +22,14 @@ public:
 		: TParticlesRange<FSolverParticles>(InParticles, InOffset, InRange)
 	{}
 
+	explicit FSolverParticlesRange(FSolverParticles& InParticles)
+		: TParticlesRange<FSolverParticles>(&InParticles, 0, InParticles.Size())
+	{}
+
+	explicit FSolverParticlesRange(const FSolverParticles& InParticles)
+		: TParticlesRange<FSolverParticles>(const_cast<FSolverParticles*>(&InParticles), 0, InParticles.Size())
+	{}
+
 	FSolverParticlesRange(TParticlesRange<FSolverParticles>&& Other)
 		: TParticlesRange<FSolverParticles>(MoveTemp(Other))
 	{}
