ChaosDeformable: Constraint hierarchy refactor and FSolverParticlesRange migration

diff --git a/Engine/Source/Runtime/Experimental/Chaos/Public/Chaos/XPBDAnisotropicBendingConstraints.h b/Engine/Source/Runtime/Experimental/Chaos/Public/Chaos/XPBDAnisotropicBendingConstraints.h
--- a/Engine/Source/Runtime/Experimental/Chaos/Public/Chaos/XPBDAnisotropicBendingConstraints.h
+++ b/Engine/Source/Runtime/Experimental/Chaos/Public/Chaos/XPBDAnisotropicBendingConstraints.h
@@ -89,24 +89,26 @@ public:
 	const TArray<int32>& GetConstraintsPerColorStartIndex() const { return ConstraintsPerColorStartIndex; }
 	const TArray<FSolverVec3>& GetWarpWeftBiasBaseMultipliers() const { return WarpWeftBiasBaseMultipliers; }
 
+	CHAOS_API void AddAnisotropicBendingResidualAndHessian(const FSolverParticlesRange& Particles, const int32 ConstraintIndex, const int32 ConstraintIndexLocal, const FSolverReal Dt, TVec3<FSolverReal>& ParticleResidual, Chaos::PMatrix<FSolverReal, 3, 3>& ParticleHessian);
 	CHAOS_API void AddAnisotropicBendingResidualAndHessian(const FSolverParticles& Particles, const int32 ConstraintIndex, const int32 ConstraintIndexLocal, const FSolverReal Dt, TVec3<FSolverReal>& ParticleResidual, Chaos::PMatrix<FSolverReal, 3, 3>& ParticleHessian);
 
 	TArray<TArray<int32>> GetConstraintsArray() const
 	{
 		TArray<TArray<int32>> ConstraintsArray;
 		ConstraintsArray.SetNum(Constraints.Num());
-		for (int32 i = 0; i < Constraints.Num(); i++)
+		for (int32 ConstraintIdx = 0; ConstraintIdx < Constraints.Num(); ConstraintIdx++)
 		{
-			ConstraintsArray[i].SetNum(4);
-			for (int32 j = 0; j < 4; j++)
+			ConstraintsArray[ConstraintIdx].SetNum(4);
+			for (int32 LocalIdx = 0; LocalIdx < 4; LocalIdx++)
 			{
-				ConstraintsArray[i][j] = Constraints[i][j];
+				ConstraintsArray[ConstraintIdx][LocalIdx] = Constraints[ConstraintIdx][LocalIdx];
 			}
 		}
 		return ConstraintsArray;
 	}
 
-	CHAOS_API void AddInternalForceDifferential(const FSolverParticles& InParticles, const TArray<TVector<FSolverReal, 3>>& DeltaParticles, TArray<TVector<FSolverReal, 3>>& ndf);
+	CHAOS_API void AddInternalForceDifferential(const FSolverParticlesRange& InParticles, const TArray<TVector<FSolverReal, 3>>& DeltaParticles, TArray<TVector<FSolverReal, 3>>& Ndf);
+	CHAOS_API void AddInternalForceDifferential(const FSolverParticles& InParticles, const TArray<TVector<FSolverReal, 3>>& DeltaParticles, TArray<TVector<FSolverReal, 3>>& Ndf);
 
 	void ResetRestLengths(const TConstArrayView<FSolverVec3>& Positions,
 		const FCollectionPropertyConstFacade& PropertyCollection,
diff --git a/Engine/Source/Runtime/Experimental/Chaos/Public/Chaos/XPBDBendingConstraints.h b/Engine/Source/Runtime/Experimental/Chaos/Public/Chaos/XPBDBendingConstraints.h
--- a/Engine/Source/Runtime/Experimental/Chaos/Public/Chaos/XPBDBendingConstraints.h
+++ b/Engine/Source/Runtime/Experimental/Chaos/Public/Chaos/XPBDBendingConstraints.h
@@ -304,24 +304,26 @@ public:
 
 	const TArray<int32>& GetConstraintsPerColorStartIndex() const { return ConstraintsPerColorStartIndex; }
 
+	CHAOS_API void AddBendingResidualAndHessian(const FSolverParticlesRange& Particles, const int32 ConstraintIndex, const int32 ConstraintIndexLocal, const FSolverReal Dt, TVec3<FSolverReal>& ParticleResidual, Chaos::PMatrix<FSolverReal, 3, 3>& ParticleHessian);
 	CHAOS_API void AddBendingResidualAndHessian(const FSolverParticles& Particles, const int32 ConstraintIndex, const int32 ConstraintIndexLocal, const FSolverReal Dt, TVec3<FSolverReal>& ParticleResidual, Chaos::PMatrix<FSolverReal, 3, 3>& ParticleHessian);
 
 	TArray<TArray<int32>> GetConstraintsArray() const
 	{
 		TArray<TArray<int32>> ConstraintsArray;
 		ConstraintsArray.SetNum(Constraints.Num());
-		for (int32 i = 0; i < Constraints.Num(); i++)
+		for (int32 ConstraintIdx = 0; ConstraintIdx < Constraints.Num(); ConstraintIdx++)
 		{
-			ConstraintsArray[i].SetNum(4);
-			for (int32 j = 0; j < 4; j++)
+			ConstraintsArray[ConstraintIdx].SetNum(4);
+			for (int32 LocalIdx = 0; LocalIdx < 4; LocalIdx++)
 			{
-				ConstraintsArray[i][j] = Constraints[i][j];
+				ConstraintsArray[ConstraintIdx][LocalIdx] = Constraints[ConstraintIdx][LocalIdx];
 			}
 		}
 		return ConstraintsArray;
 	}
 
-	CHAOS_API void AddInternalForceDifferential(const FSolverParticles& InParticles, const TArray<TVector<FSolverReal, 3>>& DeltaParticles, TArray<TVector<FSolverReal, 3>>& ndf);
+	CHAOS_API void AddInternalForceDifferential(const FSolverParticlesRange& InParticles, const TArray<TVector<FSolverReal, 3>>& DeltaParticles, TArray<TVector<FSolverReal, 3>>& Ndf);
+	CHAOS_API void AddInternalForceDifferential(const FSolverParticles& InParticles, const TArray<TVector<FSolverReal, 3>>& DeltaParticles, TArray<TVector<FSolverReal, 3>>& Ndf);
 
 	CHAOS_API FSolverReal ComputeTotalEnergy(const FSolverParticles& InParticles, const FSolverReal ExplicitStiffness = -1.f);
 
diff --git a/Engine/Source/Runtime/Experimental/Chaos/Public/Chaos/XPBDCorotatedConstraints.h b/Engine/Source/Runtime/Experimental/Chaos/Public/Chaos/XPBDCorotatedConstraints.h
--- a/Engine/Source/Runtime/Experimental/Chaos/Public/Chaos/XPBDCorotatedConstraints.h
+++ b/Engine/Source/Runtime/Experimental/Chaos/Public/Chaos/XPBDCorotatedConstraints.h
@@ -19,96 +19,21 @@ namespace Chaos::Softs
 	{
 
 	public:
-		//this one only accepts tetmesh input and mesh
-		FXPBDCorotatedTetrahedralConstraints(
+		CHAOS_API FXPBDCorotatedTetrahedralConstraints(
 			const FSolverParticles& InParticles,
 			const TArray<TVector<int32, 4>>& InMesh,
 			const bool bRecordMetricIn = true,
 			const FSolverReal& EMesh = (FSolverReal)10.0,
-			const FSolverReal& NuMesh = (FSolverReal).3
-			)
-			: bRecordMetric(bRecordMetricIn), MeshConstraints(InMesh)
-		{	
-
-			LambdaArray.Init((FSolverReal)0., 2 * MeshConstraints.Num());
-			DmInverse.Init((FSolverReal)0., 9 * MeshConstraints.Num());
-			Measure.Init((FSolverReal)0.,  MeshConstraints.Num());
-			Lambda = EMesh * NuMesh / (((FSolverReal)1. + NuMesh) * ((FSolverReal)1. - (FSolverReal)2. * NuMesh));
-			Mu = EMesh / ((FSolverReal)2. * ((FSolverReal)1. + NuMesh));
-			for (int e = 0; e < InMesh.Num(); e++)
-			{
-				PMatrix<FSolverReal, 3, 3> Dm = DsInit(e, InParticles);
-				PMatrix<FSolverReal, 3, 3> DmInv = Dm.Inverse();
-				for (int r = 0; r < 3; r++) {
-					for (int c = 0; c < 3; c++) {
-						DmInverse[(3 * 3) * e + 3 * r + c] = DmInv.GetAt(r, c);
-					}
-				}
-
-				Measure[e] = Dm.Determinant() / (FSolverReal)6.;
-
-				//part of preprocessing: if inverted element is found, 
-				//invert it so that the measure is positive
-				if (Measure[e] < (FSolverReal)0.)
-				{
-
-					Measure[e] = -Measure[e];
-				}
+			const FSolverReal& NuMesh = (FSolverReal).3);
 
-			}
-			DmInverseSave = DmInverse;
-			InitColor(InParticles);
-		}
-
-		FXPBDCorotatedTetrahedralConstraints(
+		CHAOS_API FXPBDCorotatedTetrahedralConstraints(
 			const FSolverParticles& InParticles,
 			const TArray<TVector<int32, 4>>& InMesh,
 			const TArray<FSolverReal>& EMeshArray,
 			const FSolverReal& NuMesh = (FSolverReal).3,
-			const bool bRecordMetricIn = false
-		)
-			: bRecordMetric(bRecordMetricIn), MeshConstraints(InMesh)
-		{
-			ensureMsgf(EMeshArray.Num() == InMesh.Num(), TEXT("Input Young Modulus Array Size is wrong"));
-			LambdaArray.Init((FSolverReal)0., 2 * MeshConstraints.Num());
-			DmInverse.Init((FSolverReal)0., 9 * MeshConstraints.Num());
-			Measure.Init((FSolverReal)0., MeshConstraints.Num());
-			LambdaElementArray.Init((FSolverReal)0., MeshConstraints.Num());
-			MuElementArray.Init((FSolverReal)0., MeshConstraints.Num());
-
-			for (int e = 0; e < InMesh.Num(); e++)
-			{
-				for (int32 j = 0; j < 4; j++)
-				{
-					ensure(MeshConstraints[e][j] > -1 && MeshConstraints[e][j] < int32(InParticles.Size()));
-				}
-			}
-			
-			for (int e = 0; e < InMesh.Num(); e++)
-			{
-				LambdaElementArray[e] = EMeshArray[e] * NuMesh / (((FSolverReal)1. + NuMesh) * ((FSolverReal)1. - (FSolverReal)2. * NuMesh));
-				MuElementArray[e] = EMeshArray[e] / ((FSolverReal)2. * ((FSolverReal)1. + NuMesh));
-
-				PMatrix<FSolverReal, 3, 3> Dm = DsInit(e, InParticles);
-				PMatrix<FSolverReal, 3, 3> DmInv = Dm.Inverse();
-				for (int r = 0; r < 3; r++) {
-					for (int c = 0; c < 3; c++) {
-						DmInverse[(3 * 3) * e + 3 * r + c] = DmInv.GetAt(r, c);
-					}
-				}
-
-				Measure[e] = Dm.Determinant() / (FSolverReal)6.;
-
-				if (Measure[e] < (FSolverReal)0.)
-				{
-					Measure[e] = -Measure[e];
-				}
-			}
-			DmInverseSave = DmInverse;
-			InitColor(InParticles);
-		}
+			const bool bRecordMetricIn = false);
 
-		FXPBDCorotatedTetrahedralConstraints(
+		CHAOS_API FXPBDCorotatedTetrahedralConstraints(
 			const FSolverParticlesRange& InParticles,
 			const TArray<TVector<int32, 4>>& InMesh,
 			const TArray<FSolverReal>& EMeshArray,
@@ -117,56 +42,7 @@ namespace Chaos::Softs
 			const FDeformableXPBDCorotatedParams& InParams,
 			const FSolverReal& NuMesh = (FSolverReal).3,
 			const bool bRecordMetricIn = false,
-			const bool bDoColoring = true
-		)
-			: CorotatedParams(InParams), AlphaJArray(MoveTemp(AlphaJMeshArray)), bRecordMetric(bRecordMetricIn), MeshConstraints(InMesh)
-		{
-			ensureMsgf(EMeshArray.Num() == InMesh.Num(), TEXT("Input Young Modulus Array Size is wrong"));
-			LambdaArray.Init((FSolverReal)0., 2 * MeshConstraints.Num());
-			DmInverse.Init((FSolverReal)0., 9 * MeshConstraints.Num());
-			Measure.Init((FSolverReal)0., MeshConstraints.Num());
-			LambdaElementArray.Init((FSolverReal)0., MeshConstraints.Num());
-			MuElementArray.Init((FSolverReal)0., MeshConstraints.Num());
-
-			for (int e = 0; e < InMesh.Num(); e++)
-			{
-				for (int32 j = 0; j < 4; j++)
-				{
-					ensure(MeshConstraints[e][j] > -1 && MeshConstraints[e][j] < int32(InParticles.Size()));
-
-					if (InParticles.InvM(MeshConstraints[e][j]) == (FSolverReal)0.)
-					{
-						AlphaJArray[e] = (FSolverReal)1.;
-					}
-				}
-
-			}
-
-			for (int e = 0; e < InMesh.Num(); e++)
-			{
-				LambdaElementArray[e] = EMeshArray[e] * NuMeshArray[e] / (((FSolverReal)1. + NuMeshArray[e]) * ((FSolverReal)1. - (FSolverReal)2. * NuMeshArray[e]));
-				MuElementArray[e] = EMeshArray[e] / ((FSolverReal)2. * ((FSolverReal)1. + NuMeshArray[e]));
-
-				PMatrix<FSolverReal, 3, 3> Dm = DsInit(e, InParticles);
-				PMatrix<FSolverReal, 3, 3> DmInv = Dm.Inverse();
-				for (int r = 0; r < 3; r++) {
-					for (int c = 0; c < 3; c++) {
-						DmInverse[(3 * 3) * e + 3 * r + c] = DmInv.GetAt(r, c);
-					}
-				}
-
-				Measure[e] = Dm.Determinant() / (FSolverReal)6.;
-
-				if (Measure[e] < (FSolverReal)0.)
-				{
-					Measure[e] = -Measure[e];
-				}
-			}
-			DmInverseSave = DmInverse;
-			if (bDoColoring) {
-				InitColor(InParticles);
-			}
-		}
+			const bool bDoColoring = true);
 
 		FXPBDCorotatedTetrahedralConstraints(
 			const FSolverParticles& InParticles,
@@ -177,496 +53,96 @@ namespace Chaos::Softs
 			const FDeformableXPBDCorotatedParams& InParams,
 			const FSolverReal& NuMesh = (FSolverReal).3,
 			const bool bRecordMetricIn = false,
-			const bool bDoColoring = true
-		)
+			const bool bDoColoring = true)
 			: FXPBDCorotatedTetrahedralConstraints(
-				FSolverParticlesRange(const_cast<FSolverParticles*>(&InParticles), /*InOffset = */ 0, /*InRange = */ InParticles.Size()),
-				InMesh,
-				EMeshArray,
-				NuMeshArray,
-				MoveTemp(AlphaJMeshArray),
-				InParams,
-				NuMesh,
-				bRecordMetricIn,
-				bDoColoring
-			)
+				FSolverParticlesRange(InParticles),
+				InMesh, EMeshArray, NuMeshArray, MoveTemp(AlphaJMeshArray),
+				InParams, NuMesh, bRecordMetricIn, bDoColoring)
 		{}
 
-		FXPBDCorotatedTetrahedralConstraints(
+		CHAOS_API FXPBDCorotatedTetrahedralConstraints(
 			const FSolverParticles& InParticles,
 			const TArray<TVector<int32, 4>>& InMesh,
-			const FSolverReal GridN = (FSolverReal).1,
+			const FSolverReal GridN,
 			const FSolverReal& EMesh = (FSolverReal)10.0,
-			const FSolverReal& NuMesh = (FSolverReal).3
-		)
-			: MeshConstraints(InMesh)
-		{
-			LambdaArray.Init((FSolverReal)0., 2 * MeshConstraints.Num());
-			DmInverse.Init((FSolverReal)0., 9 * MeshConstraints.Num());
-			Measure.Init((FSolverReal)0., MeshConstraints.Num());
-			Lambda = EMesh * NuMesh / (((FSolverReal)1. + NuMesh) * ((FSolverReal)1. - (FSolverReal)2. * NuMesh));
-			Mu = EMesh / ((FSolverReal)2. * ((FSolverReal)1. + NuMesh));
-			for (int e = 0; e < InMesh.Num(); e++)
-			{
-				PMatrix<FSolverReal, 3, 3> Dm = DsInit(e, InParticles);
-				PMatrix<FSolverReal, 3, 3> DmInv = Dm.Inverse();
-				for (int r = 0; r < 3; r++) {
-					for (int c = 0; c < 3; c++) {
-						DmInverse[(3 * 3) * e + 3 * r + c] = DmInv.GetAt(r, c);
-					}
-				}
-
-				Measure[e] = Dm.Determinant() / (FSolverReal)6.;
-
-				if (Measure[e] < (FSolverReal)0.)
-				{
-					Measure[e] = -Measure[e];
-				}
-			}
-			DmInverseSave = DmInverse;
-		}
+			const FSolverReal& NuMesh = (FSolverReal).3);
 
 		virtual ~FXPBDCorotatedTetrahedralConstraints() = default;
 
-		PMatrix<FSolverReal, 3, 3> DsInit(const int e, const FSolverParticlesRange& InParticles) const
-		{
-			PMatrix<FSolverReal, 3, 3> Result((FSolverReal)0.);
-			for (int i = 0; i < 3; i++)
-			{
-				for (int c = 0; c < 3; c++)
-				{
-					Result.SetAt(c, i, InParticles.GetX(MeshConstraints[e][i + 1])[c] - InParticles.GetX(MeshConstraints[e][0])[c]);
-				}
-			}
-			return Result;
-		}
-
-		PMatrix<FSolverReal, 3, 3> DsInit(const int e, const FSolverParticles& InParticles) const
-		{	
-			return DsInit(e, FSolverParticlesRange(const_cast<FSolverParticles*>(&InParticles), /*InOffset = */ 0, /*InRange = */ InParticles.Size()));
-		}
-
-		PMatrix<FSolverReal, 3, 3> Ds(const int e, const FSolverParticlesRange& InParticles) const {
-			PMatrix<FSolverReal, 3, 3> Result((FSolverReal)0.);
-			for (int i = 0; i < 3; i++) {
-				for (int c = 0; c < 3; c++) {
-					Result.SetAt(c, i, InParticles.GetP(MeshConstraints[e][i+1])[c] - InParticles.GetP(MeshConstraints[e][0])[c]);
-				}
-			}
-			return Result;
-		}
-
-		PMatrix<FSolverReal, 3, 3> F(const int e, const FSolverParticlesRange& InParticles) const {
-			return ElementDmInv(e) * Ds(e, InParticles);
-		}
-
-		PMatrix<FSolverReal, 3, 3> ElementDmInv(const int e) const {
-			PMatrix<FSolverReal, 3, 3> DmInv((FSolverReal)0.);
-			for (int r = 0; r < 3; r++) {
-				for (int c = 0; c < 3; c++) {
-					DmInv.SetAt(r, c, DmInverse[(3 * 3) * e + 3 * r + c]);
-				}
-			}
-			return DmInv;
-		}
-
-		PMatrix<FSolverReal, 3, 3> ElementDmInvSave(const int e) const {
-			PMatrix<FSolverReal, 3, 3> DmInv((FSolverReal)0.);
-			for (int r = 0; r < 3; r++) {
-				for (int c = 0; c < 3; c++) {
-					DmInv.SetAt(r, c, DmInverseSave[(3 * 3) * e + 3 * r + c]);
-				}
-			}
-			return DmInv;
-		}
-		
-		virtual void Init() const 
+		virtual void Init() const
 		{
 			for (FSolverReal& Lambdas : LambdaArray) { Lambdas = (FSolverReal)0.; }
 		}
 
-		virtual void ApplyInSerial(FSolverParticlesRange& Particles, const FSolverReal Dt, const int32 ElementIndex) const
-		{
-			TRACE_CPUPROFILER_EVENT_SCOPE(STAT_ChaosXPBDCorotatedApplySingle);
-
-			TVec4<TVector<FSolverReal, 3>> PolarDelta = GetPolarDelta(Particles, Dt, ElementIndex);
-
-			for (int i = 0; i < 4; i++)
-			{
-				Particles.P(MeshConstraints[ElementIndex][i]) += PolarDelta[i];
-			}
+		CHAOS_API virtual void ApplyInSerial(FSolverParticlesRange& Particles, const FSolverReal Dt, const int32 ElementIndex) const;
 
-			TVec4<TVector<FSolverReal, 3>> DetDelta = GetDeterminantDelta(Particles, Dt, ElementIndex);
+		CHAOS_API void ApplyInSerial(FSolverParticlesRange& Particles, const FSolverReal Dt) const;
 
-			for (int i = 0; i < 4; i++)
-			{
-				Particles.P(MeshConstraints[ElementIndex][i]) += DetDelta[i];
-			}
-		}
+		CHAOS_API void ApplyInParallel(FSolverParticlesRange& Particles, const FSolverReal Dt) const;
 
-		void ApplyInSerial(FSolverParticlesRange& Particles, const FSolverReal Dt) const
-		{
-			SCOPE_CYCLE_COUNTER(STAT_ChaosXPBDCorotated);
-			TRACE_CPUPROFILER_EVENT_SCOPE(STAT_ChaosXPBDCorotatedApplySerial);
-			const int32 ConstraintColorNum = ConstraintsPerColorStartIndex.Num() - 1;
-			for (int32 ConstraintColorIndex = 0; ConstraintColorIndex < ConstraintColorNum; ++ConstraintColorIndex)
-			{
-				const int32 ColorStart = ConstraintsPerColorStartIndex[ConstraintColorIndex];
-				const int32 ColorSize = ConstraintsPerColorStartIndex[ConstraintColorIndex + 1] - ColorStart;
-				for (int32 Index = 0; Index < ColorSize; Index++)
-				{
-					const int32 ConstraintIndex = ColorStart + Index;
-					ApplyInSerial(Particles, Dt, ConstraintIndex);
-				}
-			}
-		}
-
-		void ApplyInParallel(FSolverParticlesRange& Particles, const FSolverReal Dt) const
-		{	
-			//code for error metric:
-			if (bRecordMetric)
-			{
-				GError.Init((FSolverReal)0., 3 * Particles.Size());
-				HErrorArray.Init((FSolverReal)0., 2 * MeshConstraints.Num());
-			}
-			
-			{
-				SCOPE_CYCLE_COUNTER(STAT_ChaosXPBDCorotated);
-				TRACE_CPUPROFILER_EVENT_SCOPE(STAT_ChaosXPBDCorotatedApply);
-				if ((ConstraintsPerColorStartIndex.Num() > 1))//&& (MeshConstraints.Num() > Chaos_Spring_ParallelConstraintCount))
-				{
-					const int32 ConstraintColorNum = ConstraintsPerColorStartIndex.Num() - 1;
-
-					for (int32 ConstraintColorIndex = 0; ConstraintColorIndex < ConstraintColorNum; ++ConstraintColorIndex)
-					{
-						const int32 ColorStart = ConstraintsPerColorStartIndex[ConstraintColorIndex];
-						const int32 ColorSize = ConstraintsPerColorStartIndex[ConstraintColorIndex + 1] - ColorStart;
-
-						int32 NumBatch = ColorSize / CorotatedParams.XPBDCorotatedBatchSize;
-						if (ColorSize % CorotatedParams.XPBDCorotatedBatchSize != 0)
-						{
-							NumBatch += 1;
-						}
-
-						PhysicsParallelFor(NumBatch, [&](const int32 BatchIndex)
-							{
-								for (int32 BatchSubIndex = 0; BatchSubIndex < CorotatedParams.XPBDCorotatedBatchSize; BatchSubIndex++) {
-									int32 TaskIndex = CorotatedParams.XPBDCorotatedBatchSize * BatchIndex + BatchSubIndex;
-									const int32 ConstraintIndex = ColorStart + TaskIndex;
-									if (ConstraintIndex < ColorStart + ColorSize)
-									{
-										ApplyInSerial(Particles, Dt, ConstraintIndex);
-									}
-								}
-							}, NumBatch < CorotatedParams.XPBDCorotatedBatchThreshold);
-					}
-				}
-			}
-		}
-
-		void ApplyInParallel(FSolverParticles& Particles, const FSolverReal Dt) const
-		{
-			FSolverParticlesRange Range(&Particles, /*InOffset = */ 0, /*InRange = */ Particles.Size());
-			ApplyInParallel(Range, Dt);
-		}
+		CHAOS_API void ApplyInParallel(FSolverParticles& Particles, const FSolverReal Dt) const;
 
 		void Apply(FSolverParticlesRange& Particles, const FSolverReal Dt) const
 		{
 			ApplyInParallel(Particles, Dt);
 		}
 
-		TVec4<TVector<FSolverReal, 3>> GetPolarGradient(const PMatrix<FSolverReal, 3, 3>& Fe, const PMatrix<FSolverReal, 3, 3>& Re, const PMatrix<FSolverReal, 3, 3>& DmInvT, const FSolverReal C1) const
-		{
-			//TVector<FSolverReal, 81> dRdF((FSolverReal)0.);
-			//Chaos::dRdFCorotated(Fe, dRdF);
-			TVec4<TVector<FSolverReal, 3>> dC1(TVector<FSolverReal, 3>((FSolverReal)0.));
-			//dC1 = dC1dF * dFdX
-			PMatrix<FSolverReal, 3, 3> A = DmInvT * (Fe - Re);
-			for (int alpha = 0; alpha < 3; alpha++) {
-				for (int l = 0; l < 3; l++) {
-					dC1[0][alpha] -= A.GetAt(alpha, l);
-				}
-			}
-			for (int ie = 0; ie < 3; ie++) {
-				for (int alpha = 0; alpha < 3; alpha++) {
-					dC1[ie + 1][alpha] = A.GetAt(alpha, ie);
-				}
-			}
+		CHAOS_API void ModifyDmInverseFromMuscleLength(const int32 ElemIdx, const FSolverReal FiberLengthRatio, const PMatrix<FSolverReal, 3, 3>& FiberDir, const FSolverReal ContractionVolumeScale) const;
 
-			if (C1 != 0)
-			{
-				for (int i = 0; i < 4; i++)
-				{
-					for (int j = 0; j < 3; j++)
-					{
-						dC1[i][j] /= C1;
-					}
-				}
-			}
-			return dC1;
+		CHAOS_API void ModifyDmInverseSaveFromInflationVolumeScale(const int32 ElemIdx, const FSolverReal InflationVolumeScale, const PMatrix<FSolverReal, 3, 3>& FiberDir);
 
-		}
+	protected:
 
-		TVec4<TVector<FSolverReal, 3>> GetDeterminantGradient(const PMatrix<FSolverReal, 3, 3>& Fe, const PMatrix<FSolverReal, 3, 3>& DmInvT) const
+		PMatrix<FSolverReal, 3, 3> ElementDmInv(const int32 ElemIdx) const
 		{
-			TRACE_CPUPROFILER_EVENT_SCOPE(STAT_ChaosXPBDCorotatedGetDetGradient);
-			//SCOPE_CYCLE_COUNTER(STAT_ChaosXPBDCorotatedDet);
-			//const TVec4<int32>& Constraint = MeshConstraints[ElementIndex];
-
-			//const PMatrix<FSolverReal, 3, 3> Fe = F(ElementIndex, Particles);
-			//PMatrix<FSolverReal, 3, 3> DmInvT = ElementDmInv(ElementIndex).GetTransposed();
-			TVec4<TVector<FSolverReal, 3>> dC2(TVector<FSolverReal, 3>((FSolverReal)0.));
-
-			PMatrix<FSolverReal, 3, 3> JFinvT((FSolverReal)0.);
-			JFinvT.SetAt(0, 0, Fe.GetAt(1, 1) * Fe.GetAt(2, 2) - Fe.GetAt(2, 1) * Fe.GetAt(1, 2));
-			JFinvT.SetAt(0, 1, Fe.GetAt(2, 0) * Fe.GetAt(1, 2) - Fe.GetAt(1, 0) * Fe.GetAt(2, 2));
-			JFinvT.SetAt(0, 2, Fe.GetAt(1, 0) * Fe.GetAt(2, 1) - Fe.GetAt(2, 0) * Fe.GetAt(1, 1));
-			JFinvT.SetAt(1, 0, Fe.GetAt(2, 1) * Fe.GetAt(0, 2) - Fe.GetAt(0, 1) * Fe.GetAt(2, 2));
-			JFinvT.SetAt(1, 1, Fe.GetAt(0, 0) * Fe.GetAt(2, 2) - Fe.GetAt(2, 0) * Fe.GetAt(0, 2));
-			JFinvT.SetAt(1, 2, Fe.GetAt(2, 0) * Fe.GetAt(0, 1) - Fe.GetAt(0, 0) * Fe.GetAt(2, 1));
-			JFinvT.SetAt(2, 0, Fe.GetAt(0, 1) * Fe.GetAt(1, 2) - Fe.GetAt(1, 1) * Fe.GetAt(0, 2));
-			JFinvT.SetAt(2, 1, Fe.GetAt(1, 0) * Fe.GetAt(0, 2) - Fe.GetAt(0, 0) * Fe.GetAt(1, 2));
-			JFinvT.SetAt(2, 2, Fe.GetAt(0, 0) * Fe.GetAt(1, 1) - Fe.GetAt(1, 0) * Fe.GetAt(0, 1));
-
-			PMatrix<FSolverReal, 3, 3> JinvTDmInvT = DmInvT * JFinvT;
-
-			for (int ie = 0; ie < 3; ie++) {
-				for (int alpha = 0; alpha < 3; alpha++) {
-					dC2[ie + 1][alpha] = JinvTDmInvT.GetAt(alpha, ie);
-				}
-			}
-			for (int alpha = 0; alpha < 3; alpha++) {
-				for (int l = 0; l < 3; l++) {
-					dC2[0][alpha] -= JinvTDmInvT.GetAt(alpha, l);
+			PMatrix<FSolverReal, 3, 3> DmInv((FSolverReal)0.);
+			for (int32 Row = 0; Row < 3; Row++) {
+				for (int32 Col = 0; Col < 3; Col++) {
+					DmInv.SetAt(Row, Col, DmInverse[(3 * 3) * ElemIdx + 3 * Row + Col]);
 				}
 			}
-			return dC2;
+			return DmInv;
 		}
 
-		void ModifyDmInverseFromMuscleLength(const int32 ElemIdx, const FSolverReal FiberLengthRatio, const PMatrix<FSolverReal, 3, 3>& FiberDir, const FSolverReal ContractionVolumeScale) const
+		PMatrix<FSolverReal, 3, 3> ElementDmInvSave(const int32 ElemIdx) const
 		{
-			TRACE_CPUPROFILER_EVENT_SCOPE(STAT_ChaosXPBDCorotatedModifyDmInverseFromMuscleLength);
-			PMatrix<FSolverReal, 3, 3> DmInv = ElementDmInvSave(ElemIdx);
-			if (FiberLengthRatio < 1)
-			{
-				PMatrix<FSolverReal, 3, 3> D(1 / FiberLengthRatio, FMath::Pow(FiberLengthRatio, ContractionVolumeScale/FSolverReal(2)), FMath::Pow(FiberLengthRatio, ContractionVolumeScale/FSolverReal(2)));
-				PMatrix<FSolverReal, 3, 3> Factor = FiberDir * D * FiberDir.GetTransposed();
-				DmInv = Factor * DmInv; //Matrix multiplication convention: AB = B*A
-			}
-			for (int Row = 0; Row < 3; Row++)
-			{
-				for (int Column = 0; Column < 3; Column++)
-				{
-					DmInverse[(3 * 3) * ElemIdx + 3 * Row + Column] = DmInv.GetAt(Row, Column);
+			PMatrix<FSolverReal, 3, 3> DmInv((FSolverReal)0.);
+			for (int32 Row = 0; Row < 3; Row++) {
+				for (int32 Col = 0; Col < 3; Col++) {
+					DmInv.SetAt(Row, Col, DmInverseSave[(3 * 3) * ElemIdx + 3 * Row + Col]);
 				}
 			}
+			return DmInv;
 		}
 
-		void ModifyDmInverseSaveFromInflationVolumeScale(const int32 ElemIdx, const FSolverReal InflationVolumeScale, const PMatrix<FSolverReal, 3, 3>& FiberDir)
-		{
-			PMatrix<FSolverReal, 3, 3> DmInv = ElementDmInvSave(ElemIdx);
-			if (FMath::Abs(InflationVolumeScale - 1) > UE_SMALL_NUMBER)
-			{
-				FSolverReal InvSqrt = FMath::InvSqrt(InflationVolumeScale);
-				PMatrix<FSolverReal, 3, 3> D(1, InvSqrt, InvSqrt);
-				PMatrix<FSolverReal, 3, 3> Factor = FiberDir * D * FiberDir.GetTransposed();
-				DmInv = Factor * DmInv; //Matrix multiplication convention: AB = B*A
-			}
-			for (int Row = 0; Row < 3; Row++)
-			{
-				for (int Column = 0; Column < 3; Column++)
-				{
-					DmInverseSave[(3 * 3) * ElemIdx + 3 * Row + Column] = DmInv.GetAt(Row, Column);
-					DmInverse[(3 * 3) * ElemIdx + 3 * Row + Column] = DmInv.GetAt(Row, Column);
-				}
-			}
-		}
-	protected:
+		PMatrix<FSolverReal, 3, 3> DsInit(const int32 ElemIdx, const FSolverParticlesRange& InParticles) const;
 
-		void InitColor(const FSolverParticlesRange& Particles)
+		PMatrix<FSolverReal, 3, 3> DsInit(const int32 ElemIdx, const FSolverParticles& InParticles) const
 		{
-			const TArray<TArray<int32>> ConstraintsPerColor = FGraphColoring::ComputeGraphColoringParticlesOrRange(MeshConstraints, Particles, 0, Particles.Size());
-
-			// Reorder constraints based on color so each array in ConstraintsPerColor contains contiguous elements.
-			TArray<TVec4<int32>> ReorderedConstraints;
-			TArray<FSolverReal> ReorderedMeasure;
-			TArray<FSolverReal> ReorderedDmInverse;
-			TArray<int32> OrigToReorderedIndices; // used to reorder stiffness indices
-			ReorderedConstraints.SetNumUninitialized(MeshConstraints.Num());
-			ReorderedMeasure.SetNumUninitialized(Measure.Num());
-			ReorderedDmInverse.SetNumUninitialized(DmInverse.Num());
-			OrigToReorderedIndices.SetNumUninitialized(MeshConstraints.Num());
-
-			ConstraintsPerColorStartIndex.Reset(ConstraintsPerColor.Num() + 1);
-
-			int32 ReorderedIndex = 0;
-			for (const TArray<int32>& ConstraintsBatch : ConstraintsPerColor)
-			{
-				ConstraintsPerColorStartIndex.Add(ReorderedIndex);
-				for (const int32& BatchConstraint : ConstraintsBatch)
-				{
-					const int32 OrigIndex = BatchConstraint;
-					ReorderedConstraints[ReorderedIndex] = MeshConstraints[OrigIndex];
-					ReorderedMeasure[ReorderedIndex] = Measure[OrigIndex];
-					for (int32 kk = 0; kk < 9; kk++)
-					{
-						ReorderedDmInverse[9 * ReorderedIndex + kk] = DmInverse[9 * OrigIndex + kk];
-					}
-					OrigToReorderedIndices[OrigIndex] = ReorderedIndex;
-
-					++ReorderedIndex;
-				}
-			}
-			ConstraintsPerColorStartIndex.Add(ReorderedIndex);
-
-			MeshConstraints = MoveTemp(ReorderedConstraints);
-			Measure = MoveTemp(ReorderedMeasure);
-			DmInverse = MoveTemp(ReorderedDmInverse);
+			return DsInit(ElemIdx, FSolverParticlesRange(InParticles));
 		}
 
-		void InitColor(const FSolverParticles& Particles)
-		{
-			return InitColor(FSolverParticlesRange(const_cast<FSolverParticles*>(&Particles), /*InOffset = */ 0, /*InRange = */ Particles.Size()));
-		}
+		PMatrix<FSolverReal, 3, 3> Ds(const int32 ElemIdx, const FSolverParticlesRange& InParticles) const;
 
-		virtual TVec4<TVector<FSolverReal, 3>> GetDeterminantDelta(const FSolverParticlesRange& Particles, const FSolverReal Dt, const int32 ElementIndex, const FSolverReal Tol = (FSolverReal)1e-3) const
+		PMatrix<FSolverReal, 3, 3> F(const int32 ElemIdx, const FSolverParticlesRange& InParticles) const
 		{
-			TRACE_CPUPROFILER_EVENT_SCOPE(STAT_ChaosXPBDCorotatedApplyDet);
-			//SCOPE_CYCLE_COUNTER(STAT_ChaosXPBDCorotatedDet);
-			////const TVec4<int32>& Constraint = MeshConstraints[ElementIndex];
-
-			const PMatrix<FSolverReal, 3, 3> Fe = F(ElementIndex, Particles);
-			PMatrix<FSolverReal, 3, 3> DmInvT = ElementDmInv(ElementIndex).GetTransposed();
-			
-			FSolverReal J = Fe.Determinant();
-			if (J - AlphaJArray[ElementIndex] < Tol)
-			{
-				return TVec4<TVector<FSolverReal, 3>>(TVector<FSolverReal, 3>((FSolverReal)0.));
-			}
-
-			TVec4<TVector<FSolverReal, 3>> dC2 = GetDeterminantGradient(Fe, DmInvT);
-
-			//T AlphaTilde = (FSolverReal)2. / (Dt * Dt * Lambda * Measure[ElementIndex]);
-			FSolverReal AlphaTilde = (FSolverReal)2. / (Dt * Dt * LambdaElementArray[ElementIndex] * Measure[ElementIndex]);
-
-			if (LambdaElementArray[ElementIndex] > (FSolverReal)1. / (FSolverReal)UE_SMALL_NUMBER)
-			{
-				AlphaTilde = (FSolverReal)0.;
-			}
-
-			if (bRecordMetric)
-			{
-				HError += J - 1 + AlphaTilde * LambdaArray[2 * ElementIndex + 1];
-				for (int32 ie = 0; ie < 4; ie++)
-				{
-					for (int32 Alpha = 0; Alpha < 3; Alpha++)
-					{
-						GError[MeshConstraints[ElementIndex][ie] * 3 + Alpha] -= dC2[ie][Alpha] * LambdaArray[2 * ElementIndex + 1];
-					}
-				}
-				HErrorArray[2 * ElementIndex + 1] = J - 1 + AlphaTilde * LambdaArray[2 * ElementIndex + 1];
-			}
-			
-
-			FSolverReal DLambda = (AlphaJArray[ElementIndex] - J) - AlphaTilde * LambdaArray[2 * ElementIndex + 1];
-
-			FSolverReal Denom = AlphaTilde;
-			for (int i = 0; i < 4; i++)
-			{
-				for (int j = 0; j < 3; j++)
-				{
-					Denom += dC2[i][j] * Particles.InvM(MeshConstraints[ElementIndex][i]) * dC2[i][j];
-				}
-			}
-			DLambda /= Denom;
-			LambdaArray[2 * ElementIndex + 1] += DLambda;
-			TVec4<TVector<FSolverReal, 3>> Delta(TVector<FSolverReal, 3>((FSolverReal)0.));
-			for (int i = 0; i < 4; i++)
-			{
-				for (int j = 0; j < 3; j++)
-				{
-					Delta[i][j] = Particles.InvM(MeshConstraints[ElementIndex][i]) * dC2[i][j] * DLambda;
-				}
-			}
-			return Delta;
+			return ElementDmInv(ElemIdx) * Ds(ElemIdx, InParticles);
 		}
 
-		virtual TVec4<TVector<FSolverReal, 3>> GetPolarDelta(const FSolverParticlesRange& Particles, const FSolverReal Dt, const int32 ElementIndex, const FSolverReal Tol = (FSolverReal)1e-3) const
-		{	
-			TRACE_CPUPROFILER_EVENT_SCOPE(STAT_ChaosXPBDCorotatedApplyPolar);
-			SCOPE_CYCLE_COUNTER(STAT_ChaosXPBDCorotatedPolar);
-			const PMatrix<FSolverReal, 3, 3> Fe = F(ElementIndex, Particles);
-
-			PMatrix<FSolverReal, 3, 3> Re((FSolverReal)0.), Se((FSolverReal)0.);
-
-			Chaos::PolarDecomposition(Fe, Re, Se);
-
-			Re *= FGenericPlatformMath::Pow(AlphaJArray[ElementIndex], (FSolverReal)1. / (FSolverReal)3.);
-
-			FSolverReal C1 = (FSolverReal)0.;
-			for (int i = 0; i < 3; i++)
-			{
-				for (int j = 0; j < 3; j++)
-				{
-					C1 += FMath::Square((Fe - Re).GetAt(i, j));
-				}
-			}
-			C1 = FMath::Sqrt(C1);
-
-			if (C1 < Tol )
-			{
-				return TVec4<TVector<FSolverReal, 3>>(TVector<FSolverReal, 3>((FSolverReal)0.));
-			}
-
-
-			PMatrix<FSolverReal, 3, 3> DmInvT = ElementDmInv(ElementIndex).GetTransposed();
+		TVec4<TVector<FSolverReal, 3>> GetPolarGradient(const PMatrix<FSolverReal, 3, 3>& Fe, const PMatrix<FSolverReal, 3, 3>& Re, const PMatrix<FSolverReal, 3, 3>& DmInvT, const FSolverReal C1) const;
 
-			TVec4<TVector<FSolverReal, 3>> dC1 = GetPolarGradient(Fe, Re, DmInvT, C1);
+		TVec4<TVector<FSolverReal, 3>> GetDeterminantGradient(const PMatrix<FSolverReal, 3, 3>& Fe, const PMatrix<FSolverReal, 3, 3>& DmInvT) const;
 
-			//T AlphaTilde = (FSolverReal)1. / (Dt * Dt * Mu * Measure[ElementIndex]);
-			FSolverReal AlphaTilde = (FSolverReal)1. / (Dt * Dt * MuElementArray[ElementIndex] * Measure[ElementIndex]);
-
-			if (MuElementArray[ElementIndex] > (FSolverReal)1./ (FSolverReal)UE_SMALL_NUMBER) 
-			{
-				AlphaTilde = (FSolverReal)0.;
-			}
-
-			if (bRecordMetric)
-			{
-				HError += C1 + AlphaTilde * LambdaArray[2 * ElementIndex + 0];
-				for (int32 ie = 0; ie < 4; ie++)
-				{
-					for (int32 Alpha = 0; Alpha < 3; Alpha++)
-					{
-						GError[MeshConstraints[ElementIndex][ie] * 3 + Alpha] -= dC1[ie][Alpha] * LambdaArray[2 * ElementIndex + 0];
-					}
-				}
-				HErrorArray[2 * ElementIndex + 0] = C1 + AlphaTilde * LambdaArray[2 * ElementIndex + 0];
-			}
+		CHAOS_API void InitColor(const FSolverParticlesRange& Particles);
 
-			FSolverReal DLambda = -C1 - AlphaTilde* LambdaArray[2 * ElementIndex + 0];
+		void InitColor(const FSolverParticles& Particles)
+		{
+			return InitColor(FSolverParticlesRange(Particles));
+		}
 
-			FSolverReal Denom = AlphaTilde;
-			for (int i = 0; i < 4; i++)
-			{
-				for (int j = 0; j < 3; j++)
-				{
-					Denom += dC1[i][j] * Particles.InvM(MeshConstraints[ElementIndex][i]) * dC1[i][j];
-				}
-			}
-			DLambda /= Denom;
-			LambdaArray[2 * ElementIndex + 0] += DLambda;
-			TVec4<TVector<FSolverReal, 3>> Delta(TVector<FSolverReal, 3>((FSolverReal)0.));
-			for (int i = 0; i < 4; i++) 
-			{
-				for (int j = 0; j < 3; j++)
-				{
-					Delta[i][j] = Particles.InvM(MeshConstraints[ElementIndex][i]) * dC1[i][j] * DLambda;
-				}
-			}
-			return Delta;
+		CHAOS_API virtual TVec4<TVector<FSolverReal, 3>> GetDeterminantDelta(const FSolverParticlesRange& Particles, const FSolverReal Dt, const int32 ElementIndex, const FSolverReal Tol = (FSolverReal)1e-3) const;
 
-		}
+		CHAOS_API virtual TVec4<TVector<FSolverReal, 3>> GetPolarDelta(const FSolverParticlesRange& Particles, const FSolverReal Dt, const int32 ElementIndex, const FSolverReal Tol = (FSolverReal)1e-3) const;
 
 	protected:
 		mutable TArray<FSolverReal> LambdaArray;
diff --git a/Engine/Source/Runtime/Experimental/Chaos/Public/Chaos/XPBDGridBasedCorotatedConstraints.h b/Engine/Source/Runtime/Experimental/Chaos/Public/Chaos/XPBDGridBasedCorotatedConstraints.h
--- a/Engine/Source/Runtime/Experimental/Chaos/Public/Chaos/XPBDGridBasedCorotatedConstraints.h
+++ b/Engine/Source/Runtime/Experimental/Chaos/Public/Chaos/XPBDGridBasedCorotatedConstraints.h
@@ -53,6 +53,8 @@ namespace Chaos::Softs
 			Measure.Init((FSolverReal)0., MeshConstraints.Num());
 			Lambda = EMesh * NuMesh / (((FSolverReal)1. + NuMesh) * ((FSolverReal)1. - (FSolverReal)2. * NuMesh));
 			Mu = EMesh / ((FSolverReal)2. * ((FSolverReal)1. + NuMesh));
+			LambdaElementArray.Init(Lambda, MeshConstraints.Num());
+			MuElementArray.Init(Mu, MeshConstraints.Num());
 			for (int e = 0; e < InMesh.Num(); e++)
 			{
 				PMatrix<FSolverReal, 3, 3> Dm = Base::DsInit(e, InParticles);
diff --git a/Engine/Source/Runtime/Experimental/Chaos/Public/Chaos/XPBDStretchBiasElementConstraints.h b/Engine/Source/Runtime/Experimental/Chaos/Public/Chaos/XPBDStretchBiasElementConstraints.h
--- a/Engine/Source/Runtime/Experimental/Chaos/Public/Chaos/XPBDStretchBiasElementConstraints.h
+++ b/Engine/Source/Runtime/Experimental/Chaos/Public/Chaos/XPBDStretchBiasElementConstraints.h
@@ -4,6 +4,7 @@
 #include "Chaos/CollectionPropertyFacade.h"
 #include "Chaos/PBDSoftsEvolutionFwd.h"
 #include "Chaos/PBDSoftsSolverParticles.h"
+#include "Chaos/SoftsSolverParticlesRange.h"
 #include "Chaos/PBDStiffness.h"
 #include "Chaos/ParticleRule.h"
 #include "ChaosStats.h"
@@ -109,16 +110,16 @@ public:
 	
 	const TArray<TVec3<int32>>& GetConstraints() const { return Constraints; }
 
-	TArray<TArray<int32>> GetConstraintsArray() const
+	TArray<TArray<int32>> GetConstraintsArray() const //TODO: avoid converting
 	{
 		TArray<TArray<int32>> ConstraintsArray;
 		ConstraintsArray.SetNum(Constraints.Num());
-		for (int32 i = 0; i < Constraints.Num(); i++)
+		for (int32 ConstraintIdx = 0; ConstraintIdx < Constraints.Num(); ConstraintIdx++)
 		{
-			ConstraintsArray[i].SetNum(3);
-			for (int32 j = 0; j < 3; j++)
+			ConstraintsArray[ConstraintIdx].SetNum(3);
+			for (int32 TriLocalIdx = 0; TriLocalIdx < 3; TriLocalIdx++)
 			{
-				ConstraintsArray[i][j] = Constraints[i][j];
+				ConstraintsArray[ConstraintIdx][TriLocalIdx] = Constraints[ConstraintIdx][TriLocalIdx];
 			}
 		}
 		return ConstraintsArray;
@@ -131,11 +132,15 @@ public:
 	}
 	const TArray<int32>& GetConstraintsPerColorStartIndex() const { return ConstraintsPerColorStartIndex; }
 
+	CHAOS_API void AddStretchBiasElementResidualAndHessian(const FSolverParticlesRange& Particles, const int32 ConstraintIndex, const int32 ConstraintIndexLocal, const FSolverReal Dt, TVec3<FSolverReal>& ParticleResidual, Chaos::PMatrix<FSolverReal, 3, 3>& ParticleHessian);
 	CHAOS_API void AddStretchBiasElementResidualAndHessian(const FSolverParticles& Particles, const int32 ConstraintIndex, const int32 ConstraintIndexLocal, const FSolverReal Dt, TVec3<FSolverReal>& ParticleResidual, Chaos::PMatrix<FSolverReal, 3, 3>& ParticleHessian);
 
+	CHAOS_API void InitializeDmInvAndMeasures(const FSolverParticlesRange& Particles);
+
 	CHAOS_API void InitializeDmInvAndMeasures(const FSolverParticles& Particles);
 
-	CHAOS_API void AddInternalForceDifferential(const FSolverParticles& InParticles, const TArray<TVector<FSolverReal, 3>>& DeltaParticles, TArray<TVector<FSolverReal, 3>>& ndf);
+	CHAOS_API void AddInternalForceDifferential(const FSolverParticlesRange& InParticles, const TArray<TVector<FSolverReal, 3>>& DeltaParticles, TArray<TVector<FSolverReal, 3>>& Ndf);
+	CHAOS_API void AddInternalForceDifferential(const FSolverParticles& InParticles, const TArray<TVector<FSolverReal, 3>>& DeltaParticles, TArray<TVector<FSolverReal, 3>>& Ndf);
 
 private:
 	template<typename SolverParticlesOrRange>
diff --git a/Engine/Source/Runtime/Experimental/Chaos/Public/Chaos/XPBDWeakConstraints.h b/Engine/Source/Runtime/Experimental/Chaos/Public/Chaos/XPBDWeakConstraints.h
--- a/Engine/Source/Runtime/Experimental/Chaos/Public/Chaos/XPBDWeakConstraints.h
+++ b/Engine/Source/Runtime/Experimental/Chaos/Public/Chaos/XPBDWeakConstraints.h
@@ -13,6 +13,7 @@
 #include "Misc/Paths.h"
 #include "Chaos/GraphColoring.h"
 #include "Chaos/DebugDrawQueue.h"
+#include "Chaos/SoftsSolverParticlesRange.h"
 
 //DECLARE_CYCLE_STAT(TEXT("Chaos XPBD Corotated Constraint"), STAT_ChaosXPBDCorotated, STATGROUP_Chaos);
 //DECLARE_CYCLE_STAT(TEXT("Chaos XPBD Corotated Constraint Polar Compute"), STAT_ChaosXPBDCorotatedPolar, STATGROUP_Chaos);
@@ -37,7 +38,7 @@ namespace Chaos::Softs
 	public:
 		static constexpr FSolverReal SoftMaxStiffness = (FSolverReal)1e14; // Stiffnesses greater than this will be treated as "hard" PBD constraints
 
-		// Used for spring constraints without self targets (against a fixed constraint target instead of deformable vertices)
+		// FSolverParticles constructors (delegate to Range versions)
 		CHAOS_API FXPBDWeightedSpringConstraints(
 			const FSolverParticles& InParticles,
 			const TArray<TArray<int32>>& InIndices,
@@ -66,6 +67,35 @@ namespace Chaos::Softs
 			const FDeformableXPBDWeightedSpringConstraintParams& InParams
 		);
 
+		// Range-based constructors
+		CHAOS_API FXPBDWeightedSpringConstraints(
+			const FSolverParticlesRange& InParticlesRange,
+			const TArray<TArray<int32>>& InIndices,
+			const TArray<TArray<FSolverReal>>& InWeights,
+			const TArray<FSolverReal>& InStiffness
+		);
+
+		CHAOS_API FXPBDWeightedSpringConstraints(
+			const FSolverParticlesRange& InParticlesRange,
+			const TArray<TArray<int32>>& InIndices,
+			const TArray<TArray<FSolverReal>>& InWeights,
+			const TArray<FSolverReal>& InStiffness,
+			const TArray<TArray<int32>>& InSecondIndices,
+			const TArray<TArray<FSolverReal>>& InSecondWeights,
+			const FDeformableXPBDWeightedSpringConstraintParams& InParams
+		);
+
+		CHAOS_API FXPBDWeightedSpringConstraints(
+			const FSolverParticlesRange& InParticlesRange,
+			const FSolverParticlesRange& InSecondParticlesRange,
+			const TArray<TArray<int32>>& InIndices,
+			const TArray<TArray<FSolverReal>>& InWeights,
+			const TArray<FSolverReal>& InStiffness,
+			const TArray<TArray<int32>>& InSecondIndices,
+			const TArray<TArray<FSolverReal>>& InSecondWeights,
+			const FDeformableXPBDWeightedSpringConstraintParams& InParams
+		);
+
 		virtual ~FXPBDWeightedSpringConstraints() = default;
 
 		CHAOS_API void ApplyInParallel(FSolverParticlesRange& Particles, const FSolverReal Dt) const;
@@ -114,10 +144,8 @@ namespace Chaos::Softs
 		TArray<FSolverReal> Stiffness;
 
 	private:
-
-		void InitColor(const FSolverParticles& Particles);
-
-		void InitColor(const FSolverParticles& Particles, const FSolverParticles& SecondParticles);
+		void InitColor(const FSolverParticlesRange& ParticlesRange);
+		void InitColor(const FSolverParticlesRange& ParticlesRange, const FSolverParticlesRange& SecondParticlesRange);
 
 		void ApplySingleConstraintWithoutSelfTarget(FSolverParticlesRange& Particles, const FSolverReal Dt, const int32 ConstraintIndex, const FSolverReal Tol = UE_KINDA_SMALL_NUMBER) const;
 
