ChaosDeformable: Fix incomplete type in CrossAssetDynamicWeakConstraintCreator

#Synced-CL 52219429

diff --git a/Engine/Source/Runtime/Experimental/Chaos/Public/Chaos/PBDAnimDriveConstraint.h b/Engine/Source/Runtime/Experimental/Chaos/Public/Chaos/PBDAnimDriveConstraint.h
--- a/Engine/Source/Runtime/Experimental/Chaos/Public/Chaos/PBDAnimDriveConstraint.h
+++ b/Engine/Source/Runtime/Experimental/Chaos/Public/Chaos/PBDAnimDriveConstraint.h
@@ -204,14 +204,21 @@ namespace Chaos::Softs
 			}
 
 			FSolverVec3& ParticlePosition = Particles.P(ParticleIndex);
+			UE_LOG(LogChaos, Verbose, TEXT("AnimDrive[%d] Before: P=(%.4f, %.4f, %.4f) AnimPos=(%.4f, %.4f, %.4f) Stiffness=%.4f Damping=%.4f"),
+				ParticleIndex, ParticlePosition.X, ParticlePosition.Y, ParticlePosition.Z, AnimPosition.X, AnimPosition.Y, AnimPosition.Z, InStiffness, InDamping);
 
 			ParticlePosition -= InStiffness * (ParticlePosition - AnimPosition);
+			UE_LOG(LogChaos, Verbose, TEXT("AnimDrive[%d] AfterStiffness: P=(%.4f, %.4f, %.4f)"), ParticleIndex, ParticlePosition.X, ParticlePosition.Y, ParticlePosition.Z);
 
 			const FSolverVec3 ParticleDisplacement = ParticlePosition - Particles.GetX(ParticleIndex);
 			const FSolverVec3 AnimationDisplacement = (AnimVelocity) * Dt;
 			const FSolverVec3 RelativeDisplacement = ParticleDisplacement - AnimationDisplacement;
+			UE_LOG(LogChaos, Verbose, TEXT("AnimDrive[%d] X=(%.4f, %.4f, %.4f) AnimVel=(%.4f, %.4f, %.4f) RelDisp=(%.4f, %.4f, %.4f)"),
+				ParticleIndex, Particles.GetX(ParticleIndex).X, Particles.GetX(ParticleIndex).Y, Particles.GetX(ParticleIndex).Z,
+				AnimVelocity.X, AnimVelocity.Y, AnimVelocity.Z, RelativeDisplacement.X, RelativeDisplacement.Y, RelativeDisplacement.Z);
 
 			ParticlePosition -= InDamping * RelativeDisplacement;
+			UE_LOG(LogChaos, Verbose, TEXT("AnimDrive[%d] AfterDamping: P=(%.4f, %.4f, %.4f)"), ParticleIndex, ParticlePosition.X, ParticlePosition.Y, ParticlePosition.Z);
 		}
 
 	private:
