[Feature] toggling IDO serialization based on whether impersonation is enabled

#Synced-CL 52110756

diff --git a/Engine/Source/Runtime/Experimental/JsonObjectGraph/Private/JsonStringifyImpl.cpp b/Engine/Source/Runtime/Experimental/JsonObjectGraph/Private/JsonStringifyImpl.cpp
--- a/Engine/Source/Runtime/Experimental/JsonObjectGraph/Private/JsonStringifyImpl.cpp
+++ b/Engine/Source/Runtime/Experimental/JsonObjectGraph/Private/JsonStringifyImpl.cpp
@@ -24,6 +24,7 @@
 #if WITH_EDITORONLY_DATA
 #include "UObject/InstanceDataObjectUtils.h"
 #include "UObject/PropertyBagRepository.h"
+#include "UObject/UObjectThreadContext.h"
 #endif// WITH_EDITORONLY_DATA
 
 namespace UE::Private
@@ -379,17 +380,20 @@ void FJsonStringifyImpl::WriteObjectToJson(const UObject* Object)
 	// Write native UObject data - name, type, flags, native user serialize, etc
 	WriteNativeObjectData();
 
-	for (TFieldIterator<FProperty> FieldIt(Object->GetClass()); FieldIt; ++FieldIt)
-	{
-		FProperty* Property = *FieldIt;
-		const UObject* ValidatedArchtype = ArchetypeClass && ArchetypeClass->IsChildOf(Property->GetOwnerClass()) ?
-			Archetype : nullptr;
-
-		WriteIdentifierAndValueToJson(Object, ValidatedArchtype, Property);
-	}
 #if WITH_EDITORONLY_DATA
-	WriteInstanceDataObjectProperties(Object);
+	// attempt to write the IDO instead since it will have more info. if that fails, fallback to the regular object
+	if (!TryWriteInstanceDataObjectProperties(Object))
 #endif // WITH_EDITORONLY_DATA
+	{ 
+		for (TFieldIterator<FProperty> FieldIt(Object->GetClass()); FieldIt; ++FieldIt)
+		{
+			FProperty* Property = *FieldIt;
+			const UObject* ValidatedArchtype = ArchetypeClass && ArchetypeClass->IsChildOf(Property->GetOwnerClass()) ?
+				Archetype : nullptr;
+
+			WriteIdentifierAndValueToJson(Object, ValidatedArchtype, Property);
+		}
+	}
 	// We may have inner objects that were not referenced directly by this object, 
 	// but will be referenced by other objects in the graph. We must write them here, or we
 	// would have to encode them at the root level, which would disrupt locality.
@@ -460,19 +464,29 @@ void FJsonStringifyImpl::WriteNativeObjectData()
 }
 
 #if WITH_EDITORONLY_DATA
-void FJsonStringifyImpl::WriteInstanceDataObjectProperties(const UObject* ForObject)
+bool FJsonStringifyImpl::TryWriteInstanceDataObjectProperties(const UObject* ForObject)
 {
+#if !UE_INTERNAL_UNIVERSAL_PLACEHOLDERS
+	// unless UNIVERSAL_PLACEHOLDERS is enabled, respect the impersonation flag.
+	FUObjectSerializeContext* SerializeContext = FUObjectThreadContext::Get().GetSerializeContext();
+	if (!SerializeContext || !SerializeContext->bImpersonateProperties)
+	{
+		return false;
+	}
+#endif
+
 	UE::FPropertyBagRepository& Repository = UE::FPropertyBagRepository::Get();
 	const UObject* InstanceDataObject = Repository.FindInstanceDataObject(ForObject);
 	if (!InstanceDataObject)
 	{
-		return;
+		return false;
 	}
 
 	for (TFieldIterator<FProperty> FieldIt(InstanceDataObject->GetClass()); FieldIt; ++FieldIt)
 	{
 		WriteIdentifierAndValueToJson(InstanceDataObject, nullptr, *FieldIt);
 	}
+	return true;
 }
 #endif// WITH_EDITORONLY_DATA
 
diff --git a/Engine/Source/Runtime/Experimental/JsonObjectGraph/Private/JsonStringifyImpl.h b/Engine/Source/Runtime/Experimental/JsonObjectGraph/Private/JsonStringifyImpl.h
--- a/Engine/Source/Runtime/Experimental/JsonObjectGraph/Private/JsonStringifyImpl.h
+++ b/Engine/Source/Runtime/Experimental/JsonObjectGraph/Private/JsonStringifyImpl.h
@@ -48,7 +48,7 @@ private:
 	// UObject construction (find archetypes, UClasses, UStructs, etc):
 	void WriteNativeObjectData();
 #if WITH_EDITORONLY_DATA
-	void WriteInstanceDataObjectProperties(const UObject* ForObject);
+	bool TryWriteInstanceDataObjectProperties(const UObject* ForObject);
 #endif // WITH_EDITORONLY_DATA
 	void WriteIndirectlyReferencedContainedObjects(const UObject* ForObject);
 	FUtf8String WriteObjectReference(const UObject* Object) const;
