Chaos Cloth :
- Added a collection dirty flag to the property facade as to prevent unecessary copies.
- Restricted the propery copies to the current LOD only.
- Updated the clearing of the dirty flags to happen inside the proxy instead of in all components code.
- Fixed CollectionPropertyFacade to use ModifyAttribute as to be able to use the underlaying collection dirty flag.
- Deprecated the FManagedArrayCollection Dirty flag API, and updated the Geometry Collection to use its own API.

#Synced-CL 52291475

diff --git a/Engine/Plugins/ChaosCloth/Source/ChaosCloth/Private/ChaosCloth/ChaosClothingSimulationConfig.cpp b/Engine/Plugins/ChaosCloth/Source/ChaosCloth/Private/ChaosCloth/ChaosClothingSimulationConfig.cpp
--- a/Engine/Plugins/ChaosCloth/Source/ChaosCloth/Private/ChaosCloth/ChaosClothingSimulationConfig.cpp
+++ b/Engine/Plugins/ChaosCloth/Source/ChaosCloth/Private/ChaosCloth/ChaosClothingSimulationConfig.cpp
@@ -25,7 +25,8 @@ namespace Chaos
 
 	FClothingSimulationConfig::FClothingSimulationConfig(const TArray<TSharedPtr<const FManagedArrayCollection>>& InPropertyCollections)
 	{
-		Initialize(InPropertyCollections);
+		constexpr bool bForceDirty = true;
+		Initialize(InPropertyCollections, bForceDirty);
 	}
 
 	FClothingSimulationConfig::~FClothingSimulationConfig() = default;
@@ -258,28 +259,41 @@ namespace Chaos
 		Property->AddValue(UseLegacyConfigName, bUseLegacyConfig, NonAnimatablePropertyFlags);
 	}
 
-	void FClothingSimulationConfig::Initialize(const TArray<TSharedPtr<const FManagedArrayCollection>>& InPropertyCollections)
+	void FClothingSimulationConfig::Initialize(const TArray<TSharedPtr<const FManagedArrayCollection>>& InPropertyCollections, bool bForceDirty, int32 LODIndex)
 	{
+		if (LODIndex != INDEX_NONE && !InPropertyCollections.IsValidIndex(LODIndex))
+		{
+			return;
+		}
 		PropertyCollections.SetNum(InPropertyCollections.Num());
 		Properties.SetNum(InPropertyCollections.Num());
 
-		for (int32 Index = 0; Index < InPropertyCollections.Num(); ++Index)
+		const int32 StartIndex = (LODIndex != INDEX_NONE) ? LODIndex : 0;
+		const int32 EndIndex = (LODIndex != INDEX_NONE) ? LODIndex : InPropertyCollections.Num() - 1;
+
+		for (int32 Index = StartIndex; Index <= EndIndex; ++Index)
 		{
 			const TSharedPtr<const FManagedArrayCollection>& InPropertyCollection = InPropertyCollections[Index];
-			TSharedPtr<FManagedArrayCollection>& PropertyCollection = PropertyCollections[Index];
-			TUniquePtr<Softs::FCollectionPropertyMutableFacade>& Property = Properties[Index];
-			if (!PropertyCollection || !Property)
+			::Chaos::Softs::FCollectionPropertyConstFacade InPropertyFacade(InPropertyCollection);
+			if (bForceDirty || InPropertyFacade.IsDirty())
 			{
-				PropertyCollection = MakeShared<FManagedArrayCollection>();
-				Property = MakeUnique<Softs::FCollectionPropertyMutableFacade>(PropertyCollection);
+				TSharedPtr<FManagedArrayCollection>& PropertyCollection = PropertyCollections[Index];
+				TUniquePtr<Softs::FCollectionPropertyMutableFacade>& Property = Properties[Index];
+				if (!PropertyCollection || !Property)
+				{
+					PropertyCollection = MakeShared<FManagedArrayCollection>();
+					Property = MakeUnique<Softs::FCollectionPropertyMutableFacade>(PropertyCollection);
+				}
+				Property->Copy(*InPropertyCollection);
 			}
-			Property->Copy(*InPropertyCollection);
 		}
 	}
+
 	const Softs::FCollectionPropertyConstFacade& FClothingSimulationConfig::GetProperties(int32 LODIndex) const
 	{
 		return bIsLegacySingleLOD ? *Properties[0] : *Properties[LODIndex];
 	}
+
 	Softs::FCollectionPropertyFacade& FClothingSimulationConfig::GetProperties(int32 LODIndex)
 	{
 		return bIsLegacySingleLOD ? *Properties[0] : *Properties[LODIndex];
diff --git a/Engine/Plugins/ChaosCloth/Source/ChaosCloth/Public/ChaosCloth/ChaosClothingSimulationConfig.h b/Engine/Plugins/ChaosCloth/Source/ChaosCloth/Public/ChaosCloth/ChaosClothingSimulationConfig.h
--- a/Engine/Plugins/ChaosCloth/Source/ChaosCloth/Public/ChaosCloth/ChaosClothingSimulationConfig.h
+++ b/Engine/Plugins/ChaosCloth/Source/ChaosCloth/Public/ChaosCloth/ChaosClothingSimulationConfig.h
@@ -39,8 +39,13 @@ namespace Chaos
 		 */
 		CHAOSCLOTH_API void Initialize(const UChaosClothConfig* ClothConfig, const UChaosClothSharedSimConfig* ClothSharedConfig, bool bUseLegacyConfig = false);
 
-		/** Initialize config from an array of property collections (one per LOD). */
-		CHAOSCLOTH_API void Initialize(const TArray<TSharedPtr<const FManagedArrayCollection>>& InPropertyCollections);
+		/**
+		 * Initialize config from an array of property collections (one per LOD).
+		 * @param InPropertyCollections The cloth source property collections to initialize from.
+		 * @param bForceDirty Whether to force the initialization even if the source property collections haven't changed.
+		 * @param LODIndex The index of the LOD to replace, or INDEX_NONE to replace all LODs.
+		 */
+		CHAOSCLOTH_API void Initialize(const TArray<TSharedPtr<const FManagedArrayCollection>>& InPropertyCollections, bool bForceDirty = false, int32 LODIndex = INDEX_NONE);
 
 		/** Return a property collection facade for reading properties from this configuration. */
 		CHAOSCLOTH_API const Softs::FCollectionPropertyConstFacade& GetProperties(int32 LODIndex) const;
diff --git a/Engine/Plugins/ChaosClothAsset/Source/ChaosClothAssetEngine/Private/ChaosClothAsset/ClothAssetSKMClothingSimulation.cpp b/Engine/Plugins/ChaosClothAsset/Source/ChaosClothAssetEngine/Private/ChaosClothAsset/ClothAssetSKMClothingSimulation.cpp
--- a/Engine/Plugins/ChaosClothAsset/Source/ChaosClothAssetEngine/Private/ChaosClothAsset/ClothAssetSKMClothingSimulation.cpp
+++ b/Engine/Plugins/ChaosClothAsset/Source/ChaosClothAssetEngine/Private/ChaosClothAsset/ClothAssetSKMClothingSimulation.cpp
@@ -391,6 +391,24 @@ namespace UE::Chaos::ClothAsset
 		return EmptyManagedArrayCollectionArray;
 	}
 
+	void FSKMClothingSimulation::ClearPropertyCollectionsDirtyFlags(int32 LodIndex) const
+	{
+		for (const FAssetClothSimulationProperties& Properties : AssetClothSimulationProperties)
+		{
+			if (const UChaosClothAssetSKMClothingAsset* const ClothAssetSKMClothingAsset = Properties.Key)
+			{
+				if (Properties.Value.CollectionPropertyFacades.IsValidIndex(LodIndex))
+				{
+					const TSharedPtr<::Chaos::Softs::FCollectionPropertyFacade>& CollectionPropertyFacade = Properties.Value.CollectionPropertyFacades[LodIndex];
+					if (CollectionPropertyFacade.IsValid())
+					{
+						CollectionPropertyFacade->ClearDirtyFlags();
+					}
+				}
+			}
+		}
+	}
+
 	bool FSKMClothingSimulation::HasAnySimulationMeshData(int32 LodIndex) const
 	{
 		// Note that for ClothAsset-based simulation, OwnerLODIndex == ClothAssetLODIndex
diff --git a/Engine/Plugins/ChaosClothAsset/Source/ChaosClothAssetEngine/Private/ChaosClothAsset/ClothComponent.cpp b/Engine/Plugins/ChaosClothAsset/Source/ChaosClothAssetEngine/Private/ChaosClothAsset/ClothComponent.cpp
--- a/Engine/Plugins/ChaosClothAsset/Source/ChaosClothAssetEngine/Private/ChaosClothAsset/ClothComponent.cpp
+++ b/Engine/Plugins/ChaosClothAsset/Source/ChaosClothAssetEngine/Private/ChaosClothAsset/ClothComponent.cpp
@@ -704,20 +704,22 @@ void UChaosClothComponent::SetSimulateInEditor(const bool bNewSimulateState)
 
 void UChaosClothComponent::StartNewParallelSimulation(float DeltaTime)
 {
-	using namespace ::Chaos::Softs;
 	if (ClothSimulationProxy.IsValid())
 	{
 		CSV_SCOPED_TIMING_STAT(Animation, Cloth);
-		const bool bIsSimulating = ClothSimulationProxy->Tick_GameThread(DeltaTime);
-		const int32 CurrentLOD = GetPredictedLODLevel();
+		ClothSimulationProxy->Tick_GameThread(DeltaTime);
+	}
+}
 
-		for (int32 ModelIndex = 0; ModelIndex < ClothSimulationProperties.Num(); ++ModelIndex)
+void UChaosClothComponent::ClearPropertyCollectionsDirtyFlags(int32 LODIndex) const
+{
+	using namespace ::Chaos::Softs;
+	for (int32 ModelIndex = 0; ModelIndex < ClothSimulationProperties.Num(); ++ModelIndex)
+	{
+		const TArray<TSharedPtr<FCollectionPropertyFacade>>& CollectionPropertyFacades = ClothSimulationProperties[ModelIndex].CollectionPropertyFacades;
+		if (CollectionPropertyFacades.IsValidIndex(LODIndex) && CollectionPropertyFacades[LODIndex].IsValid())
 		{
-			TArray<TSharedPtr<FCollectionPropertyFacade>>& CollectionPropertyFacades = ClothSimulationProperties[ModelIndex].CollectionPropertyFacades;
-			if (bIsSimulating && CollectionPropertyFacades.IsValidIndex(CurrentLOD) && CollectionPropertyFacades[CurrentLOD].IsValid())
-			{
-				CollectionPropertyFacades[CurrentLOD]->ClearDirtyFlags();
-			}
+			CollectionPropertyFacades[LODIndex]->ClearDirtyFlags();
 		}
 	}
 }
@@ -826,23 +828,7 @@ void UChaosClothComponent::WriteToSimulation(const float DeltaTime, const bool b
 	using namespace ::Chaos::Softs;
 	if (ClothSimulationProxy.IsValid())
 	{
-		const bool bIsSimulating = ClothSimulationProxy->PreSimulate_GameThread(DeltaTime);
-
-		if (bIsSimulating)
-		{
-			// Clear property dirty flags
-			const int32 CurrentLOD = GetPredictedLODLevel();
-
-			for (int32 ModelIndex = 0; ModelIndex < ClothSimulationProperties.Num(); ++ModelIndex)
-			{
-				TArray<TSharedPtr<FCollectionPropertyFacade>>& CollectionPropertyFacades = ClothSimulationProperties[ModelIndex].CollectionPropertyFacades;
-
-				if (CollectionPropertyFacades.IsValidIndex(CurrentLOD) && CollectionPropertyFacades[CurrentLOD].IsValid())
-				{
-					CollectionPropertyFacades[CurrentLOD]->ClearDirtyFlags();
-				}
-			}
-		}
+		ClothSimulationProxy->PreSimulate_GameThread(DeltaTime);
 	}
 }
 
diff --git a/Engine/Plugins/ChaosClothAsset/Source/ChaosClothAssetEngine/Private/ChaosClothAsset/ClothSimulationProxy.cpp b/Engine/Plugins/ChaosClothAsset/Source/ChaosClothAssetEngine/Private/ChaosClothAsset/ClothSimulationProxy.cpp
--- a/Engine/Plugins/ChaosClothAsset/Source/ChaosClothAssetEngine/Private/ChaosClothAsset/ClothSimulationProxy.cpp
+++ b/Engine/Plugins/ChaosClothAsset/Source/ChaosClothAssetEngine/Private/ChaosClothAsset/ClothSimulationProxy.cpp
@@ -413,9 +413,7 @@ namespace UE::Chaos::ClothAsset
 
 		// Fill a new context, note the context is also needed when the simulation is suspended or playing back the cache
 		constexpr bool bIsInitializationFalse = false;
-		PRAGMA_DISABLE_DEPRECATION_WARNINGS
 		FillSimulationContext(DeltaTime, bIsInitializationFalse);
-		PRAGMA_ENABLE_DEPRECATION_WARNINGS
 
 		// Check whether the solver should be enabled for caching purpose
 		const bool bShouldEnableSolver = ShouldEnableSolver(Solver->GetEnableSolver());  // Note: needs to be called after filling the context
@@ -495,6 +493,17 @@ namespace UE::Chaos::ClothAsset
 		if (bIsSimulating)
 		{
 			WriteSimulationData(ExternalClothSimulData);
+
+			// Clean properties dirty flags to avoid re-updating the constraints at the next tick when properties have not changed
+			// Note this does not clear the source collection whose dirty flags are cleared after the copy in InitializeConfigs()
+			for (TUniquePtr<::Chaos::FClothingSimulationConfig>& Config : Configs)
+			{
+				const int32 LODIndex = Solver->GetSolverLOD();
+				if (LODIndex < Config->GetNumLODs())
+				{
+					Config->GetProperties(LODIndex).ClearDirtyFlags();
+				}
+			}
 		}
 	}
 
@@ -507,7 +516,6 @@ namespace UE::Chaos::ClothAsset
 		}
 		if (!bIsSimulating)
 		{
-
 			// If the simulation is enabled, then it is suspended or reading from cache and the simulation data still needs updating (transforms and LODs)
 			if (ClothComponentAdapter.IsSimulationEnabled())
 			{
@@ -1043,12 +1051,16 @@ namespace UE::Chaos::ClothAsset
 	void FClothSimulationProxy::InitializeConfigs()
 	{
 		TRACE_CPUPROFILER_EVENT_SCOPE(FClothSimulationProxy_InitializeConfigs);
-		
+
+		const int32 LodIndex = ClothSimulationContext->LodIndex;
+		Solver->SetSolverLOD(LodIndex);
+
 		if (Configs.Num())
 		{
 			// Replace physics thread's configs with the game thread's configs
+			constexpr bool bForceDirty = false;  // Do not force the copy of the collection if it isn't dirty
 			int32 ConfigIndex = 0;
-			Configs[ConfigIndex++]->Initialize(ClothComponentAdapter.GetSolverPropertyCollections());
+			Configs[ConfigIndex++]->Initialize(ClothComponentAdapter.GetSolverPropertyCollections(), bForceDirty, LodIndex);
 
 			for (const UChaosClothAssetBase* const Asset : ClothComponentAdapter.GetAssets())
 			{
@@ -1060,13 +1072,13 @@ namespace UE::Chaos::ClothAsset
 						if (SimulationGroupId != INDEX_NONE &&  // INDEX_NONE means this model hasn't been added to the simulation
 							ensure(Configs.IsValidIndex(ConfigIndex)))
 						{
-							Configs[ConfigIndex++]->Initialize(ClothComponentAdapter.GetPropertyCollections(Asset, ModelIndex));
+							Configs[ConfigIndex++]->Initialize(ClothComponentAdapter.GetPropertyCollections(Asset, ModelIndex), bForceDirty, LodIndex);
 						}
 					}
 				}
 			}
+			ClothComponentAdapter.ClearPropertyCollectionsDirtyFlags(LodIndex);
 		}
-		Solver->SetSolverLOD(ClothSimulationContext->LodIndex);
 	}
 
 	void FClothSimulationProxy::FillSimulationContext(float DeltaTime, bool bIsInitialization)
diff --git a/Engine/Plugins/ChaosClothAsset/Source/ChaosClothAssetEngine/Public/ChaosClothAsset/ClothAssetSKMClothingSimulation.h b/Engine/Plugins/ChaosClothAsset/Source/ChaosClothAssetEngine/Public/ChaosClothAsset/ClothAssetSKMClothingSimulation.h
--- a/Engine/Plugins/ChaosClothAsset/Source/ChaosClothAssetEngine/Public/ChaosClothAsset/ClothAssetSKMClothingSimulation.h
+++ b/Engine/Plugins/ChaosClothAsset/Source/ChaosClothAssetEngine/Public/ChaosClothAsset/ClothAssetSKMClothingSimulation.h
@@ -134,6 +134,7 @@ namespace UE::Chaos::ClothAsset
 		virtual float GetClothGeometryScale() const override;
 		virtual const TArray<TSharedPtr<const FManagedArrayCollection>>& GetPropertyCollections(const UChaosClothAssetBase* Asset, int32 ModelIndex) const override;
 		virtual const TArray<TSharedPtr<const FManagedArrayCollection>>& GetSolverPropertyCollections() const override;
+		virtual void ClearPropertyCollectionsDirtyFlags(int32 LodIndex) const override;
 		virtual bool HasAnySimulationMeshData(int32 LodIndex) const override;
 		//~ End IClothComponentAdapter Interface
 
diff --git a/Engine/Plugins/ChaosClothAsset/Source/ChaosClothAssetEngine/Public/ChaosClothAsset/ClothComponent.h b/Engine/Plugins/ChaosClothAsset/Source/ChaosClothAssetEngine/Public/ChaosClothAsset/ClothComponent.h
--- a/Engine/Plugins/ChaosClothAsset/Source/ChaosClothAssetEngine/Public/ChaosClothAsset/ClothComponent.h
+++ b/Engine/Plugins/ChaosClothAsset/Source/ChaosClothAssetEngine/Public/ChaosClothAsset/ClothComponent.h
@@ -369,6 +369,7 @@ protected:
 		checkf(GetSimulationGroupId(InAsset, ModelIndex) != INDEX_NONE, TEXT("Invalid arguments"));
 		return GetPropertyCollections(ModelIndex);  // There's only this one asset for now, therefore this is the same as calling GetPropertyCollections(ModelIndex)
 	}
+	UE_API virtual void ClearPropertyCollectionsDirtyFlags(int32 LODIndex) const override;
 	UE_API virtual bool HasAnySimulationMeshData(int32 LODIndex) const override;
 	//~ End IClothComponentAdapter Interface
 
diff --git a/Engine/Plugins/ChaosClothAsset/Source/ChaosClothAssetEngine/Public/ChaosClothAsset/ClothComponentAdapter.h b/Engine/Plugins/ChaosClothAsset/Source/ChaosClothAssetEngine/Public/ChaosClothAsset/ClothComponentAdapter.h
--- a/Engine/Plugins/ChaosClothAsset/Source/ChaosClothAssetEngine/Public/ChaosClothAsset/ClothComponentAdapter.h
+++ b/Engine/Plugins/ChaosClothAsset/Source/ChaosClothAssetEngine/Public/ChaosClothAsset/ClothComponentAdapter.h
@@ -69,6 +69,9 @@ namespace UE::Chaos::ClothAsset
 		/** Return the runtime solver property collections. */
 		virtual const TArray<TSharedPtr<const FManagedArrayCollection>>& GetSolverPropertyCollections() const = 0;
 
+		/** Clear all property collections dirty flags. This is handled by the simulation proxy once the collections have been copied. */
+		virtual void ClearPropertyCollectionsDirtyFlags(int32 LodIndex) const {}
+
 		/** Return true when any of the simulated assets LOD models have valid simulation data, or false otherwise. */
 		virtual bool HasAnySimulationMeshData(int32 LodIndex) const = 0;
 
diff --git a/Engine/Source/Runtime/Experimental/Chaos/Private/Chaos/CollectionPropertyFacade.cpp b/Engine/Source/Runtime/Experimental/Chaos/Private/Chaos/CollectionPropertyFacade.cpp
--- a/Engine/Source/Runtime/Experimental/Chaos/Private/Chaos/CollectionPropertyFacade.cpp
+++ b/Engine/Source/Runtime/Experimental/Chaos/Private/Chaos/CollectionPropertyFacade.cpp
@@ -14,6 +14,8 @@ namespace Chaos::Softs
 		static const FName HighValueName("HighValue");  // Boolean, 24 bit integer (max 16777215), float, or vector value of the highest weight on the weight map if any
 		static const FName StringValueName("StringValue");  // String value, or weight map name, ...etc.
 		static const FName FlagsName("Flags");  // Whether this property is enabled, animatable, ...etc.
+		static const FName PropertyGlobalGroup("PropertyGlobal");  // Contains all single elements of this facade
+		static const FName DirtyName("Dirty");  // Dirty flags for this facade at the collection level
 
 		// Previous name for String-based Key
 		static const FName KeyName("Key");
@@ -42,7 +44,18 @@ namespace Chaos::Softs
 			ManagedArrayCollection->HasAttribute(PropertyFacadeNames::LowValueName, PropertyFacadeNames::PropertyGroup) &&
 			ManagedArrayCollection->HasAttribute(PropertyFacadeNames::HighValueName, PropertyFacadeNames::PropertyGroup) &&
 			ManagedArrayCollection->HasAttribute(PropertyFacadeNames::StringValueName, PropertyFacadeNames::PropertyGroup) &&
-			ManagedArrayCollection->HasAttribute(PropertyFacadeNames::FlagsName, PropertyFacadeNames::PropertyGroup);
+			ManagedArrayCollection->HasAttribute(PropertyFacadeNames::FlagsName, PropertyFacadeNames::PropertyGroup) &&
+			ManagedArrayCollection->HasAttribute(PropertyFacadeNames::DirtyName, PropertyFacadeNames::PropertyGlobalGroup);
+	}
+
+	bool FCollectionPropertyConstFacade::IsDirty() const
+	{
+		if (const TManagedArray<bool>* const DirtyAttribute = ManagedArrayCollection->FindAttributeTyped<bool>(PropertyFacadeNames::DirtyName, PropertyFacadeNames::PropertyGlobalGroup);
+			DirtyAttribute && DirtyAttribute->Num())
+		{
+			return (*DirtyAttribute)[0];
+		}
+		return false;
 	}
 
 	void FCollectionPropertyConstFacade::UpdateArrays()
@@ -117,41 +130,97 @@ namespace Chaos::Softs
 	{
 	}
 
+	TArrayView<FName> FCollectionPropertyFacade::ModifyKeyNameArray()
+	{
+		MakeDirty();
+		TManagedArray<FName>& KeyNames = GetManagedArrayCollection()->ModifyAttribute<FName>(
+			PropertyFacadeNames::KeyNameName, PropertyFacadeNames::PropertyGroup);
+		return TArrayView<FName>(KeyNames.GetData(), KeyNames.Num());
+	}
+
+	TArrayView<FVector3f> FCollectionPropertyFacade::ModifyLowValueArray()
+	{
+		MakeDirty();
+		TManagedArray<FVector3f>& LowValues = GetManagedArrayCollection()->ModifyAttribute<FVector3f>(
+			PropertyFacadeNames::LowValueName, PropertyFacadeNames::PropertyGroup);
+		return TArrayView<FVector3f>(LowValues.GetData(), LowValues.Num());
+	}
+
+	TArrayView<FVector3f> FCollectionPropertyFacade::ModifyHighValueArray()
+	{
+		MakeDirty();
+		TManagedArray<FVector3f>& HighValues = GetManagedArrayCollection()->ModifyAttribute<FVector3f>(
+			PropertyFacadeNames::HighValueName, PropertyFacadeNames::PropertyGroup);
+		return TArrayView<FVector3f>(HighValues.GetData(), HighValues.Num());
+	}
+
+	TArrayView<FString> FCollectionPropertyFacade::ModifyStringValueArray()
+	{
+		MakeDirty();
+		TManagedArray<FString>& StringValues = GetManagedArrayCollection()->ModifyAttribute<FString>(
+			PropertyFacadeNames::StringValueName, PropertyFacadeNames::PropertyGroup);
+		return TArrayView<FString>(StringValues.GetData(), StringValues.Num());
+	}
+
+	TArrayView<ECollectionPropertyFlags> FCollectionPropertyFacade::ModifyFlagsArray()
+	{
+		MakeDirty();
+		TManagedArray<uint8>& FlagBytes = GetManagedArrayCollection()->ModifyAttribute<uint8>(
+			PropertyFacadeNames::FlagsName, PropertyFacadeNames::PropertyGroup);
+		return TArrayView<ECollectionPropertyFlags>(
+			reinterpret_cast<ECollectionPropertyFlags*>(FlagBytes.GetData()), FlagBytes.Num());
+	}
+
 	void FCollectionPropertyFacade::SetFlags(int32 KeyIndex, ECollectionPropertyFlags Flags)
 	{
 		// Cannot set string dirty without also dirtying the property
 		Flags |= EnumHasAnyFlags(Flags, ECollectionPropertyFlags::StringDirty) ? ECollectionPropertyFlags::Dirty : ECollectionPropertyFlags::None;
 		// Cannot remove the Dirty, StringDirty, and Intrinsic flags
-		Flags |= GetFlagsArray()[KeyIndex] & (ECollectionPropertyFlags::Dirty | ECollectionPropertyFlags::StringDirty | ECollectionPropertyFlags::Intrinsic | ECollectionPropertyFlags::Interpolable);
+		Flags |= (ECollectionPropertyFlags)FlagsArray[KeyIndex] & (ECollectionPropertyFlags::Dirty | ECollectionPropertyFlags::StringDirty | ECollectionPropertyFlags::Intrinsic | ECollectionPropertyFlags::Interpolable);
 
-		SetValue(KeyIndex, GetFlagsArray(), Flags);
+		SetValue(KeyIndex, ModifyFlagsArray(), Flags);
 	}
 
 	void FCollectionPropertyFacade::ClearDirtyFlags()
 	{
-		for (ECollectionPropertyFlags& Flags : GetFlagsArray())
+		for (ECollectionPropertyFlags& Flags : ModifyFlagsArray())
 		{
 			EnumRemoveFlags(Flags, ECollectionPropertyFlags::StringDirty | ECollectionPropertyFlags::Dirty);
 		}
+		if (TManagedArray<bool>* const DirtyAttribute = GetManagedArrayCollection()->FindAttributeTyped<bool>(PropertyFacadeNames::DirtyName, PropertyFacadeNames::PropertyGlobalGroup);
+			DirtyAttribute && DirtyAttribute->Num())
+		{
+			(*DirtyAttribute)[0] = false;
+		}
+	}
+
+	void FCollectionPropertyFacade::MakeDirty()
+	{
+		if (TManagedArray<bool>* const DirtyAttribute = GetManagedArrayCollection()->FindAttributeTyped<bool>(PropertyFacadeNames::DirtyName, PropertyFacadeNames::PropertyGlobalGroup);
+			DirtyAttribute && DirtyAttribute->Num())
+		{
+			(*DirtyAttribute)[0] = true;
+		}
 	}
 
 	void FCollectionPropertyFacade::EnableFlags(int32 KeyIndex, ECollectionPropertyFlags Flags, bool bEnable)
 	{
-		const ECollectionPropertyFlags CurrentFlags = GetFlagsArray()[KeyIndex];
+		const ECollectionPropertyFlags CurrentFlags = (ECollectionPropertyFlags)FlagsArray[KeyIndex];
 
 		if (bEnable)
 		{
 			if (!EnumHasAllFlags(CurrentFlags, Flags))
 			{
-				EnumAddFlags(GetFlagsArray()[KeyIndex], Flags | ECollectionPropertyFlags::Dirty);  // Changing any flags adds the dirty flag
+				EnumAddFlags(ModifyFlagsArray()[KeyIndex], Flags | ECollectionPropertyFlags::Dirty);  // Changing any flags adds the dirty flag
 			}
 		}
 		else
 		{
 			if (EnumHasAnyFlags(CurrentFlags, Flags))
 			{
-				EnumRemoveFlags(GetFlagsArray()[KeyIndex], Flags);
-				EnumAddFlags(GetFlagsArray()[KeyIndex], ECollectionPropertyFlags::Dirty);  // Changing any flags adds the dirty flag, dirty flags can only be removed from the ClearDirtyFlags function
+				TArrayView<ECollectionPropertyFlags> MutableFlagsArray = ModifyFlagsArray();
+				EnumRemoveFlags(MutableFlagsArray[KeyIndex], Flags);
+				EnumAddFlags(MutableFlagsArray[KeyIndex], ECollectionPropertyFlags::Dirty);  // Changing any flags adds the dirty flag, dirty flags can only be removed from the ClearDirtyFlags function
 			}
 		}
 	}
@@ -195,6 +264,14 @@ namespace Chaos::Softs
 		GetManagedArrayCollection()->AddAttribute<FString>(PropertyFacadeNames::StringValueName, PropertyFacadeNames::PropertyGroup);
 		GetManagedArrayCollection()->AddAttribute<uint8>(PropertyFacadeNames::FlagsName, PropertyFacadeNames::PropertyGroup);
 
+		// Add the global group and check whether it has one single element, otherwise add it and make the property collection dirty
+		TManagedArray<bool>& DirtyAttribute = GetManagedArrayCollection()->AddAttribute<bool>(PropertyFacadeNames::DirtyName, PropertyFacadeNames::PropertyGlobalGroup);
+		if (!GetManagedArrayCollection()->NumElements(PropertyFacadeNames::PropertyGlobalGroup))
+		{
+			GetManagedArrayCollection()->AddElements(1, PropertyFacadeNames::PropertyGlobalGroup);
+			DirtyAttribute.Fill(true);
+		}
+
 		UpdateArrays();
 		RebuildKeyIndices();
 	}
@@ -216,9 +293,9 @@ namespace Chaos::Softs
 		UpdateArrays();
 
 		// Setup the new element's default value and enable the property by default
-		GetKeyNameArray()[Index] = Key;
-		GetLowValueArray()[Index] = GetHighValueArray()[Index] = FVector3f::ZeroVector;
-		GetFlagsArray()[Index] = Flags;
+		ModifyKeyNameArray()[Index] = Key;
+		ModifyLowValueArray()[Index] = ModifyHighValueArray()[Index] = FVector3f::ZeroVector;
+		ModifyFlagsArray()[Index] = Flags;
 
 		// Update search map
 		KeyNameIndices.Emplace(KeyNameArray[Index], Index);
@@ -248,9 +325,9 @@ namespace Chaos::Softs
 			for (int32 Index = StartIndex; Index < NumProperties + StartIndex; ++Index)
 			{
 				// Setup the new elements' default value and enable the property by default
-				GetKeyNameArray()[Index] = Keys[Index - StartIndex];
-				GetLowValueArray()[Index] = GetHighValueArray()[Index] = FVector3f::ZeroVector;
-				GetFlagsArray()[Index] = Flags;
+				ModifyKeyNameArray()[Index] = Keys[Index - StartIndex];
+				ModifyLowValueArray()[Index] = ModifyHighValueArray()[Index] = FVector3f::ZeroVector;
+				ModifyFlagsArray()[Index] = Flags;
 
 				// Update search map
 				KeyNameIndices.Emplace(KeyNameArray[Index], Index);
@@ -274,7 +351,8 @@ namespace Chaos::Softs
 			InManagedArrayCollection.HasAttribute(PropertyFacadeNames::LowValueName, PropertyFacadeNames::PropertyGroup) &&
 			InManagedArrayCollection.HasAttribute(PropertyFacadeNames::HighValueName, PropertyFacadeNames::PropertyGroup) &&
 			InManagedArrayCollection.HasAttribute(PropertyFacadeNames::StringValueName, PropertyFacadeNames::PropertyGroup) &&
-			InManagedArrayCollection.HasAttribute(PropertyFacadeNames::FlagsName, PropertyFacadeNames::PropertyGroup))
+			InManagedArrayCollection.HasAttribute(PropertyFacadeNames::FlagsName, PropertyFacadeNames::PropertyGroup) &&
+			InManagedArrayCollection.HasAttribute(PropertyFacadeNames::DirtyName, PropertyFacadeNames::PropertyGlobalGroup))
 		{
 			// Compare keys
 			bool bKeyIndicesMatch = false;
@@ -306,6 +384,7 @@ namespace Chaos::Softs
 			{
 				RebuildKeyIndices();
 			}
+			GetManagedArrayCollection()->CopyAttribute(InManagedArrayCollection, PropertyFacadeNames::DirtyName, PropertyFacadeNames::PropertyGlobalGroup);
 		}
 		else
 		{
@@ -441,13 +520,20 @@ namespace Chaos::Softs
 					check(InKeyArray->Num() == KeyNameArray.Num());
 					for (int32 Index = 0; Index < KeyNameArray.Num(); ++Index)
 					{
-						GetKeyNameArray()[Index] = FName((*InKeyArray)[Index]);
+						ModifyKeyNameArray()[Index] = FName((*InKeyArray)[Index]);
 					}
 					RebuildKeyIndices();
 				}
 
 				GetManagedArrayCollection()->RemoveAttribute(PropertyFacadeNames::KeyName, PropertyFacadeNames::PropertyGroup);
 			}
+			// Add Dirty attribute if missing
+			TManagedArray<bool>& DirtyAttribute = GetManagedArrayCollection()->AddAttribute<bool>(PropertyFacadeNames::DirtyName, PropertyFacadeNames::PropertyGlobalGroup);
+			if (!GetManagedArrayCollection()->NumElements(PropertyFacadeNames::PropertyGlobalGroup))
+			{
+				GetManagedArrayCollection()->AddElements(1, PropertyFacadeNames::PropertyGlobalGroup);
+				DirtyAttribute.Fill(true);
+			}
 		}
 	}
 }  // End namespace Chaos::Softs
diff --git a/Engine/Source/Runtime/Experimental/Chaos/Public/Chaos/CollectionPropertyFacade.h b/Engine/Source/Runtime/Experimental/Chaos/Public/Chaos/CollectionPropertyFacade.h
--- a/Engine/Source/Runtime/Experimental/Chaos/Public/Chaos/CollectionPropertyFacade.h
+++ b/Engine/Source/Runtime/Experimental/Chaos/Public/Chaos/CollectionPropertyFacade.h
@@ -61,6 +61,9 @@ namespace Chaos::Softs
 		/** Return whether the facade is defined on the collection. */
 		CHAOS_API bool IsValid() const;
 
+		/** Return whether this collection facade is dirty (does not rely on individual property flags). */
+		CHAOS_API bool IsDirty() const;
+
 		/** Return the number of properties in this collection. */
 		int32 Num() const { return KeyNameIndices.Num(); }
 
@@ -353,14 +356,6 @@ namespace Chaos::Softs
 
 		bool HasAnyFlags(int32 KeyIndex, ECollectionPropertyFlags Flags) const { return EnumHasAnyFlags(GetValue<ECollectionPropertyFlags, uint8>(KeyIndex, FlagsArray), Flags); }
 
-		// Attribute groups, predefined data member of the collection
-		static CHAOS_API const FName PropertyGroup;
-		static CHAOS_API const FName KeyNameName;  // Property key, name to look for
-		static CHAOS_API const FName LowValueName;  // Boolean, 24 bit integer (max 16777215), float, or vector value, or value of the lowest weight on the weight map if any
-		static CHAOS_API const FName HighValueName;  // Boolean, 24 bit integer (max 16777215), float, or vector value of the highest weight on the weight map if any
-		static CHAOS_API const FName StringValueName;  // String value, or weight map name, ...etc.
-		static CHAOS_API const FName FlagsName;  // Whether this property is enabled, animatable, ...etc.
-
 		// Property Group array views
 		TConstArrayView<FName> KeyNameArray;
 		TConstArrayView<FVector3f> LowValueArray;
@@ -404,10 +399,10 @@ namespace Chaos::Softs
 
 		//~ Values set per index
 		template<typename T UE_REQUIRES(TIsWeightedType<T>::Value)>
-		void SetLowValue(int32 KeyIndex, const T& Value) { SetValue(KeyIndex, GetLowValueArray(), FVector3f(Value)); }
+		void SetLowValue(int32 KeyIndex, const T& Value) { SetValue(KeyIndex, ModifyLowValueArray(), FVector3f(Value)); }
 
 		template<typename T UE_REQUIRES(TIsWeightedType<T>::Value)>
-		void SetHighValue(int32 KeyIndex, const T& Value) { SetValue(KeyIndex, GetHighValueArray(), FVector3f(Value)); }
+		void SetHighValue(int32 KeyIndex, const T& Value) { SetValue(KeyIndex, ModifyHighValueArray(), FVector3f(Value)); }
 
 		template<typename T UE_REQUIRES(TIsWeightedType<T>::Value)>
 		void SetWeightedValue(int32 KeyIndex, const T& LowValue, const T& HighValue) { SetLowValue(KeyIndex, LowValue); SetHighValue(KeyIndex, HighValue); }
@@ -417,7 +412,7 @@ namespace Chaos::Softs
 		template<typename T UE_REQUIRES(TIsWeightedType<T>::Value)>
 		void SetValue(int32 KeyIndex, const T& Value) { SetWeightedValue(KeyIndex, Value, Value); }
 
-		void SetStringValue(int32 KeyIndex, const FString& Value) { if (GetStringValueArray()[KeyIndex] != Value) { GetStringValueArray()[KeyIndex] = Value; SetStringDirty(KeyIndex); } }
+		void SetStringValue(int32 KeyIndex, const FString& Value) { if (StringValueArray[KeyIndex] != Value) { ModifyStringValueArray()[KeyIndex] = Value; SetStringDirty(KeyIndex); } }
 
 		/** SetFlags cannot be used to remove Dirty, StringDirty, Interpolable or Intrinsic flags. Use ClearDirtyFlags to remove dirty flags. */
 		CHAOS_API void SetFlags(int32 KeyIndex, ECollectionPropertyFlags Flags);
@@ -608,7 +603,12 @@ namespace Chaos::Softs
 			return SafeSet(FName(Key), [this](int32 KeyIndex) { SetInterpolable(KeyIndex); });
 		}
 
+		/** Clear all individual properties' dirty flags and this collection facade dirty flag. */
 		CHAOS_API void ClearDirtyFlags();
+
+		/** Make this collection facade dirty (does not change individual property flags). */
+		CHAOS_API void MakeDirty();
+
 		CHAOS_API void UpdateProperties(const TSharedPtr<const FManagedArrayCollection>& InManagedArrayCollection);
 
 	protected:
@@ -617,10 +617,23 @@ namespace Chaos::Softs
 
 		// Access to a writeable ManagedArrayCollection is protected, use an FPropertyCollectionMutableAdapter if needed to get a non const pointer
 		TSharedPtr<FManagedArrayCollection> GetManagedArrayCollection() { return ConstCastSharedPtr<FManagedArrayCollection>(ManagedArrayCollection); }
+
+		// Mutable array accessors via ModifyAttribute (marks collection dirty)
+		CHAOS_API TArrayView<FName> ModifyKeyNameArray();
+		CHAOS_API TArrayView<FVector3f> ModifyLowValueArray();
+		CHAOS_API TArrayView<FVector3f> ModifyHighValueArray();
+		CHAOS_API TArrayView<FString> ModifyStringValueArray();
+		CHAOS_API TArrayView<ECollectionPropertyFlags> ModifyFlagsArray();
+
+		UE_DEPRECATED(5.8, "Use ModifyKeyNameArray instead.")
 		const TArrayView<FName>& GetKeyNameArray() { return reinterpret_cast<TArrayView<FName>&>(KeyNameArray); }
+		UE_DEPRECATED(5.8, "Use ModifyLowValueArray instead.")
 		const TArrayView<FVector3f>& GetLowValueArray() { return reinterpret_cast<TArrayView<FVector3f>&>(LowValueArray); }
+		UE_DEPRECATED(5.8, "Use ModifyHighValueArray instead.")
 		const TArrayView<FVector3f>& GetHighValueArray() { return reinterpret_cast<TArrayView<FVector3f>&>(HighValueArray); }
+		UE_DEPRECATED(5.8, "Use ModifyStringValueArray instead.")
 		const TArrayView<FString>& GetStringValueArray() { return reinterpret_cast<TArrayView<FString>&>(StringValueArray); }
+		UE_DEPRECATED(5.8, "Use ModifyFlagsArray instead.")
 		const TArrayView<ECollectionPropertyFlags>& GetFlagsArray() { return reinterpret_cast<const TArrayView<ECollectionPropertyFlags>&>(FlagsArray); }
 
 		UE_DEPRECATED(5.7, "Use GetKeyNameArray instead")
diff --git a/Engine/Source/Runtime/Experimental/Chaos/Public/GeometryCollectionProxyData.h b/Engine/Source/Runtime/Experimental/Chaos/Public/GeometryCollectionProxyData.h
--- a/Engine/Source/Runtime/Experimental/Chaos/Public/GeometryCollectionProxyData.h
+++ b/Engine/Source/Runtime/Experimental/Chaos/Public/GeometryCollectionProxyData.h
@@ -171,12 +171,23 @@ public:
 
 	CHAOS_API void CopyInitialVelocityAttributesFrom(const FGeometryDynamicCollection& SourceCollection);
 
+	void MakeDirty() { bDirty = true; }
+	void MakeClean() { bDirty = false; }
+	bool IsDirty() const { return bDirty; }
+
+	virtual void Reset() override
+	{
+		Super::Reset();
+		MakeDirty();
+	}
+
 private:
 	TManagedArray<uint8>	  InternalClusterParentType;
 
 	TManagedArray<FVector3f>* OptionalLinearVelocityAttribute;
 	TManagedArray<FVector3f>* OptionalAngularVelocityAttribute;
 	TManagedArray<bool>*	  OptionalAnimateTransformAttribute;
+	bool bDirty = false;
 };
 
 /**
