IAD - Implement seperate stats for each CasType.
- The OpenMapped path now updates stats by calling FOnDemandInstallCacheStats::OnReadCompleted()
- FChunkRequest is now always and only detached from the IoRequest in CompleteRequest(). This makes its lifetime easier to reason about.
- Moved TotalFragmentionRatio and DefraggedBlockFragmentionRatio from installer analytics to cache analytics, since they are actually cache analytics and need to be tracked per CAS.
- Removed OldestBlockAccess analytic since it wasn't usefull and we have replaced it with immediate analytics for block access.
- Added a shared ResolveErrorCount analytic for read errors that happen before a CAS has been determined.  ReadErrorCount is now per CAS and tracks errors after a CAS has been determined.
- Namespaced counters and events with CasType where needed.  Analytcis for the General CAS are not namespaced to preserve compatibility with existing analytics.

#Synced-CL 52147163

diff --git a/Engine/Source/Runtime/Experimental/IoStore/OnDemand/Private/OnDemandInstallCache.cpp b/Engine/Source/Runtime/Experimental/IoStore/OnDemand/Private/OnDemandInstallCache.cpp
--- a/Engine/Source/Runtime/Experimental/IoStore/OnDemand/Private/OnDemandInstallCache.cpp
+++ b/Engine/Source/Runtime/Experimental/IoStore/OnDemand/Private/OnDemandInstallCache.cpp
@@ -111,7 +111,7 @@ using FWeakMappedFileHandle			= TWeakPtr<IMappedFileHandle>;
 class FSharedMappedFileHandleProxy final : public IMappedFileHandle
 {
 public:
-	FSharedMappedFileHandleProxy(FSharedMappedFileHandleRef&& InSharedMappedFileHandle)
+	FSharedMappedFileHandleProxy(EOnDemandInstallCasType InCasType, FSharedMappedFileHandleRef&& InSharedMappedFileHandle)
 		: IMappedFileHandle(InSharedMappedFileHandle->GetFileSize())
 		, SharedMappedFileHandle(MoveTemp(InSharedMappedFileHandle))
 	{}
@@ -125,12 +125,18 @@ public:
 
 		if (EnumHasAnyFlags(Flags.Flags, EMappedFileFlags::EFileWritable))
 		{
+			FOnDemandInstallCacheStats::OnReadCompleted(CasType, EIoErrorCode::ReadError, 0);
+
 			// Only allow read access through the mapped handle.
 			checkf(false, TEXT("Writes to the OnDeamnd cache are only allowed from the OnDemandContentInstaller!"));
 			return nullptr;
 		}
 
-		return SharedMappedFileHandle->MapRegion(Offset, BytesToMap, Flags);
+		IMappedFileRegion* Region = SharedMappedFileHandle->MapRegion(Offset, BytesToMap, Flags);
+
+		FOnDemandInstallCacheStats::OnReadCompleted(CasType, Region ? EIoErrorCode::Ok : EIoErrorCode::ReadError, BytesToMap);
+
+		return Region;
 	}
 
 	virtual void Flush() override
@@ -140,6 +146,7 @@ public:
 
 private:
 	FSharedMappedFileHandleRef SharedMappedFileHandle;
+	EOnDemandInstallCasType CasType = EOnDemandInstallCasType::None;
 };
 
 
@@ -594,7 +601,7 @@ FResult FCasJournal::Commit(FTransaction&& Transaction, uint64& OutByteCount, ui
 				ECasErrorCode::CommitJournalFailed, EIoErrorCode::FileOpenFailed, 
 				FString::Printf(TEXT("Failed to open journal '%s'"), *Transaction.JournalFile),
 				OpenReadResult.StealError());
-			FOnDemandInstallCacheStats::OnJournalCommit(Result, 0);
+			FOnDemandInstallCacheStats::OnJournalCommit(Transaction.CasType, Result, 0);
 			return Result;
 		}
 
@@ -605,7 +612,7 @@ FResult FCasJournal::Commit(FTransaction&& Transaction, uint64& OutByteCount, ui
 		{
 			FResult Result = MakeJournalError(
 				ECasErrorCode::CommitJournalFailed, EIoErrorCode::FileCorrupt, FString::Printf(TEXT("Failed to validate journal header in '%s'"), *Transaction.JournalFile));
-			FOnDemandInstallCacheStats::OnJournalCommit(Result, 0);
+			FOnDemandInstallCacheStats::OnJournalCommit(Transaction.CasType, Result, 0);
 			return Result;
 		}
 
@@ -614,7 +621,7 @@ FResult FCasJournal::Commit(FTransaction&& Transaction, uint64& OutByteCount, ui
 		{
 			FResult Result = MakeJournalError(
 				ECasErrorCode::CommitJournalFailed, EIoErrorCode::ReadError, FString::Printf(TEXT("Failed to validate journal header in '%s'"), *Transaction.JournalFile));
-			FOnDemandInstallCacheStats::OnJournalCommit(Result, 0);
+			FOnDemandInstallCacheStats::OnJournalCommit(Transaction.CasType, Result, 0);
 			return Result;
 		}
 
@@ -623,7 +630,7 @@ FResult FCasJournal::Commit(FTransaction&& Transaction, uint64& OutByteCount, ui
 		{
 			FResult Result = MakeJournalError(
 				ECasErrorCode::CommitJournalFailed, EIoErrorCode::FileSeekFailed, FString::Printf(TEXT("Failed to validate journal footer in '%s'"), *Transaction.JournalFile));
-			FOnDemandInstallCacheStats::OnJournalCommit(Result, 0);
+			FOnDemandInstallCacheStats::OnJournalCommit(Transaction.CasType, Result, 0);
 			return Result;
 		}
 
@@ -632,7 +639,7 @@ FResult FCasJournal::Commit(FTransaction&& Transaction, uint64& OutByteCount, ui
 		{
 			FResult Result = MakeJournalError(
 				ECasErrorCode::CommitJournalFailed, EIoErrorCode::ReadError, FString::Printf(TEXT("Found to validate journal footer in '%s'"), *Transaction.JournalFile));
-			FOnDemandInstallCacheStats::OnJournalCommit(Result, 0);
+			FOnDemandInstallCacheStats::OnJournalCommit(Transaction.CasType, Result, 0);
 			return Result;
 		}
 	}
@@ -650,7 +657,7 @@ FResult FCasJournal::Commit(FTransaction&& Transaction, uint64& OutByteCount, ui
 				ECasErrorCode::CommitJournalFailed, EIoErrorCode::FileOpenFailed, 
 				FString::Printf(TEXT("Failed to open journal '%s'"), *Transaction.JournalFile),
 				OpenWriteResult.HasError() ? OpenWriteResult.StealError() : FFileSystemError(FString()));
-			FOnDemandInstallCacheStats::OnJournalCommit(Result, 0);
+			FOnDemandInstallCacheStats::OnJournalCommit(Transaction.CasType, Result, 0);
 			return Result;
 		}
 
@@ -658,7 +665,7 @@ FResult FCasJournal::Commit(FTransaction&& Transaction, uint64& OutByteCount, ui
 		{
 			FResult Result = MakeJournalError(
 				ECasErrorCode::CommitJournalFailed, EIoErrorCode::FileSeekFailed, FString::Printf(TEXT("Failed to seek to journal entries offset %" INT64_FMT " in '%s'"), EntriesPos, *Transaction.JournalFile));
-			FOnDemandInstallCacheStats::OnJournalCommit(Result, 0);
+			FOnDemandInstallCacheStats::OnJournalCommit(Transaction.CasType, Result, 0);
 			return Result;
 		}
 
@@ -669,7 +676,7 @@ FResult FCasJournal::Commit(FTransaction&& Transaction, uint64& OutByteCount, ui
 		{
 			FResult Result = MakeJournalError(
 				ECasErrorCode::CommitJournalFailed, EIoErrorCode::WriteError, FString::Printf(TEXT("Failed to write journal entries to '%s'"), *Transaction.JournalFile));
-			FOnDemandInstallCacheStats::OnJournalCommit(Result, 0);
+			FOnDemandInstallCacheStats::OnJournalCommit(Transaction.CasType, Result, 0);
 			return Result;
 		}
 		OutOpCount++;
@@ -681,7 +688,7 @@ FResult FCasJournal::Commit(FTransaction&& Transaction, uint64& OutByteCount, ui
 		{
 			FResult Result = MakeJournalError(
 				ECasErrorCode::CommitJournalFailed, EIoErrorCode::WriteError, FString::Printf(TEXT("Failed to write journal footer to '%s'"), *Transaction.JournalFile));
-			FOnDemandInstallCacheStats::OnJournalCommit(Result, 0);
+			FOnDemandInstallCacheStats::OnJournalCommit(Transaction.CasType, Result, 0);
 			return Result;
 		}
 		OutOpCount++;
@@ -691,7 +698,7 @@ FResult FCasJournal::Commit(FTransaction&& Transaction, uint64& OutByteCount, ui
 		{
 			FResult Result = MakeJournalError(
 				ECasErrorCode::CommitJournalFailed, EIoErrorCode::FileFlushFailed, FString::Printf(TEXT("Failed to flush journal entries to '%s'"), *Transaction.JournalFile));
-			FOnDemandInstallCacheStats::OnJournalCommit(Result, 0);
+			FOnDemandInstallCacheStats::OnJournalCommit(Transaction.CasType, Result, 0);
 			return Result;
 		}
 		OutOpCount++;
@@ -700,7 +707,7 @@ FResult FCasJournal::Commit(FTransaction&& Transaction, uint64& OutByteCount, ui
 			Transaction.Entries.Num(), ToKiB(TotalEntrySize), *Transaction.JournalFile);
 
 		FResult Result = MakeValue();
-		FOnDemandInstallCacheStats::OnJournalCommit(Result, TotalEntrySize);
+		FOnDemandInstallCacheStats::OnJournalCommit(Transaction.CasType, Result, TotalEntrySize);
 		return Result;
 	}
 }
@@ -2446,9 +2453,9 @@ class FOnDemandInstallCache final
 			uint64 RequestedRawSize)
 				: DispatcherRequest(Request)
 				, ChunkInfo(MoveTemp(Info))
-				, ChunkRange(Range)
-				, EncodedChunk(ChunkRange.GetLength())
+				, EncodedChunk(Range.GetLength())
 				, RawSize(RequestedRawSize)
+				, ChunkRange(Range)
 		{
 			check(DispatcherRequest != nullptr);
 			check(ChunkInfo.IsValid());
@@ -2479,16 +2486,17 @@ class FOnDemandInstallCache final
 		{
 			void* ChunkRequest = nullptr;
 			Swap(ChunkRequest, Request.BackendData);
-			return TUniquePtr<FChunkRequest>(reinterpret_cast<FChunkRequest*>(ChunkRequest));
+			return TUniquePtr<FChunkRequest>(static_cast<FChunkRequest*>(ChunkRequest));
 		}
 
 		FSharedAsyncFileHandle			SharedFileHandle;
 		TUniquePtr<IAsyncReadRequest>	FileReadRequest;
-		FIoRequestImpl*					DispatcherRequest;
-		FOnDemandChunkInfo				ChunkInfo;
-		FIoOffsetAndLength				ChunkRange;
+		FIoRequestImpl*					DispatcherRequest = nullptr;
+		FOnDemandChunkInfo				ChunkInfo;		
 		FIoBuffer						EncodedChunk;
-		uint64							RawSize;
+		uint64							RawSize = 0;
+		FIoOffsetAndLength				ChunkRange;
+		EOnDemandInstallCasType			CasType = EOnDemandInstallCasType::None;
 	};
 
 	struct FPendingChunks
@@ -2916,10 +2924,6 @@ FIoRequestImpl* FOnDemandInstallCache::GetCompletedIoRequests()
 	FIoRequestImpl* FirstCompleted = nullptr;
 	{
 		UE::TUniqueLock Lock(Mutex);
-		for (FIoRequestImpl& Completed : CompletedRequests)
-		{
-			TUniquePtr<FChunkRequest> Detached = FChunkRequest::Detach(Completed);
-		}
 		FirstCompleted = CompletedRequests.GetHead();
 		CompletedRequests = FIoRequestList();
 	}
@@ -3023,6 +3027,7 @@ TIoStatusOr<FIoMappedRegion> FOnDemandInstallCache::OpenMapped(const FIoChunkId&
 				}
 			}
 #endif
+			FOnDemandInstallCacheStats::OnReadCompleted(EOnDemandInstallCasType::None, EIoErrorCode::NotInstalled, 0);
 			return FIoStatus(EIoErrorCode::NotInstalled);
 		}
 		return FIoStatus(EIoErrorCode::NotFound);
@@ -3041,7 +3046,8 @@ TIoStatusOr<FIoMappedRegion> FOnDemandInstallCache::OpenMapped(const FIoChunkId&
 
 	if (MaybeChunkRange.IsOk() == false)
 	{
-		UE_LOG(LogIoStoreOnDemand, Error, TEXT("Failed to get chunk range"));
+		UE_LOG(LogIoStoreOnDemand, Error, TEXT("Failed to get chunk range for OpenMapped"));
+		FOnDemandInstallCacheStats::OnReadCompleted(EOnDemandInstallCasType::None, MaybeChunkRange.Status().GetErrorCode(), 0);
 		return MaybeChunkRange.Status();
 	}
 
@@ -3074,6 +3080,7 @@ TIoStatusOr<FIoMappedRegion> FOnDemandInstallCache::OpenMapped(const FIoChunkId&
 		{
 			// Chunk is referenced but not in any cache. This should probably never happen currently.
 			// If pinning of unreferenced chunks was allowed, then this would be expected...
+			FOnDemandInstallCacheStats::OnReadCompleted(EOnDemandInstallCasType::None, EIoErrorCode::NotInstalled, 0);
 			return FIoStatus(EIoErrorCode::NotInstalled, TEXT("IoChunk should be installed but was not found!"));
 		}
 
@@ -3115,10 +3122,14 @@ TIoStatusOr<FIoMappedRegion> FOnDemandInstallCache::OpenMapped(const FIoChunkId&
 
 			if (const FCasErrorContext* Ctx = FileOpenResult.GetError().GetErrorContext<FCasErrorContext>())
 			{
-				return FIoStatus(Ctx->IoErrorCode, Ctx->SystemErrorCode, *LexToString(FileOpenResult.GetError()));
+				FIoStatus Status(Ctx->IoErrorCode, Ctx->SystemErrorCode, *LexToString(FileOpenResult.GetError()));
+				FOnDemandInstallCacheStats::OnReadCompleted(CacheState.Cas.GetType(), Status.GetErrorCode(), 0);
+				return Status;
 			}
 
-			return FIoStatus(EIoErrorCode::FileOpenFailed, *LexToString(FileOpenResult.GetError()));
+			FIoStatus Status(EIoErrorCode::FileOpenFailed, *LexToString(FileOpenResult.GetError()));
+			FOnDemandInstallCacheStats::OnReadCompleted(CacheState.Cas.GetType(), Status.GetErrorCode(), 0);
+			return Status;
 		}
 
 		break;
@@ -3135,12 +3146,16 @@ TIoStatusOr<FIoMappedRegion> FOnDemandInstallCache::OpenMapped(const FIoChunkId&
 	{
 		const FString Filename = CacheState.Cas.GetBlockFilename(CasLoc.BlockId);
 		UE_LOG(LogIoStoreOnDemand, Error, TEXT("Failed to map region for chunk in %s"), *Filename);
-		return FIoStatus(EIoErrorCode::ReadError, TEXT("Failed to map region"));
+		FIoStatus Status(EIoErrorCode::ReadError, TEXT("Failed to map region"));
+		FOnDemandInstallCacheStats::OnReadCompleted(CacheState.Cas.GetType(), Status.GetErrorCode(), 0);
+		return Status;
 	}
 
+	FOnDemandInstallCacheStats::OnReadCompleted(CacheState.Cas.GetType(), EIoErrorCode::Ok, ChunkRange.GetLength());
+
 	return FIoMappedRegion
 	{
-		.MappedFileHandle = new FSharedMappedFileHandleProxy(MappedHandle.ToSharedRef()),
+		.MappedFileHandle = new FSharedMappedFileHandleProxy(CacheState.Cas.GetType(), MappedHandle.ToSharedRef()),
 		.MappedFileRegion = MappedRegion
 	};
 }
@@ -3202,7 +3217,7 @@ bool FOnDemandInstallCache::Resolve(FIoRequestImpl* Request)
 	// TODO: Should this support GetTargetVa() ?
 
 	// The internal request parameters are attached/owned by the I/O request via
-	// the backend data parameter. The chunk request is deleted in GetCompletedRequests
+	// the backend data parameter. The chunk request is deleted in CompleteRequest
 	FChunkRequest::Attach(*Request, new FChunkRequest(
 		Request,
 		MoveTemp(ChunkInfo),
@@ -3223,7 +3238,6 @@ bool FOnDemandInstallCache::Resolve(FIoRequestImpl* Request)
 	{
 		// Chunk is referenced but not in any cache. This should probably never happen currently.
 		// If pinning of unreferenced chunks was allowed, then this would be expected...
-		TUniquePtr<FChunkRequest> Detached = FChunkRequest::Detach(*Request);
 		CompleteRequest(Request, EIoErrorCode::NotInstalled);
 	}
 
@@ -3245,10 +3259,11 @@ bool FOnDemandInstallCache::ResolveFromCas(FCacheState& CacheState, FIoRequestIm
 		const FIoBuffer* PendingChunk = CacheState.PendingChunks.FindChunk(ChunkRequest.ChunkInfo.Hash());
 		if (PendingChunk != nullptr)
 		{
+			ChunkRequest.CasType = CacheState.Cas.GetType();
+
 			if (bAllowResolve == false)
 			{
 				// Can't use the Resolve API to read this data
-				TUniquePtr<FChunkRequest> Detached = FChunkRequest::Detach(*Request);
 				CompleteRequest(Request, EIoErrorCode::FileOpenFailed);
 				return true;
 			}
@@ -3274,10 +3289,11 @@ bool FOnDemandInstallCache::ResolveFromCas(FCacheState& CacheState, FIoRequestIm
 		return false;
 	}
 
+	ChunkRequest.CasType = CacheState.Cas.GetType();
+
 	if (bAllowResolve == false)
 	{
 		// Can't use the Resolve API to read this data
-		TUniquePtr<FChunkRequest> Detached = FChunkRequest::Detach(*Request);
 		CompleteRequest(Request, EIoErrorCode::FileOpenFailed);
 		return true;
 	}
@@ -3356,7 +3372,6 @@ bool FOnDemandInstallCache::ResolveFromCas(FCacheState& CacheState, FIoRequestIm
 		const FString Filename = CacheState.Cas.GetBlockFilename(CasLoc.BlockId);
 		UE_LOG(LogIoStoreOnDemand, Error, TEXT("Failed to open CAS block '%s' in CAS %s for async reading, reason '%s'"),
 			*Filename, LexToString(CacheState.Cas.GetType()), *FileOpenResult.GetError().GetErrorMessage(true).ToString());
-		TUniquePtr<FChunkRequest> Detached = FChunkRequest::Detach(*Request);
 		CompleteRequest(Request, EIoErrorCode::FileOpenFailed);
 		return true;
 	}
@@ -3382,7 +3397,6 @@ bool FOnDemandInstallCache::ResolveFromCas(FCacheState& CacheState, FIoRequestIm
 	if (ChunkRequest.FileReadRequest.IsValid() == false)
 	{
 		TRACE_IOSTORE_BACKEND_REQUEST_FAILED(Request);
-		TUniquePtr<FChunkRequest> Detached = FChunkRequest::Detach(*Request);
 		CompleteRequest(Request, EIoErrorCode::ReadError);
 		return true;
 	}
@@ -3445,7 +3459,7 @@ TResult<bool> FOnDemandInstallCache::PutChunk(EOnDemandInstallCasType CasType, F
 
 	TUniqueLock Lock(CacheState->PendingChunks.SharedMutex);
 
-	FOnDemandInstallCacheStats::OnPutChunk(Chunk.GetSize());
+	FOnDemandInstallCacheStats::OnPutChunk(CasType, Chunk.GetSize());
 	CacheState->PendingChunks.Append(MoveTemp(Chunk), ChunkHash);
 	return MakeValue(true);
 }
@@ -3475,7 +3489,6 @@ FResult FOnDemandInstallCache::Purge(EOnDemandInstallCasType CasType, const uint
 	uint64 ReferencedBytes = 0;
 	uint64 FragmentedBytes = 0;
 	uint64 TotalReferencedBlockBytes = 0;
-	int64 OldestBlockAccess = FDateTime::MaxValue().GetTicks();
 
 	AddReferencesToBlocks(*CacheState, Containers, ChunkEntryIndices, BlockInfo, ReferencedBytes);
 	for (const TPair<FCasBlockId, FCasBlockInfo>& Kv : BlockInfo)
@@ -3504,16 +3517,11 @@ FResult FOnDemandInstallCache::Purge(EOnDemandInstallCasType CasType, const uint
 		{
 			TotalReferencedBlockBytes += Info.FileSize;
 		}
-
-		if (Info.LastAccess < OldestBlockAccess)
-		{
-			OldestBlockAccess = Info.LastAccess;
-		}
 	}
 
 	// TODO: IadMultiCAS: analytics need to track CAS type as well
 	FOnDemandInstallCacheStats::OnCacheUsage(
-		CacheState->MaxCacheSize, TotalCachedBytes, TotalReferencedBlockBytes, ReferencedBytes, FragmentedBytes, OldestBlockAccess);
+		CacheState->Cas.GetType(), CacheState->MaxCacheSize, TotalCachedBytes, TotalReferencedBlockBytes, ReferencedBytes, FragmentedBytes);
 
 	const uint64 TotalUncachedBytes = BytesToInstall + CacheState->PendingChunks.TotalSize;
 	const uint64 TotalRequiredBytes = TotalCachedBytes + TotalUncachedBytes;
@@ -3559,7 +3567,7 @@ FResult FOnDemandInstallCache::Purge(EOnDemandInstallCasType CasType, const uint
 		}
 	}
 
-	FOnDemandInstallCacheStats::OnPurge(Result, CacheState->MaxCacheSize, NewCachedBytes, TotalBytesToPurge, TotalPurgedBytes + DefragPurgedBytes);
+	FOnDemandInstallCacheStats::OnPurge(CacheState->Cas.GetType(), Result, CacheState->MaxCacheSize, NewCachedBytes, TotalBytesToPurge, TotalPurgedBytes + DefragPurgedBytes);
 	return Result;
 }
 
@@ -3630,7 +3638,7 @@ FResult FOnDemandInstallCache::PurgeAllUnreferenced(EOnDemandInstallCasType CasT
 		Result = Defrag(*CacheState, Containers, ChunkEntryIndices, BlockInfo);
 	}
 
-	FOnDemandInstallCacheStats::OnPurge(Result, CacheState->MaxCacheSize, NewCachedBytes, TotalBytesToPurge, TotalPurgedBytes);
+	FOnDemandInstallCacheStats::OnPurge(CacheState->Cas.GetType(), Result, CacheState->MaxCacheSize, NewCachedBytes, TotalBytesToPurge, TotalPurgedBytes);
 	return Result;
 }
 
@@ -3954,7 +3962,7 @@ FResult FOnDemandInstallCache::InitialVerify()
 				}
 			}
 
-			FOnDemandInstallCacheStats::OnCasVerification(Verify, RemovedChunks.Num(), TotalBlockBytesOverBudget);
+			FOnDemandInstallCacheStats::OnCasVerification(CacheState.Cas.GetType(), Verify, RemovedChunks.Num(), TotalBlockBytesOverBudget);
 
 			if (Verify.HasError())
 			{
@@ -4149,7 +4157,7 @@ FResult FOnDemandInstallCache::PurgeInternal(
 		}
 
 		constexpr bool bFromDefragFalse = false;
-		FOnDemandInstallCacheStats::OnBlockDeleted(Info.LastAccess, Info.LastModification, bFromDefragFalse);
+		FOnDemandInstallCacheStats::OnBlockDeleted(CacheState.Cas.GetType(), Info.LastAccess, Info.LastModification, bFromDefragFalse);
 
 		// This should be the only thread writing to CurrentBlock
 		FCasBlockId MaybeCurrentBlock = BlockId;
@@ -4306,7 +4314,7 @@ FResult FOnDemandInstallCache::Defrag(
 			FInstallCacheErrorContext Ctx = MakeInstallCacheErrorContext(CacheState, TotalCachedBytes, *TotalBytesToFree);
 			Ctx.IoError = FIoStatus(EIoErrorCode::OutOfDiskSpace, ErrorMessage);
 			FResult Result = MakeError(InstallCacheDefragError::MakeError(MoveTemp(Ctx), UE::UnifiedError::EDetailFilter::All));
-			FOnDemandInstallCacheStats::OnDefrag(Result, TotalFragmentedBytes, TotalCachedBytes, BlocksToDefragFragmentedSize, BlocksToDefragTotalSize);
+			FOnDemandInstallCacheStats::OnDefrag(CacheState.Cas.GetType(), Result, TotalFragmentedBytes, TotalCachedBytes, BlocksToDefragFragmentedSize, BlocksToDefragTotalSize);
 			return Result;
 		}
 	}
@@ -4373,7 +4381,7 @@ FResult FOnDemandInstallCache::Defrag(
 			}
 			
 			FResult Result = MakeValue();
-			FOnDemandInstallCacheStats::OnDefrag(Result, TotalFragmentedBytes, TotalCachedBytes, BlocksToDefragFragmentedSize, BlocksToDefragTotalSize);
+			FOnDemandInstallCacheStats::OnDefrag(CacheState.Cas.GetType(), Result, TotalFragmentedBytes, TotalCachedBytes, BlocksToDefragFragmentedSize, BlocksToDefragTotalSize);
 			return Result;
 		}
 	}
@@ -4406,7 +4414,7 @@ FResult FOnDemandInstallCache::Defrag(
 			UE_LOGFMT(LogIoStoreOnDemand, Error, "Failed to commit critical errors to CAS journal ({Error})", CommitResult.GetError());
 		}
 
-		FOnDemandInstallCacheStats::OnDefrag(Result, TotalFragmentedBytes, TotalCachedBytes, BlocksToDefragFragmentedSize, BlocksToDefragTotalSize);
+		FOnDemandInstallCacheStats::OnDefrag(CacheState.Cas.GetType(), Result, TotalFragmentedBytes, TotalCachedBytes, BlocksToDefragFragmentedSize, BlocksToDefragTotalSize);
 		return Result;
 	}
 
@@ -4473,7 +4481,7 @@ FResult FOnDemandInstallCache::Defrag(
 				FInstallCacheErrorContext Ctx = MakeInstallCacheErrorContext(CacheState, TotalCachedBytes, TotalBytesToFree ? *TotalBytesToFree : TotalFragmentedBytes);
 				FileOpenResult.GetError().PushErrorContext(MoveTemp(Ctx));
 				FResult Result = MakeError(FileOpenResult.StealError());
-				FOnDemandInstallCacheStats::OnDefrag(Result, TotalFragmentedBytes, TotalCachedBytes, BlocksToDefragFragmentedSize, BlocksToDefragTotalSize);
+				FOnDemandInstallCacheStats::OnDefrag(CacheState.Cas.GetType(), Result, TotalFragmentedBytes, TotalCachedBytes, BlocksToDefragFragmentedSize, BlocksToDefragTotalSize);
 				return Result;
 			}
 
@@ -4512,7 +4520,7 @@ FResult FOnDemandInstallCache::Defrag(
 					FInstallCacheErrorContext Ctx = MakeInstallCacheErrorContext(CacheState, TotalCachedBytes, TotalBytesToFree ? *TotalBytesToFree : TotalFragmentedBytes);
 					ReadResult.GetError().PushErrorContext(MoveTemp(Ctx));
 					FResult Result = MakeError(ReadResult.StealError());
-					FOnDemandInstallCacheStats::OnDefrag(Result, TotalFragmentedBytes, TotalCachedBytes, BlocksToDefragFragmentedSize, BlocksToDefragTotalSize);
+					FOnDemandInstallCacheStats::OnDefrag(CacheState.Cas.GetType(), Result, TotalFragmentedBytes, TotalCachedBytes, BlocksToDefragFragmentedSize, BlocksToDefragTotalSize);
 					return Result;
 				}
 
@@ -4532,7 +4540,7 @@ FResult FOnDemandInstallCache::Defrag(
 
 					if (FResult Result = FlushPendingChunks(CacheState, DefragPendingChunks, DefragBlock.LastAccess, DefragBlock.LastModification); Result.HasError())
 					{
-						FOnDemandInstallCacheStats::OnDefrag(Result, TotalFragmentedBytes, TotalCachedBytes, BlocksToDefragFragmentedSize, BlocksToDefragTotalSize);
+						FOnDemandInstallCacheStats::OnDefrag(CacheState.Cas.GetType(), Result, TotalFragmentedBytes, TotalCachedBytes, BlocksToDefragFragmentedSize, BlocksToDefragTotalSize);
 						return Result;
 					}
 					check(DefragPendingChunks.IsEmpty());
@@ -4544,7 +4552,7 @@ FResult FOnDemandInstallCache::Defrag(
 							.ExpectedHash	= ReffedChunk.Hash,
 							.ActualHash		= ChunkHash
 						}));
-					FOnDemandInstallCacheStats::OnDefrag(Result, TotalFragmentedBytes, TotalCachedBytes, BlocksToDefragFragmentedSize, BlocksToDefragTotalSize);
+					FOnDemandInstallCacheStats::OnDefrag(CacheState.Cas.GetType(), Result, TotalFragmentedBytes, TotalCachedBytes, BlocksToDefragFragmentedSize, BlocksToDefragTotalSize);
 					return Result;
 				}
 
@@ -4552,7 +4560,7 @@ FResult FOnDemandInstallCache::Defrag(
 				{
 					if (FResult Result = FlushPendingChunks(CacheState, DefragPendingChunks, DefragBlock.LastAccess, DefragBlock.LastModification); Result.HasError())
 					{
-						FOnDemandInstallCacheStats::OnDefrag(Result, TotalFragmentedBytes, TotalCachedBytes, BlocksToDefragFragmentedSize, BlocksToDefragTotalSize);
+						FOnDemandInstallCacheStats::OnDefrag(CacheState.Cas.GetType(), Result, TotalFragmentedBytes, TotalCachedBytes, BlocksToDefragFragmentedSize, BlocksToDefragTotalSize);
 						return Result;
 					}
 					check(DefragPendingChunks.IsEmpty());
@@ -4565,7 +4573,7 @@ FResult FOnDemandInstallCache::Defrag(
 
 			if (FResult Result = FlushPendingChunks(CacheState, DefragPendingChunks, DefragBlock.LastAccess, DefragBlock.LastModification); Result.HasError())
 			{
-				FOnDemandInstallCacheStats::OnDefrag(Result, TotalFragmentedBytes, TotalCachedBytes, BlocksToDefragFragmentedSize, BlocksToDefragTotalSize);
+				FOnDemandInstallCacheStats::OnDefrag(CacheState.Cas.GetType(), Result, TotalFragmentedBytes, TotalCachedBytes, BlocksToDefragFragmentedSize, BlocksToDefragTotalSize);
 				return Result;
 			}
 			check(DefragPendingChunks.IsEmpty());
@@ -4582,18 +4590,18 @@ FResult FOnDemandInstallCache::Defrag(
 		// Can now remove the old block
 		if (FResult Result = CacheState.Cas.DeleteBlock(DefragBlock.BlockId, HandleChunkRemoved); Result.HasError())
 		{
-			FOnDemandInstallCacheStats::OnDefrag(Result, TotalFragmentedBytes, TotalCachedBytes, BlocksToDefragFragmentedSize, BlocksToDefragTotalSize);
+			FOnDemandInstallCacheStats::OnDefrag(CacheState.Cas.GetType(), Result, TotalFragmentedBytes, TotalCachedBytes, BlocksToDefragFragmentedSize, BlocksToDefragTotalSize);
 			return Result;
 		}
 
 		constexpr bool bFromDefragTrue = true;
-		FOnDemandInstallCacheStats::OnBlockDeleted(DefragBlock.LastAccess, DefragBlock.LastModification, bFromDefragTrue);
+		FOnDemandInstallCacheStats::OnBlockDeleted(CacheState.Cas.GetType(), DefragBlock.LastAccess, DefragBlock.LastModification, bFromDefragTrue);
 
 		Transaction.BlockDeleted(DefragBlock.BlockId);
 
 		if (FResult Result = FCasJournal::Commit(MoveTemp(Transaction)); Result.HasError())
 		{
-			FOnDemandInstallCacheStats::OnDefrag(Result, TotalFragmentedBytes, TotalCachedBytes, BlocksToDefragFragmentedSize, BlocksToDefragTotalSize);
+			FOnDemandInstallCacheStats::OnDefrag(CacheState.Cas.GetType(), Result, TotalFragmentedBytes, TotalCachedBytes, BlocksToDefragFragmentedSize, BlocksToDefragTotalSize);
 			return Result;
 		}
 	}
@@ -4607,7 +4615,7 @@ FResult FOnDemandInstallCache::Defrag(
 	}
 
 	FResult Result = MakeValue();
-	FOnDemandInstallCacheStats::OnDefrag(Result, TotalFragmentedBytes, TotalCachedBytes, BlocksToDefragFragmentedSize, BlocksToDefragTotalSize);
+	FOnDemandInstallCacheStats::OnDefrag(CacheState.Cas.GetType(), Result, TotalFragmentedBytes, TotalCachedBytes, BlocksToDefragFragmentedSize, BlocksToDefragTotalSize);
 	return Result;
 }
 
@@ -4836,7 +4844,7 @@ FResult FOnDemandInstallCache::FlushPendingChunksImpl(FCacheState& CacheState, c
 			FInstallCacheErrorContext Ctx = MakeInstallCacheErrorContext(CacheState);
 			OpenWriteResult.GetError().PushErrorContext(MoveTemp(Ctx));
 			FResult Result = MakeError(OpenWriteResult.StealError());
-			FOnDemandInstallCacheStats::OnFlush(Result, TotalCasBytes, bFromDefrag);
+			FOnDemandInstallCacheStats::OnFlush(CacheState.Cas.GetType(), Result, TotalCasBytes, bFromDefrag);
 			return Result;
 		}
 		
@@ -4881,7 +4889,7 @@ FResult FOnDemandInstallCache::FlushPendingChunksImpl(FCacheState& CacheState, c
 					ECasErrorCode::WriteBlockFailed, EIoErrorCode::WriteError, FString::Printf(TEXT("Failed to write %" INT64_FMT " bytes to CAS block %u"), BytesToWrite, CurrentBlockId.Id));
 				FInstallCacheErrorContext Ctx = MakeInstallCacheErrorContext(CacheState);
 				Result.GetError().PushErrorContext(MoveTemp(Ctx));
-				FOnDemandInstallCacheStats::OnFlush(Result, TotalCasBytes, bFromDefrag);
+				FOnDemandInstallCacheStats::OnFlush(CacheState.Cas.GetType(), Result, TotalCasBytes, bFromDefrag);
 				return Result;
 			}
 			TotalOpCount++;
@@ -4893,7 +4901,7 @@ FResult FOnDemandInstallCache::FlushPendingChunksImpl(FCacheState& CacheState, c
 					ECasErrorCode::WriteBlockFailed, EIoErrorCode::FileFlushFailed, FString::Printf(TEXT("Failed to flush %" INT64_FMT " bytes to CAS block %u"), BytesToWrite, CurrentBlockId.Id));
 				FInstallCacheErrorContext Ctx = MakeInstallCacheErrorContext(CacheState);
 				Result.GetError().PushErrorContext(MoveTemp(Ctx));
-				FOnDemandInstallCacheStats::OnFlush(Result, TotalCasBytes, bFromDefrag);
+				FOnDemandInstallCacheStats::OnFlush(CacheState.Cas.GetType(), Result, TotalCasBytes, bFromDefrag);
 				return Result;
 			}
 			TotalOpCount++;
@@ -4943,7 +4951,7 @@ FResult FOnDemandInstallCache::FlushPendingChunksImpl(FCacheState& CacheState, c
 		{
 			FInstallCacheErrorContext Ctx = MakeInstallCacheErrorContext(CacheState);
 			Result.GetError().PushErrorContext(MoveTemp(Ctx));
-			FOnDemandInstallCacheStats::OnFlush(Result, TotalCasBytes, bFromDefrag);
+			FOnDemandInstallCacheStats::OnFlush(CacheState.Cas.GetType(), Result, TotalCasBytes, bFromDefrag);
 			return Result;
 		}
 
@@ -4957,7 +4965,7 @@ FResult FOnDemandInstallCache::FlushPendingChunksImpl(FCacheState& CacheState, c
 	}
 
 	FResult Result = MakeValue();
-	FOnDemandInstallCacheStats::OnFlush(Result, TotalCasBytes, bFromDefrag);
+	FOnDemandInstallCacheStats::OnFlush(CacheState.Cas.GetType(), Result, TotalCasBytes, bFromDefrag);
 	return Result;
 }
 
@@ -5007,10 +5015,13 @@ void FOnDemandInstallCache::CompleteRequest(FIoRequestImpl* Request, EIoErrorCod
 		TRACE_IOSTORE_BACKEND_REQUEST_COMPLETED(Request, Request->GetBuffer().GetSize());
 	}
 
+	const TUniquePtr<FChunkRequest> ChunkRequest = FChunkRequest::Detach(*Request);
+	const EOnDemandInstallCasType CasType = ChunkRequest ? ChunkRequest->CasType : EOnDemandInstallCasType::None;
+
 	{
 		UE::TUniqueLock Lock(Mutex);
 		CompletedRequests.AddTail(Request);
-		FOnDemandInstallCacheStats::OnReadCompleted(Status, EncodedSize);
+		FOnDemandInstallCacheStats::OnReadCompleted(CasType, Status, EncodedSize);
 	}
 
 	BackendContext->WakeUpDispatcherThreadDelegate.Execute();
diff --git a/Engine/Source/Runtime/Experimental/IoStore/OnDemand/Private/Statistics.cpp b/Engine/Source/Runtime/Experimental/IoStore/OnDemand/Private/Statistics.cpp
--- a/Engine/Source/Runtime/Experimental/IoStore/OnDemand/Private/Statistics.cpp
+++ b/Engine/Source/Runtime/Experimental/IoStore/OnDemand/Private/Statistics.cpp
@@ -503,13 +503,17 @@ CSV_DEFINE_STAT(Iax, InstallerInstallCount);
 CSV_DEFINE_STAT(Iax, InstallerInflightCount);
 CSV_DEFINE_STAT(Iax, InstallerDownloadedMiB);
 CSV_DEFINE_STAT(Iax, InstallerDurationAvgMs);
+CSV_DEFINE_STAT(Iax, InstallerAvgCacheHitPct);
 
-CSV_DEFINE_STAT(Iax, InstallCachePurgeCount);
-CSV_DEFINE_STAT(Iax, InstallCacheDefragCount);
-CSV_DEFINE_STAT(Iax, InstallCacheWriteCount);
-CSV_DEFINE_STAT(Iax, InstallCacheWriteMiB);
-CSV_DEFINE_STAT(Iax, InstallCacheDecodeErrorCount);
-CSV_DEFINE_STAT(Iax, InstallCacheAvgCacheHitPct);
+#define DEFINE_CAS_STATS_FOR_CSV(CasType) \
+	CSV_DEFINE_STAT(Iax, CasType##_InstallCachePurgeCount) \
+	CSV_DEFINE_STAT(Iax, CasType##_InstallCacheDefragCount) \
+	CSV_DEFINE_STAT(Iax, CasType##_InstallCacheWriteCount) \
+	CSV_DEFINE_STAT(Iax, CasType##_InstallCacheWriteMiB) \
+	CSV_DEFINE_STAT(Iax, CasType##_InstallCacheDecodeErrorCount)
+
+ON_DEMAND_INSTALL_CAS_TYPE_LIST(DEFINE_CAS_STATS_FOR_CSV)
+#undef DEFINE_CAS_STATS_FOR_CSV
 
 CSV_DEFINE_STAT(Iax, StreamingCacheWriteCount);
 CSV_DEFINE_STAT(Iax, StreamingCacheWriteMiB);
@@ -572,8 +576,6 @@ struct FInstallerAnalytics
 	uint64 DownloadedBytes = 0;
 	uint64 TotalInstallDurationMs = 0;
 	double TotalCacheHitRatio = 0.0;
-	double TotalFragmentionRatio = 0.0;
-	double DefraggedBlockFragmentionRatio = 0.0;
 };
 FMutex FInstallerAnalytics::Mutex;
 static FInstallerAnalytics InstallerAnalytics;
@@ -581,10 +583,17 @@ static FInstallerAnalytics InstallerAnalytics;
 ////////////////////////////////////////////////////////////////////////////////
 struct FInstallCacheAnalytics
 {
-	void Lock()		{ Mutex.Lock(); }
-	void Unlock()	{ Mutex.Unlock(); }
 	static FMutex Mutex; 
 
+	// Stats not associated with a particular CAS
+	struct FShared
+	{
+		uint64 ResolveErrorCount = 0;
+	};
+
+	static FShared Shared;
+
+	// CAS specific stats
 	uint64 VerificationRemovedBlockCount = 0;
 	uint64 FlushCount = 0;
 	uint64 FlushErrorCount = 0;
@@ -606,28 +615,53 @@ struct FInstallCacheAnalytics
 	uint64 MaxReferencedBlockSize = 0;
 	uint64 MaxReferencedSize = 0;
 	uint64 MaxFragmentedSize = 0;
-	int64 OldestBlockAccess = FDateTime::MaxValue().GetTicks();
 	uint32 StartupErrorCode = 0;
+
+	double TotalFragmentionRatio = 0.0;
+	double DefraggedBlockFragmentionRatio = 0.0;
+
+	EOnDemandInstallCasType CasType = EOnDemandInstallCasType::None;
+
+	static FInstallCacheAnalytics& Get(EOnDemandInstallCasType CasType)
+	{
+		check(CasType != EOnDemandInstallCasType::None);
+
+		// Please update this function if you hit this static assert
+		static_assert(EOnDemandInstallCasType::Count == 2);
+
+		if constexpr(FPlatformProperties::SupportsMemoryMappedFiles())
+		{
+			if (CasType == EOnDemandInstallCasType::MMap)
+			{
+				static FInstallCacheAnalytics MMapCacheAnalytics{ .CasType = EOnDemandInstallCasType::MMap };
+				return MMapCacheAnalytics;
+			}
+		}
+		
+		static FInstallCacheAnalytics InstallCacheAnalytics{ .CasType = EOnDemandInstallCasType::General };
+		return InstallCacheAnalytics;
+	}
 };
 FMutex FInstallCacheAnalytics::Mutex;
-static FInstallCacheAnalytics InstallCacheAnalytics;
+FInstallCacheAnalytics::FShared FInstallCacheAnalytics::Shared;
 
 ////////////////////////////////////////////////////////////////////////////////
 struct FInstallCacheCounters
 {
-	FInstallCacheCounters()
-		: MaxCacheSize(TEXT("Iax/InstallCache/MaxCacheSize"), TraceCounterDisplayHint_Memory)
-		, CacheSize(TEXT("Iax/InstallCache/CacheSize"), TraceCounterDisplayHint_Memory)
-		, FlushedBytes(TEXT("Iax/IntallCache/WriteBytes"), TraceCounterDisplayHint_Memory)
-		, FlushCount(TEXT("Iax/InstallCache/WriteCount"), TraceCounterDisplayHint_None)
-		, InMemoryBytes(TEXT("Iax/InstallCache/InMemoryBytes"), TraceCounterDisplayHint_Memory)
-		, InMemoryCount(TEXT("Iax/InstallCache/InMemoryCount"), TraceCounterDisplayHint_None)
-		, ReadBytes(TEXT("Iax/InstallCache/ReadBytes"), TraceCounterDisplayHint_Memory)
-		, ReadCount(TEXT("Iax/InstallCache/ReadCount"), TraceCounterDisplayHint_None)
-		, PurgeCount(TEXT("Iax/InstallCache/PurgeCount"), TraceCounterDisplayHint_None)
-		, DefragCount(TEXT("Iax/InstallCache/DefragCount"), TraceCounterDisplayHint_None)
-		, ReadErrorCount(TEXT("Iax/InstallCache/ReadErrorCount"), TraceCounterDisplayHint_None)
-		, ReadDecodeErrorCount(TEXT("Iax/InstallCache/ReadDecodeErrorCount"), TraceCounterDisplayHint_None)
+	FInstallCacheCounters(EOnDemandInstallCasType InCasType)
+		: MaxCacheSize(TraceCounterNameType_Dynamic,			*FString::Printf(TEXT("Iax/InstallCache/%s/MaxCacheSize"),			LexToString(InCasType)), TraceCounterDisplayHint_Memory)
+		, CacheSize(TraceCounterNameType_Dynamic,				*FString::Printf(TEXT("Iax/InstallCache/%s/CacheSize"),				LexToString(InCasType)), TraceCounterDisplayHint_Memory)
+		, FlushedBytes(TraceCounterNameType_Dynamic,			*FString::Printf(TEXT("Iax/InstallCache/%s/WriteBytes"),			LexToString(InCasType)), TraceCounterDisplayHint_Memory)
+		, FlushCount(TraceCounterNameType_Dynamic,				*FString::Printf(TEXT("Iax/InstallCache/%s/WriteCount"),			LexToString(InCasType)), TraceCounterDisplayHint_None)
+		, InMemoryBytes(TraceCounterNameType_Dynamic,			*FString::Printf(TEXT("Iax/InstallCache/%s/InMemoryBytes"),			LexToString(InCasType)), TraceCounterDisplayHint_Memory)
+		, InMemoryCount(TraceCounterNameType_Dynamic,			*FString::Printf(TEXT("Iax/InstallCache/%s/InMemoryCount"),			LexToString(InCasType)), TraceCounterDisplayHint_None)
+		, ReadBytes(TraceCounterNameType_Dynamic,				*FString::Printf(TEXT("Iax/InstallCache/%s/ReadBytes"),				LexToString(InCasType)), TraceCounterDisplayHint_Memory)
+		, ReadCount(TraceCounterNameType_Dynamic,				*FString::Printf(TEXT("Iax/InstallCache/%s/ReadCount"),				LexToString(InCasType)), TraceCounterDisplayHint_None)
+		, PurgeCount(TraceCounterNameType_Dynamic,				*FString::Printf(TEXT("Iax/InstallCache/%s/PurgeCount"),			LexToString(InCasType)), TraceCounterDisplayHint_None)
+		, DefragCount(TraceCounterNameType_Dynamic,				*FString::Printf(TEXT("Iax/InstallCache/%s/DefragCount"),			LexToString(InCasType)), TraceCounterDisplayHint_None)
+		, ReadErrorCount(TraceCounterNameType_Dynamic,			*FString::Printf(TEXT("Iax/InstallCache/%s/ReadErrorCount"),		LexToString(InCasType)), TraceCounterDisplayHint_None)
+		, ReadDecodeErrorCount(TraceCounterNameType_Dynamic,	*FString::Printf(TEXT("Iax/InstallCache/%s/ReadDecodeErrorCount"),	LexToString(InCasType)), TraceCounterDisplayHint_None)
+		, CasType(InCasType)
 	{ }
 
 	FCounterAtomicInt MaxCacheSize;
@@ -642,8 +676,38 @@ struct FInstallCacheCounters
 	FCounterAtomicInt DefragCount;
 	FCounterAtomicInt ReadErrorCount;
 	FCounterAtomicInt ReadDecodeErrorCount;
+
+	EOnDemandInstallCasType CasType = EOnDemandInstallCasType::None;
+
+	// Stats not associated with a particular CAS
+	struct FShared
+	{
+		FCounterAtomicInt ResolveErrorCount{ TEXT("Iax/InstallCache/Shared/ResolveErrorCount"), TraceCounterDisplayHint_None };
+	};
+
+	static FShared Shared;
+
+	static FInstallCacheCounters& Get(EOnDemandInstallCasType CasType)
+	{
+		check(CasType != EOnDemandInstallCasType::None);
+
+		// Please update this function if you hit this static assert
+		static_assert(EOnDemandInstallCasType::Count == 2);
+
+		if constexpr (FPlatformProperties::SupportsMemoryMappedFiles())
+		{
+			if (CasType == EOnDemandInstallCasType::MMap)
+			{
+				static FInstallCacheCounters MMapCacheCounters{ EOnDemandInstallCasType::MMap };
+				return MMapCacheCounters;
+			}
+		}
+
+		static FInstallCacheCounters InstallCacheCounters{ EOnDemandInstallCasType::General };
+		return InstallCacheCounters;
+	}
 };
-FInstallCacheCounters InstallCacheCounters;
+FInstallCacheCounters::FShared FInstallCacheCounters::Shared;
 
 ////////////////////////////////////////////////////////////////////////////////
 static FOnDemandIoBackendStats* GStatistics = nullptr;
@@ -724,11 +788,21 @@ void FOnDemandIoBackendStats::UpdateCSVValues()
 
 	// Install cache stats
 	{
-		CSV_CUSTOM_STAT_DEFINED(InstallCachePurgeCount, int32(InstallCacheCounters.PurgeCount.Get()), ECsvCustomStatOp::Set);
-		CSV_CUSTOM_STAT_DEFINED(InstallCacheDefragCount, int32(InstallCacheCounters.DefragCount.Get()), ECsvCustomStatOp::Set);
-		CSV_CUSTOM_STAT_DEFINED(InstallCacheWriteCount, int32(InstallCacheCounters.FlushCount.Get()), ECsvCustomStatOp::Set);
-		CSV_CUSTOM_STAT_DEFINED(InstallCacheWriteMiB, BytesToApproxMB(InstallCacheCounters.FlushedBytes.Get()), ECsvCustomStatOp::Set);
-		CSV_CUSTOM_STAT_DEFINED(InstallCacheDecodeErrorCount, int32(InstallCacheCounters.ReadDecodeErrorCount.Get()), ECsvCustomStatOp::Set);
+		#define DEFINE_CUSTOM_CAS_STATS_FOR_CSV(InCasType) \
+			{ \
+				FInstallCacheCounters& InstallCacheCounters = FInstallCacheCounters::Get(EOnDemandInstallCasType::InCasType); \
+				if (InstallCacheCounters.CasType == EOnDemandInstallCasType::InCasType) /* If a CAS type is not enabled, if will fall back to EOnDemandInstallCasType::General*/ \
+				{ \
+					CSV_CUSTOM_STAT_DEFINED(InCasType##_InstallCachePurgeCount, int32(InstallCacheCounters.PurgeCount.Get()), ECsvCustomStatOp::Set) \
+					CSV_CUSTOM_STAT_DEFINED(InCasType##_InstallCacheDefragCount, int32(InstallCacheCounters.DefragCount.Get()), ECsvCustomStatOp::Set) \
+					CSV_CUSTOM_STAT_DEFINED(InCasType##_InstallCacheWriteCount, int32(InstallCacheCounters.FlushCount.Get()), ECsvCustomStatOp::Set) \
+					CSV_CUSTOM_STAT_DEFINED(InCasType##_InstallCacheWriteMiB, BytesToApproxMB(InstallCacheCounters.FlushedBytes.Get()), ECsvCustomStatOp::Set) \
+					CSV_CUSTOM_STAT_DEFINED(InCasType##_InstallCacheDecodeErrorCount, int32(InstallCacheCounters.ReadDecodeErrorCount.Get()), ECsvCustomStatOp::Set) \
+				} \
+			}
+
+		ON_DEMAND_INSTALL_CAS_TYPE_LIST(DEFINE_CUSTOM_CAS_STATS_FOR_CSV)
+		#undef DEFINE_CUSTOM_CAS_STATS_FOR_CSV
 	}
 
 	// Installer stats
@@ -744,7 +818,7 @@ void FOnDemandIoBackendStats::UpdateCSVValues()
 		CSV_CUSTOM_STAT_DEFINED(InstallerInflightCount, int32(InflightCount), ECsvCustomStatOp::Set);
 		CSV_CUSTOM_STAT_DEFINED(InstallerDownloadedMiB, BytesToApproxMB(DownloadedBytes), ECsvCustomStatOp::Set);
 		CSV_CUSTOM_STAT_DEFINED(InstallerDurationAvgMs, int32(FMath::RoundToInt64(DurationAvgMs)), ECsvCustomStatOp::Set);
-		CSV_CUSTOM_STAT_DEFINED(InstallCacheAvgCacheHitPct, int32(FMath::RoundToInt64(CacheHitAvgPct)), ECsvCustomStatOp::Set);
+		CSV_CUSTOM_STAT_DEFINED(InstallerAvgCacheHitPct, int32(FMath::RoundToInt64(CacheHitAvgPct)), ECsvCustomStatOp::Set);
 	}
 
 	// HTTP stats
@@ -823,11 +897,16 @@ void FOnDemandIoBackendStats::PrintPeriodicLogging()
 				InstallCount, int32(DurationAvgMs), CacheHitAvgPct);
 		}
 
+		for (EOnDemandInstallCasType CasType : TEnumRange<EOnDemandInstallCasType>())
 		{
-			const int64 WriteCount = InstallCacheCounters.FlushCount.Get();
-			const int64 WriteBytes = InstallCacheCounters.FlushedBytes.Get();
-			UE_LOG(LogIoStoreOnDemand, Log, TEXT("InstallCache: WriteCount=%lld, WriteBytes=%.2lf MiB"),
-				WriteCount, double(WriteBytes) / 1024.0 / 1024.0);
+			FInstallCacheCounters& InstallCacheCounters = FInstallCacheCounters::Get(CasType);
+			if (InstallCacheCounters.CasType == CasType) // If a CAS type is not enabled, if will fall back to EOnDemandInstallCasType::General
+			{
+				const int64 WriteCount = InstallCacheCounters.FlushCount.Get();
+				const int64 WriteBytes = InstallCacheCounters.FlushedBytes.Get();
+				UE_LOG(LogIoStoreOnDemand, Log, TEXT("%s InstallCache: WriteCount=%lld, WriteBytes=%.2lf MiB"),
+					LexToString(CasType), WriteCount, double(WriteBytes) / 1024.0 / 1024.0);
+			}
 		}
 
 		LastLogTime = Time;
@@ -1079,15 +1158,21 @@ void FOnDemandIoBackendStats::PrintOnScreenStatistics(FCoreDelegates::FSeverityM
 	}
 
 	// Print Install Cache
+	for (EOnDemandInstallCasType CasType : TEnumRange<EOnDemandInstallCasType>())
 	{
-		TStringBuilder<256> Sb;
-		Sb << TEXT("Install Cache: ");
-		Sb << TEXT(" Memory: ") << FText::AsMemory(InstallCacheCounters.InMemoryBytes.Get()).ToString();
-		Sb << TEXT(" (") << InstallCacheCounters.InMemoryCount.Get() << TEXT(")");
-		Sb << TEXT(" | Write: ") << FText::AsMemory(InstallCacheCounters.FlushedBytes.Get()).ToString();
-		Sb << TEXT(" (") << InstallCacheCounters.FlushCount.Get() << TEXT(")");
-		Sb << TEXT(" | Read: ") << FText::AsMemory(InstallCacheCounters.ReadBytes.Get()).ToString();
-		OutMessages.Add(FCoreDelegates::EOnScreenMessageSeverity::Info, FText::FromStringView(Sb.ToView()));
+		FInstallCacheCounters& InstallCacheCounters = FInstallCacheCounters::Get(CasType);
+		if (InstallCacheCounters.CasType == CasType) // If a CAS type is not enabled, if will fall back to EOnDemandInstallCasType::General
+		{
+			TStringBuilder<256> Sb;
+			Sb << LexToString(CasType);
+			Sb << TEXT(" Install Cache: ");
+			Sb << TEXT(" Memory: ") << FText::AsMemory(InstallCacheCounters.InMemoryBytes.Get()).ToString();
+			Sb << TEXT(" (") << InstallCacheCounters.InMemoryCount.Get() << TEXT(")");
+			Sb << TEXT(" | Write: ") << FText::AsMemory(InstallCacheCounters.FlushedBytes.Get()).ToString();
+			Sb << TEXT(" (") << InstallCacheCounters.FlushCount.Get() << TEXT(")");
+			Sb << TEXT(" | Read: ") << FText::AsMemory(InstallCacheCounters.ReadBytes.Get()).ToString();
+			OutMessages.Add(FCoreDelegates::EOnScreenMessageSeverity::Info, FText::FromStringView(Sb.ToView()));
+		}
 	}
 
 #endif // UE_ENABLE_ONSCREEN_STATISTICS
@@ -1653,21 +1738,27 @@ static void LogOrReportContentInstallerAnalytics(bool bReport, bool bLog, TArray
 		InstallerAnalytics = FInstallerAnalytics();
 	}
 
-	FInstallCacheAnalytics CurrentInstallCacheAnalytics;
+	FInstallCacheAnalytics::FShared CurrentSharedInstallCacheAnalytics;
+	TArray<FInstallCacheAnalytics, TInlineAllocator<EOnDemandInstallCasType::Count>> CurrentInstallCacheAnalytics;
 	{
-		TUniqueLock Lock(InstallCacheAnalytics);
-		CurrentInstallCacheAnalytics = InstallCacheAnalytics;
-		InstallCacheAnalytics = FInstallCacheAnalytics();
-	}
+		TUniqueLock Lock(FInstallCacheAnalytics::Mutex);
 
-	const FDateTime Now = FDateTime::UtcNow();
-	const FDateTime OldestBlockAccess(CurrentInstallCacheAnalytics.OldestBlockAccess);
-	const FTimespan OldestBlockAge = (Now >= OldestBlockAccess) ? (Now - OldestBlockAccess) : FTimespan::MaxValue();
+		CurrentSharedInstallCacheAnalytics = FInstallCacheAnalytics::Shared;
+		FInstallCacheAnalytics::Shared = FInstallCacheAnalytics::FShared();
+
+		for (EOnDemandInstallCasType CasType : TEnumRange<EOnDemandInstallCasType>())
+		{
+			FInstallCacheAnalytics& InstallCacheAnalytics = FInstallCacheAnalytics::Get(CasType);
+			if (InstallCacheAnalytics.CasType == CasType) // If a CAS type is not enabled, if will fall back to EOnDemandInstallCasType::General
+			{
+				CurrentInstallCacheAnalytics.Add(InstallCacheAnalytics);
+				InstallCacheAnalytics = FInstallCacheAnalytics{ .CasType = CasType };
+			}
+		}
+	}
 
 	const double AvgInstallDurationMs = CurrentInstallerAnalytics.InstallCount > 0 ? (double)CurrentInstallerAnalytics.TotalInstallDurationMs / (double)CurrentInstallerAnalytics.InstallCount : 0.0;
-	const double AvgCacheHitRatio = CurrentInstallerAnalytics.InstallCount > 0 ? (double)CurrentInstallerAnalytics.TotalCacheHitRatio / (double)CurrentInstallerAnalytics.InstallCount : 1.0;
-	const double AvgTotalFragmentationRatio = CurrentInstallCacheAnalytics.DefragCount > 0 ? (double)CurrentInstallerAnalytics.TotalFragmentionRatio / (double)CurrentInstallCacheAnalytics.DefragCount : 0.0;
-	const double AvgDefraggedBlockFragmentationRatio = CurrentInstallCacheAnalytics.DefragCount > 0 ? (double)CurrentInstallerAnalytics.DefraggedBlockFragmentionRatio / (double)CurrentInstallCacheAnalytics.DefragCount : 0.0;
+	const double AvgCacheHitRatio = CurrentInstallerAnalytics.InstallCount > 0 ? CurrentInstallerAnalytics.TotalCacheHitRatio / (double)CurrentInstallerAnalytics.InstallCount : 1.0;
 
 	#define INSTALLER_ANALYTICS_LIST(XANALYTIC) \
 		XANALYTIC(IadTotalInstallCount, CurrentInstallerAnalytics.InstallCount) \
@@ -1676,39 +1767,88 @@ static void LogOrReportContentInstallerAnalytics(bool bReport, bool bLog, TArray
 		XANALYTIC(IadTotalDownloadedBytes, CurrentInstallerAnalytics.DownloadedBytes) \
 		XANALYTIC(IadTotalInstallDurationMs, CurrentInstallerAnalytics.TotalInstallDurationMs) \
 		XANALYTIC(IadAvgInstallDurationMs, AvgInstallDurationMs) \
-		XANALYTIC(IadAvgCacheHitRatio, AvgCacheHitRatio) \
-		XANALYTIC(IadAvgTotalFragmentationRatio, AvgTotalFragmentationRatio) \
-		XANALYTIC(IadAvgDefraggedBlockFragmentationRatio, AvgDefraggedBlockFragmentationRatio) \
-		XANALYTIC(IadInstallCacheStartupErrorCode, CurrentInstallCacheAnalytics.StartupErrorCode) \
-		XANALYTIC(IadInstallCacheVerificationRemovedBlockCount, CurrentInstallCacheAnalytics.VerificationRemovedBlockCount) \
-		XANALYTIC(IadInstallCacheFlushCount, CurrentInstallCacheAnalytics.FlushCount) \
-		XANALYTIC(IadInstallCacheFlushErrorCount, CurrentInstallCacheAnalytics.FlushErrorCount) \
-		XANALYTIC(IadInstallCacheFlushedBytes, CurrentInstallCacheAnalytics.FlushedBytes) \
-		XANALYTIC(IadInstallCacheDefragFlushCount, CurrentInstallCacheAnalytics.DefragFlushCount) \
-		XANALYTIC(IadInstallCacheDefragFlushErrorCount, CurrentInstallCacheAnalytics.DefragFlushErrorCount) \
-		XANALYTIC(IadInstallCacheDefragFlushedBytes, CurrentInstallCacheAnalytics.DefragFlushedBytes) \
-		XANALYTIC(IadInstallCachePurgeCount, CurrentInstallCacheAnalytics.PurgeCount) \
-		XANALYTIC(IadInstallCachePurgeErrorCount, CurrentInstallCacheAnalytics.PurgeErrorCount) \
-		XANALYTIC(IadInstallCacheDefragCount, CurrentInstallCacheAnalytics.DefragCount) \
-		XANALYTIC(IadInstallCacheDefragErrorCount, CurrentInstallCacheAnalytics.DefragErrorCount) \
-		XANALYTIC(IadInstallCacheJournalCommitCount, CurrentInstallCacheAnalytics.JournalCommitCount) \
-		XANALYTIC(IadInstallCacheJournalCommitErrorCount, CurrentInstallCacheAnalytics.JournalCommitErrorCount) \
-		XANALYTIC(IadInstallCacheMaxSize, CurrentInstallCacheAnalytics.MaxCacheSize) \
-		XANALYTIC(IadInstallCacheMaxUsageSize, CurrentInstallCacheAnalytics.MaxCacheUsageSize) \
-		XANALYTIC(IadInstallCacheMaxReferencedBlockSize, CurrentInstallCacheAnalytics.MaxReferencedBlockSize) \
-		XANALYTIC(IadInstallCacheMaxReferencedSize, CurrentInstallCacheAnalytics.MaxReferencedSize) \
-		XANALYTIC(IadInstallCacheMaxFragmentedSize, CurrentInstallCacheAnalytics.MaxFragmentedSize) \
-		XANALYTIC(IadInstallCacheOldestBlockAgeMinutes, OldestBlockAge.GetTotalMinutes()) \
-		XANALYTIC(IadInstallCacheReadCount, CurrentInstallCacheAnalytics.ReadCount) \
-		XANALYTIC(IadInstallCacheReadErrorCount, CurrentInstallCacheAnalytics.ReadErrorCount)
+		XANALYTIC(IadAvgCacheHitRatio, AvgCacheHitRatio)
+
+	#define SHARED_INSTALL_CACHE_ANALYTICS_LIST(XANALYTIC) \
+		XANALYTIC(IadInstallCacheResolveErrorCount, CurrentSharedInstallCacheAnalytics.ResolveErrorCount)
+
+	TArray<double, TInlineAllocator<EOnDemandInstallCasType::Count>> AvgTotalFragmentationRatio;
+	TArray<double, TInlineAllocator<EOnDemandInstallCasType::Count>> AvgDefraggedBlockFragmentationRatio;
+	for (FInstallCacheAnalytics& InstallCacheAnalytics : CurrentInstallCacheAnalytics)
+	{
+		AvgTotalFragmentationRatio.Emplace_GetRef() = InstallCacheAnalytics.DefragCount > 0 ? InstallCacheAnalytics.TotalFragmentionRatio / (double)InstallCacheAnalytics.DefragCount : 0.0;
+		AvgDefraggedBlockFragmentationRatio.Emplace_GetRef() = InstallCacheAnalytics.DefragCount > 0 ? InstallCacheAnalytics.DefraggedBlockFragmentionRatio / (double)InstallCacheAnalytics.DefragCount : 0.0;
+	}
+
+	#define INSTALL_CACHE_ANALYTICS_LIST(XANALYTIC) \
+		XANALYTIC(IadAvgTotalFragmentationRatio, AvgTotalFragmentationRatio[Index]) \
+		XANALYTIC(IadAvgDefraggedBlockFragmentationRatio, AvgDefraggedBlockFragmentationRatio[Index]) \
+		XANALYTIC(IadInstallCacheStartupErrorCode, CurrentInstallCacheAnalytics[Index].StartupErrorCode) \
+		XANALYTIC(IadInstallCacheVerificationRemovedBlockCount, CurrentInstallCacheAnalytics[Index].VerificationRemovedBlockCount) \
+		XANALYTIC(IadInstallCacheFlushCount, CurrentInstallCacheAnalytics[Index].FlushCount) \
+		XANALYTIC(IadInstallCacheFlushErrorCount, CurrentInstallCacheAnalytics[Index].FlushErrorCount) \
+		XANALYTIC(IadInstallCacheFlushedBytes, CurrentInstallCacheAnalytics[Index].FlushedBytes) \
+		XANALYTIC(IadInstallCacheDefragFlushCount, CurrentInstallCacheAnalytics[Index].DefragFlushCount) \
+		XANALYTIC(IadInstallCacheDefragFlushErrorCount, CurrentInstallCacheAnalytics[Index].DefragFlushErrorCount) \
+		XANALYTIC(IadInstallCacheDefragFlushedBytes, CurrentInstallCacheAnalytics[Index].DefragFlushedBytes) \
+		XANALYTIC(IadInstallCachePurgeCount, CurrentInstallCacheAnalytics[Index].PurgeCount) \
+		XANALYTIC(IadInstallCachePurgeErrorCount, CurrentInstallCacheAnalytics[Index].PurgeErrorCount) \
+		XANALYTIC(IadInstallCacheDefragCount, CurrentInstallCacheAnalytics[Index].DefragCount) \
+		XANALYTIC(IadInstallCacheDefragErrorCount, CurrentInstallCacheAnalytics[Index].DefragErrorCount) \
+		XANALYTIC(IadInstallCacheJournalCommitCount, CurrentInstallCacheAnalytics[Index].JournalCommitCount) \
+		XANALYTIC(IadInstallCacheJournalCommitErrorCount, CurrentInstallCacheAnalytics[Index].JournalCommitErrorCount) \
+		XANALYTIC(IadInstallCacheMaxSize, CurrentInstallCacheAnalytics[Index].MaxCacheSize) \
+		XANALYTIC(IadInstallCacheMaxUsageSize, CurrentInstallCacheAnalytics[Index].MaxCacheUsageSize) \
+		XANALYTIC(IadInstallCacheMaxReferencedBlockSize, CurrentInstallCacheAnalytics[Index].MaxReferencedBlockSize) \
+		XANALYTIC(IadInstallCacheMaxReferencedSize, CurrentInstallCacheAnalytics[Index].MaxReferencedSize) \
+		XANALYTIC(IadInstallCacheMaxFragmentedSize, CurrentInstallCacheAnalytics[Index].MaxFragmentedSize) \
+		XANALYTIC(IadInstallCacheReadCount, CurrentInstallCacheAnalytics[Index].ReadCount) \
+		XANALYTIC(IadInstallCacheReadErrorCount, CurrentInstallCacheAnalytics[Index].ReadErrorCount)
 
 	if (bReport)
 	{
+		#define COUNT_ANALYTIC(inName, inValue) ++AnalyticCount;
+		int32 AnalyticCount = 0;
+		INSTALLER_ANALYTICS_LIST(COUNT_ANALYTIC)
+		SHARED_INSTALL_CACHE_ANALYTICS_LIST(COUNT_ANALYTIC)
+		for (FInstallCacheAnalytics& InstallCacheAnalytics : CurrentInstallCacheAnalytics)
+		{
+			INSTALL_CACHE_ANALYTICS_LIST(COUNT_ANALYTIC)
+		}
+		OutAnalyticsArray.Reserve(AnalyticCount);
+		#undef COUNT_ANALYTIC
+		
+
 		#define REPORT_ANALYTIC(inName, inValue) , TEXT(#inName), inValue
 		AppendAnalyticsEventAttributeArray(OutAnalyticsArray
 			INSTALLER_ANALYTICS_LIST(REPORT_ANALYTIC)
 		);
 		#undef REPORT_ANALYTIC
+
+		#define REPORT_ANALYTIC(inName, inValue) , FString::Printf(TEXT("%s%s%s"), PrefixStr, DotStr, TEXT(#inName)), inValue
+		{
+			const TCHAR* PrefixStr = TEXT("Shared");
+			const TCHAR* DotStr = TEXT(".");
+			AppendAnalyticsEventAttributeArray(OutAnalyticsArray
+				SHARED_INSTALL_CACHE_ANALYTICS_LIST(REPORT_ANALYTIC)
+			);
+		}
+		for (int32 Index = 0; Index < CurrentInstallCacheAnalytics.Num(); ++Index)
+		{
+			FInstallCacheAnalytics& InstallCacheAnalytics = CurrentInstallCacheAnalytics[Index];
+			const TCHAR* PrefixStr = TEXT("");
+			const TCHAR* DotStr = TEXT("");
+			if (InstallCacheAnalytics.CasType != EOnDemandInstallCasType::General)
+			{
+				PrefixStr = LexToString(InstallCacheAnalytics.CasType);
+				DotStr = TEXT(".");
+			}
+
+			AppendAnalyticsEventAttributeArray(OutAnalyticsArray
+				INSTALL_CACHE_ANALYTICS_LIST(REPORT_ANALYTIC)
+			);
+		}
+		#undef REPORT_ANALYTIC
 	}
 
 	if (bLog)
@@ -1716,13 +1856,23 @@ static void LogOrReportContentInstallerAnalytics(bool bReport, bool bLog, TArray
 		#define ANALYTIC_LOG_FORMAT(inName, inValue) "stat:\t" #inName " {"#inName"}\n"
 		#define ANALYTIC_LOG_VALUE(inName, inValue) , UE_LOGFMT_FIELD(#inName, inValue)
 
-		UE_LOGFMT_EX(LogIoStoreOnDemand, Log, "IadStats:\n" INSTALLER_ANALYTICS_LIST(ANALYTIC_LOG_FORMAT) INSTALLER_ANALYTICS_LIST(ANALYTIC_LOG_VALUE));
+		UE_LOGFMT_EX(LogIoStoreOnDemand, Log, "IAD Installer Stats:\n" INSTALLER_ANALYTICS_LIST(ANALYTIC_LOG_FORMAT) INSTALLER_ANALYTICS_LIST(ANALYTIC_LOG_VALUE));
+
+		UE_LOGFMT_EX(LogIoStoreOnDemand, Log, "IAD Shared Cache Stats:\n" SHARED_INSTALL_CACHE_ANALYTICS_LIST(ANALYTIC_LOG_FORMAT) SHARED_INSTALL_CACHE_ANALYTICS_LIST(ANALYTIC_LOG_VALUE));
+		
+		for (int32 Index = 0; Index < CurrentInstallCacheAnalytics.Num(); ++Index)
+		{
+			FInstallCacheAnalytics& InstallCacheAnalytics = CurrentInstallCacheAnalytics[Index];
+			UE_LOGFMT_EX(LogIoStoreOnDemand, Log, "IAD {CasType} Cache Stats:\n" INSTALL_CACHE_ANALYTICS_LIST(ANALYTIC_LOG_FORMAT), UE_LOGFMT_FIELD("CasType", LexToString(InstallCacheAnalytics.CasType)) INSTALL_CACHE_ANALYTICS_LIST(ANALYTIC_LOG_VALUE));
+		}
 
 		#undef ANALYTIC_LOG_FORMAT
 		#undef ANALYTIC_LOG_VALUE
 	}
 
 	#undef INSTALLER_ANALYTICS_LIST
+	#undef SHARED_INSTALL_CACHE_ANALYTICS_LIST
+	#undef INSTALL_CACHE_ANALYTICS_LIST
 }
 
 void FOnDemandContentInstallerStats::ReportAnalytics(TArray<FAnalyticsEventAttribute>& OutAnalyticsArray)
@@ -1750,14 +1900,16 @@ void FOnDemandInstallCacheStats::OnStartupError(const FResult& Result)
 	}
 }
 
-void FOnDemandInstallCacheStats::OnPutChunk(uint64 ByteCount)
+void FOnDemandInstallCacheStats::OnPutChunk(EOnDemandInstallCasType CasType, uint64 ByteCount)
 {
+	FInstallCacheCounters& InstallCacheCounters = FInstallCacheCounters::Get(CasType);
 	InstallCacheCounters.InMemoryBytes.Add(ByteCount);
 	InstallCacheCounters.InMemoryCount.Increment();
 }
 
-void FOnDemandInstallCacheStats::OnFlush(const FResult& Result, int64 ByteCount, bool bFromDefrag)
+void FOnDemandInstallCacheStats::OnFlush(EOnDemandInstallCasType CasType, const FResult& Result, int64 ByteCount, bool bFromDefrag)
 {
+	FInstallCacheCounters& InstallCacheCounters = FInstallCacheCounters::Get(CasType);
 	InstallCacheCounters.InMemoryBytes.Set(0);
 	InstallCacheCounters.InMemoryCount.Set(0);
 
@@ -1769,7 +1921,8 @@ void FOnDemandInstallCacheStats::OnFlush(const FResult& Result, int64 ByteCount,
 	InstallCacheCounters.FlushedBytes.Add(ByteCount);
 	InstallCacheCounters.FlushCount.Increment();
 
-	UE::TUniqueLock Lock(InstallCacheAnalytics);
+	UE::TUniqueLock Lock(FInstallCacheAnalytics::Mutex);
+	FInstallCacheAnalytics& InstallCacheAnalytics = FInstallCacheAnalytics::Get(CasType);
 
 	if (bFromDefrag)
 	{
@@ -1791,14 +1944,16 @@ void FOnDemandInstallCacheStats::OnFlush(const FResult& Result, int64 ByteCount,
 	}
 }
 
-void FOnDemandInstallCacheStats::OnJournalCommit(const FResult& Result, int64 ByteCount)
+void FOnDemandInstallCacheStats::OnJournalCommit(EOnDemandInstallCasType CasType, const FResult& Result, int64 ByteCount)
 {
 	if (Result.HasError() && UE::UnifiedError::IsCancellationError(Result.GetError()))
 	{
 		return;
 	}
 
-	UE::TUniqueLock Lock(InstallCacheAnalytics);
+	UE::TUniqueLock Lock(FInstallCacheAnalytics::Mutex);
+	FInstallCacheAnalytics& InstallCacheAnalytics = FInstallCacheAnalytics::Get(CasType);
+
 	if (Result.HasError())
 	{
 		InstallCacheAnalytics.JournalCommitErrorCount++;
@@ -1806,19 +1961,22 @@ void FOnDemandInstallCacheStats::OnJournalCommit(const FResult& Result, int64 By
 	InstallCacheAnalytics.JournalCommitCount++;
 }
 
-void FOnDemandInstallCacheStats::OnCasVerification(const FResult& Result, int32 RemovedChunkCount, uint64 BlockBytesOverBudget)
+void FOnDemandInstallCacheStats::OnCasVerification(EOnDemandInstallCasType CasType, const FResult& Result, int32 RemovedChunkCount, uint64 BlockBytesOverBudget)
 {
 	FOnDemandImmediateAnalytic Event = MakeAnalyticsEventFromResult(TEXT("IadCasVerification"), Result);
 	AppendAnalyticsEventAttributeArray(Event.AnalyticsArray,
 		TEXT("RemovedChunkCount"), RemovedChunkCount,
-		TEXT("BlockBytesOverBudget"), BlockBytesOverBudget);
+		TEXT("BlockBytesOverBudget"), BlockBytesOverBudget,
+		TEXT("CasType"), LexToString(CasType));
 	SendImmediateAnalytic(MoveTemp(Event));
 
-	UE::TUniqueLock Lock(InstallCacheAnalytics);
+	UE::TUniqueLock Lock(FInstallCacheAnalytics::Mutex);
+	FInstallCacheAnalytics& InstallCacheAnalytics = FInstallCacheAnalytics::Get(CasType);
 	InstallCacheAnalytics.VerificationRemovedBlockCount += RemovedChunkCount;
 }
 
 void FOnDemandInstallCacheStats::OnPurge(
+	EOnDemandInstallCasType CasType,
 	const FResult& Result,
 	uint64 MaxCacheSize,
 	uint64 NewCacheSize,
@@ -1830,10 +1988,12 @@ void FOnDemandInstallCacheStats::OnPurge(
 		return;
 	}
 
+	FInstallCacheCounters& InstallCacheCounters = FInstallCacheCounters::Get(CasType);
 	InstallCacheCounters.PurgeCount.Increment();
 
 	{
-		UE::TUniqueLock Lock(InstallCacheAnalytics);
+		UE::TUniqueLock Lock(FInstallCacheAnalytics::Mutex);
+		FInstallCacheAnalytics& InstallCacheAnalytics = FInstallCacheAnalytics::Get(CasType);
 		if (Result.HasError())
 		{
 			InstallCacheAnalytics.PurgeErrorCount++;
@@ -1847,11 +2007,13 @@ void FOnDemandInstallCacheStats::OnPurge(
 	FOnDemandImmediateAnalytic Event = MakeAnalyticsEventFromResult(TEXT("IadCachePurge"), Result);
 	AppendAnalyticsEventAttributeArray(Event.AnalyticsArray, 
 		TEXT("PurgedBytes"), PurgedBytes,
-		TEXT("LineNo"), LineNo);
+		TEXT("LineNo"), LineNo,
+		TEXT("CasType"), LexToString(CasType));
 	SendImmediateAnalytic(MoveTemp(Event));
 }
 
 void FOnDemandInstallCacheStats::OnDefrag(
+	EOnDemandInstallCasType CasType,
 	const FResult& Result, 
 	uint64 TotalFragmentedBytes,
 	uint64 TotalCachedBytes,
@@ -1863,6 +2025,7 @@ void FOnDemandInstallCacheStats::OnDefrag(
 		return;
 	}
 
+	FInstallCacheCounters& InstallCacheCounters = FInstallCacheCounters::Get(CasType);
 	InstallCacheCounters.DefragCount.Increment();
 
 	const double TotalFragmentionRatio = TotalCachedBytes > 0
@@ -1874,18 +2037,16 @@ void FOnDemandInstallCacheStats::OnDefrag(
 		: 0.0;
 
 	{
-		UE::TUniqueLock Lock(InstallCacheAnalytics);
+		UE::TUniqueLock Lock(FInstallCacheAnalytics::Mutex);
+		FInstallCacheAnalytics& InstallCacheAnalytics = FInstallCacheAnalytics::Get(CasType);
 		if (Result.HasError())
 		{
 			InstallCacheAnalytics.DefragErrorCount++;
 		}
 		InstallCacheAnalytics.DefragCount++;
-	}
 
-	{
-		TUniqueLock Lock(InstallerAnalytics);
-		InstallerAnalytics.TotalFragmentionRatio += TotalFragmentionRatio;
-		InstallerAnalytics.DefraggedBlockFragmentionRatio += BlocksToDefragFragmentionRatio;
+		InstallCacheAnalytics.TotalFragmentionRatio += TotalFragmentionRatio;
+		InstallCacheAnalytics.DefraggedBlockFragmentionRatio += BlocksToDefragFragmentionRatio;
 	}
 
 	const int32 LineNo = TryGetLineNo(Result);
@@ -1898,11 +2059,12 @@ void FOnDemandInstallCacheStats::OnDefrag(
 		TEXT("DefraggedBlocksFragmentedSize"), BlocksToDefragFragmentedSize,
 		TEXT("DefraggedBlocksTotalSize"), BlocksToDefragTotalSize,
 		TEXT("DefraggedBlocksFragmentionRatio"), BlocksToDefragFragmentionRatio,
-		TEXT("LineNo"), LineNo);
+		TEXT("LineNo"), LineNo,
+		TEXT("CasType"), LexToString(CasType));
 	SendImmediateAnalytic(MoveTemp(Event));
 }
 
-void FOnDemandInstallCacheStats::OnBlockDeleted(int64 LastAccessTicks, int64 LastModificationTicks, bool bFromDefrag)
+void FOnDemandInstallCacheStats::OnBlockDeleted(EOnDemandInstallCasType CasType, int64 LastAccessTicks, int64 LastModificationTicks, bool bFromDefrag)
 {
 	const FDateTime Now = FDateTime::UtcNow();
 	const FDateTime BlockAccess(LastAccessTicks);
@@ -1932,53 +2094,83 @@ void FOnDemandInstallCacheStats::OnBlockDeleted(int64 LastAccessTicks, int64 Las
 		MakeAnalyticsEventAttributeArray(
 			TEXT("BlockAgeMinutes"), BlockAccessAge.GetTotalMinutes(),
 			TEXT("BlockModAgeMinutes"), BlockModAge.GetTotalMinutes(),
-			TEXT("FromDefrag"), bFromDefrag
+			TEXT("FromDefrag"), bFromDefrag,
+			TEXT("CasType"), LexToString(CasType)
 		)
 	});
 }
 
 void FOnDemandInstallCacheStats::OnCacheUsage(
+	EOnDemandInstallCasType CasType,
 	uint64 MaxCacheSize,
 	uint64 CacheSize,
 	uint64 ReferencedBlockSize,
 	uint64 ReferencedSize,
-	uint64 FragmentedSize,
-	int64 OldestBlockAccess)
+	uint64 FragmentedSize)
 {
+	FInstallCacheCounters& InstallCacheCounters = FInstallCacheCounters::Get(CasType);
 	InstallCacheCounters.MaxCacheSize.Set(MaxCacheSize);
 	InstallCacheCounters.CacheSize.Set(CacheSize);
 
-	UE::TUniqueLock Lock(InstallCacheAnalytics);
+	UE::TUniqueLock Lock(FInstallCacheAnalytics::Mutex);
+	FInstallCacheAnalytics& InstallCacheAnalytics = FInstallCacheAnalytics::Get(CasType);
+
 	InstallCacheAnalytics.MaxCacheSize				= MaxCacheSize; 
 	InstallCacheAnalytics.MaxCacheUsageSize			= FMath::Max(InstallCacheAnalytics.MaxCacheUsageSize, CacheSize);
 	InstallCacheAnalytics.MaxReferencedBlockSize	= FMath::Max(InstallCacheAnalytics.MaxReferencedBlockSize, ReferencedBlockSize);
 	InstallCacheAnalytics.MaxReferencedSize			= FMath::Max(InstallCacheAnalytics.MaxReferencedSize, ReferencedSize);
 	InstallCacheAnalytics.MaxFragmentedSize			= FMath::Max(InstallCacheAnalytics.MaxFragmentedSize, FragmentedSize);
-	InstallCacheAnalytics.OldestBlockAccess			= FMath::Min(InstallCacheAnalytics.OldestBlockAccess, OldestBlockAccess);
 }
 
-void FOnDemandInstallCacheStats::OnReadCompleted(EIoErrorCode ErrorCode, uint64 EncodedSize)
+void FOnDemandInstallCacheStats::OnReadCompleted(EOnDemandInstallCasType CasType, EIoErrorCode ErrorCode, uint64 EncodedSize)
 {
 	if (ErrorCode == EIoErrorCode::Cancelled)
 	{
 		return;
 	}
 
-	InstallCacheCounters.ReadBytes.Add(EncodedSize);
-	InstallCacheCounters.ReadCount.Increment();
-
-	UE::TUniqueLock Lock(InstallCacheAnalytics);
-	if (ErrorCode != EIoErrorCode::Ok)
+	if (CasType == EOnDemandInstallCasType::None)
+	{
+		FInstallCacheCounters::Shared.ResolveErrorCount.Increment();
+	}
+	else
 	{
-		InstallCacheCounters.ReadErrorCount.Increment();
-		++InstallCacheAnalytics.ReadErrorCount;
+		FInstallCacheCounters& InstallCacheCounters = FInstallCacheCounters::Get(CasType);		
 		
-		if (ErrorCode == EIoErrorCode::CompressionError)
+		if (ErrorCode != EIoErrorCode::Ok)
+		{
+			InstallCacheCounters.ReadErrorCount.Increment();
+			if (ErrorCode == EIoErrorCode::CompressionError)
+			{
+				InstallCacheCounters.ReadDecodeErrorCount.Increment();
+			}
+		}
+		else
+		{
+			InstallCacheCounters.ReadBytes.Add(EncodedSize);
+		}
+		InstallCacheCounters.ReadCount.Increment();
+	}
+
+	UE::TUniqueLock Lock(FInstallCacheAnalytics::Mutex);
+
+	if (CasType == EOnDemandInstallCasType::None)
+	{
+		if (ErrorCode != EIoErrorCode::Ok)
+		{
+			++FInstallCacheAnalytics::Shared.ResolveErrorCount;
+		}
+	}
+	else
+	{
+		FInstallCacheAnalytics& InstallCacheAnalytics = FInstallCacheAnalytics::Get(CasType);
+
+		if (ErrorCode != EIoErrorCode::Ok)
 		{
-			InstallCacheCounters.ReadDecodeErrorCount.Increment();
+			++InstallCacheAnalytics.ReadErrorCount;
 		}
+		++InstallCacheAnalytics.ReadCount;
 	}
-	++InstallCacheAnalytics.ReadCount;
 }
 
 } // namespace UE::IoStore
diff --git a/Engine/Source/Runtime/Experimental/IoStore/OnDemand/Private/Statistics.h b/Engine/Source/Runtime/Experimental/IoStore/OnDemand/Private/Statistics.h
--- a/Engine/Source/Runtime/Experimental/IoStore/OnDemand/Private/Statistics.h
+++ b/Engine/Source/Runtime/Experimental/IoStore/OnDemand/Private/Statistics.h
@@ -157,32 +157,33 @@ class FOnDemandInstallCacheStats
 {
 public:
 	static void OnStartupError(const FResult& Result) IAS_STATISTICS_IMPL()
-	static void OnPutChunk(uint64 ChunkSize) IAS_STATISTICS_IMPL()
-	static void OnFlush(const FResult& Result, int64 ByteCount, bool bFromDefrag) IAS_STATISTICS_IMPL()
-	static void OnJournalCommit(EIoErrorCode ErrorCode, int64 ByteCount) IAS_STATISTICS_IMPL()
-	static void OnJournalCommit(const FResult& Result, int64 ByteCount) IAS_STATISTICS_IMPL()
-	static void OnCasVerification(const FResult& Result, int32 RemovedChunkCount, uint64 BlockBytesOverBudget) IAS_STATISTICS_IMPL()
+	static void OnPutChunk(EOnDemandInstallCasType CasType, uint64 ChunkSize) IAS_STATISTICS_IMPL()
+	static void OnFlush(EOnDemandInstallCasType CasType, const FResult& Result, int64 ByteCount, bool bFromDefrag) IAS_STATISTICS_IMPL()
+	static void OnJournalCommit(EOnDemandInstallCasType CasType, const FResult& Result, int64 ByteCount) IAS_STATISTICS_IMPL()
+	static void OnCasVerification(EOnDemandInstallCasType CasType, const FResult& Result, int32 RemovedChunkCount, uint64 BlockBytesOverBudget) IAS_STATISTICS_IMPL()
 	static void OnPurge(
+		EOnDemandInstallCasType CasType,
 		const FResult& Result,	
 		uint64 MaxCacheSize,
 		uint64 NewCacheSize,
 		uint64 BytesToPurge,
 		uint64 PurgedBytes) IAS_STATISTICS_IMPL()
 	static void OnDefrag(
+		EOnDemandInstallCasType CasType,
 		const FResult& Result,
 		uint64 TotalFragmentedBytes,
 		uint64 TotalCachedBytes,
 		uint64 BlocksToDefragFragmentedSize, 
 		uint64 BlocksToDefragTotalSize) IAS_STATISTICS_IMPL()
-	static void OnBlockDeleted(int64 LastAccessTicks, int64 LastModificationTicks, bool bFromDefrag) IAS_STATISTICS_IMPL()
+	static void OnBlockDeleted(EOnDemandInstallCasType CasType, int64 LastAccessTicks, int64 LastModificationTicks, bool bFromDefrag) IAS_STATISTICS_IMPL()
 	static void OnCacheUsage(
+		EOnDemandInstallCasType CasType,
 		uint64 MaxCacheSize,
 		uint64 CacheSize,
 		uint64 ReferencedBlockSize,
 		uint64 ReferencedSize,
-		uint64 FragmentedSize,
-		int64 OldestBlockAccess) IAS_STATISTICS_IMPL()
-	static void OnReadCompleted(EIoErrorCode ErrorCode, uint64 EncodedSize) IAS_STATISTICS_IMPL()
+		uint64 FragmentedSize) IAS_STATISTICS_IMPL()
+	static void OnReadCompleted(EOnDemandInstallCasType CasType, EIoErrorCode ErrorCode, uint64 EncodedSize) IAS_STATISTICS_IMPL()
 };
 
 #undef IAS_STATISTICS_IMPL
diff --git a/Engine/Source/Runtime/Experimental/IoStore/OnDemandCore/Public/IO/IoStoreOnDemand.h b/Engine/Source/Runtime/Experimental/IoStore/OnDemandCore/Public/IO/IoStoreOnDemand.h
--- a/Engine/Source/Runtime/Experimental/IoStore/OnDemandCore/Public/IO/IoStoreOnDemand.h
+++ b/Engine/Source/Runtime/Experimental/IoStore/OnDemandCore/Public/IO/IoStoreOnDemand.h
@@ -217,19 +217,26 @@ struct FOnDemandMountResult
 /** Mount completion callback. */
 using FOnDemandMountCompleted = TUniqueFunction<void(FOnDemandMountResult)>;
 
+#define ON_DEMAND_INSTALL_CAS_TYPE_LIST(XCASTYPE) \
+	XCASTYPE(General) \
+	XCASTYPE(MMap)
+
+#define ON_DEMAND_INSTALL_DECLARE_CAS_TYPE(CasType) CasType,
+
 /** Specifies a particular CAS Type when relevant */
 struct EOnDemandInstallCasTypeNS
 {
 	enum Enum : uint8
 	{
-		None = 0xff,
-		General = 0,
-		MMap,
-		Count // Count does not include 'None'
+		ON_DEMAND_INSTALL_CAS_TYPE_LIST(ON_DEMAND_INSTALL_DECLARE_CAS_TYPE)
+		Count, // Count does not include 'None'
+		None = 0xff
 	};
 };
 using EOnDemandInstallCasType = EOnDemandInstallCasTypeNS::Enum;
 
+#undef ON_DEMAND_INSTALL_DECLARE_CAS_TYPE
+
 /** Options for controlling the behavior of the install request. */
 enum class EOnDemandInstallOptions : uint16
 {
