[Chaos] Don't allow Island to sleep when SleepType is set to NeverSleep.

#Synced-CL 52201423

diff --git a/Engine/Source/Runtime/Experimental/Chaos/Private/Chaos/Island/IslandManager.cpp b/Engine/Source/Runtime/Experimental/Chaos/Private/Chaos/Island/IslandManager.cpp
--- a/Engine/Source/Runtime/Experimental/Chaos/Private/Chaos/Island/IslandManager.cpp
+++ b/Engine/Source/Runtime/Experimental/Chaos/Private/Chaos/Island/IslandManager.cpp
@@ -2691,6 +2691,14 @@ namespace Chaos::Private
 				continue;
 			}
 
+			// The particle SleepType may have changed since the node cached thresholds
+			// This is needed because we have API on the particle to SetSleepType which doesn't end up updating thresholds which is done when setting SleepType via FPBDRigidsEvolutionBase::SetParticleSleepType
+			if (Rigid->SleepType() == ESleepType::NeverSleep)
+			{
+				bWithinSleepThreshold = false;
+				break;
+			}
+
 			UpdateParticleSleepMetrics(*Rigid, Dt);
 
 			// Did we exceed the velocity threshold?
@@ -2870,6 +2878,14 @@ namespace Chaos::Private
 				continue;
 			}
 
+			// The particle SleepType may have changed since the node cached thresholds
+			// This is needed because we have API on the particle to SetSleepType which doesn't end up updating thresholds which is done when setting SleepType via FPBDRigidsEvolutionBase::SetParticleSleepType
+			if (Rigid->SleepType() == ESleepType::NeverSleep)
+			{
+				bAllNodesSleeping = false;
+				break;
+			}
+
 			// Check the particle state against the thresholds
 			int32 SleepCounterThreshold = Node->SleepCounterThreshold;
 			int32 SleepCounter = 0;
diff --git a/Engine/Source/Runtime/Experimental/Chaos/Private/Chaos/PBDRigidsEvolution.cpp b/Engine/Source/Runtime/Experimental/Chaos/Private/Chaos/PBDRigidsEvolution.cpp
--- a/Engine/Source/Runtime/Experimental/Chaos/Private/Chaos/PBDRigidsEvolution.cpp
+++ b/Engine/Source/Runtime/Experimental/Chaos/Private/Chaos/PBDRigidsEvolution.cpp
@@ -38,7 +38,7 @@ namespace Chaos
 	FAutoConsoleVariableRef CVarAccelerationStructureIsolateQueryOnlyObjects(TEXT("p.Chaos.AccelerationStructureIsolateQueryOnlyObjects"), AccelerationStructureIsolateQueryOnlyObjects, TEXT("Set to 1: QueryOnly Objects will not be moved to acceleration structures on the Physics Thread"));
 
 	CHAOS_API int32 AccelerationStructureSplitStaticAndDynamic = 1;
-	FAutoConsoleVariableRef CVarAccelerationStructureSplitStaticAndDynamic(TEXT("p.Chaos.AccelerationStructureSplitStaticDynamic"), AccelerationStructureSplitStaticAndDynamic, TEXT("Set to 1: Sort Dynamic and Static bodies into seperate acceleration structures, any other value will disable the feature"));
+	FAutoConsoleVariableRef CVarAccelerationStructureSplitStaticAndDynamic(TEXT("p.Chaos.AccelerationStructureSplitStaticDynamic"), AccelerationStructureSplitStaticAndDynamic, TEXT("Set to 1: Sort Dynamic and Static bodies into separate acceleration structures, any other value will disable the feature"));
 
 	CHAOS_API int32 AccelerationStructureUseDynamicTree = 1;
 	FAutoConsoleVariableRef CVarAccelerationStructureUseDynamicTree(TEXT("p.Chaos.AccelerationStructureUseDynamicTree"), AccelerationStructureUseDynamicTree, TEXT("Use a dynamic BVH tree structure for dynamic objects"));
@@ -1280,7 +1280,7 @@ namespace Chaos
 		const EObjectStateType InitialState = Particle->ObjectState();
 		Particle->SetSleepType(InSleepType);
 		const EObjectStateType ObjectState = Particle->ObjectState();
-		
+
 		if (FPBDRigidClusteredParticleHandle* ClusteredParticle = Particle->CastToClustered())
 		{
 			Particles.SetClusteredParticleSOA(ClusteredParticle);
@@ -1298,6 +1298,9 @@ namespace Chaos
 				Particles.MarkTransientDirtyParticle(Particle);
 			}
 		}
+
+		// Refresh cached sleep thresholds on the island node so that the island manager respects the new sleep type
+		IslandManager.UpdateParticleMaterial(Particle);
 	}
 
 	void FPBDRigidsEvolutionBase::DisableParticles(const TSet<FGeometryParticleHandle*>& InParticles)
