Chaos Cloth - Fixed crash when modifying an empty property collection. This is a regression caused by CL 52295890.

#okforgithhub public
#rnx
#Synced-CL 52300829

diff --git a/Engine/Source/Runtime/Experimental/Chaos/Private/Chaos/CollectionPropertyFacade.cpp b/Engine/Source/Runtime/Experimental/Chaos/Private/Chaos/CollectionPropertyFacade.cpp
--- a/Engine/Source/Runtime/Experimental/Chaos/Private/Chaos/CollectionPropertyFacade.cpp
+++ b/Engine/Source/Runtime/Experimental/Chaos/Private/Chaos/CollectionPropertyFacade.cpp
@@ -132,43 +132,63 @@ namespace Chaos::Softs
 
 	TArrayView<FName> FCollectionPropertyFacade::ModifyKeyNameArray()
 	{
-		MakeDirty();
-		TManagedArray<FName>& KeyNames = GetManagedArrayCollection()->ModifyAttribute<FName>(
-			PropertyFacadeNames::KeyNameName, PropertyFacadeNames::PropertyGroup);
-		return TArrayView<FName>(KeyNames.GetData(), KeyNames.Num());
+		if (KeyNameArray.Num())
+		{
+			MakeDirty();
+			TManagedArray<FName>& KeyNames = GetManagedArrayCollection()->ModifyAttribute<FName>(
+				PropertyFacadeNames::KeyNameName, PropertyFacadeNames::PropertyGroup);
+			return TArrayView<FName>(KeyNames.GetData(), KeyNames.Num());
+		}
+		return TArrayView<FName>();
 	}
 
 	TArrayView<FVector3f> FCollectionPropertyFacade::ModifyLowValueArray()
 	{
-		MakeDirty();
-		TManagedArray<FVector3f>& LowValues = GetManagedArrayCollection()->ModifyAttribute<FVector3f>(
-			PropertyFacadeNames::LowValueName, PropertyFacadeNames::PropertyGroup);
-		return TArrayView<FVector3f>(LowValues.GetData(), LowValues.Num());
+		if (LowValueArray.Num())
+		{
+			MakeDirty();
+			TManagedArray<FVector3f>& LowValues = GetManagedArrayCollection()->ModifyAttribute<FVector3f>(
+				PropertyFacadeNames::LowValueName, PropertyFacadeNames::PropertyGroup);
+			return TArrayView<FVector3f>(LowValues.GetData(), LowValues.Num());
+		}
+		return TArrayView<FVector3f>();
 	}
 
 	TArrayView<FVector3f> FCollectionPropertyFacade::ModifyHighValueArray()
 	{
-		MakeDirty();
-		TManagedArray<FVector3f>& HighValues = GetManagedArrayCollection()->ModifyAttribute<FVector3f>(
-			PropertyFacadeNames::HighValueName, PropertyFacadeNames::PropertyGroup);
-		return TArrayView<FVector3f>(HighValues.GetData(), HighValues.Num());
+		if (HighValueArray.Num())
+		{
+			MakeDirty();
+			TManagedArray<FVector3f>& HighValues = GetManagedArrayCollection()->ModifyAttribute<FVector3f>(
+				PropertyFacadeNames::HighValueName, PropertyFacadeNames::PropertyGroup);
+			return TArrayView<FVector3f>(HighValues.GetData(), HighValues.Num());
+		}
+		return TArrayView<FVector3f>();
 	}
 
 	TArrayView<FString> FCollectionPropertyFacade::ModifyStringValueArray()
 	{
-		MakeDirty();
-		TManagedArray<FString>& StringValues = GetManagedArrayCollection()->ModifyAttribute<FString>(
-			PropertyFacadeNames::StringValueName, PropertyFacadeNames::PropertyGroup);
-		return TArrayView<FString>(StringValues.GetData(), StringValues.Num());
+		if (StringValueArray.Num())
+		{
+			MakeDirty();
+			TManagedArray<FString>& StringValues = GetManagedArrayCollection()->ModifyAttribute<FString>(
+				PropertyFacadeNames::StringValueName, PropertyFacadeNames::PropertyGroup);
+			return TArrayView<FString>(StringValues.GetData(), StringValues.Num());
+		}
+		return TArrayView<FString>();
 	}
 
 	TArrayView<ECollectionPropertyFlags> FCollectionPropertyFacade::ModifyFlagsArray()
 	{
-		MakeDirty();
-		TManagedArray<uint8>& FlagBytes = GetManagedArrayCollection()->ModifyAttribute<uint8>(
-			PropertyFacadeNames::FlagsName, PropertyFacadeNames::PropertyGroup);
-		return TArrayView<ECollectionPropertyFlags>(
-			reinterpret_cast<ECollectionPropertyFlags*>(FlagBytes.GetData()), FlagBytes.Num());
+		if (FlagsArray.Num())
+		{
+			MakeDirty();
+			TManagedArray<uint8>& FlagBytes = GetManagedArrayCollection()->ModifyAttribute<uint8>(
+				PropertyFacadeNames::FlagsName, PropertyFacadeNames::PropertyGroup);
+			return TArrayView<ECollectionPropertyFlags>(
+				reinterpret_cast<ECollectionPropertyFlags*>(FlagBytes.GetData()), FlagBytes.Num());
+		}
+		return TArrayView<ECollectionPropertyFlags>();
 	}
 
 	void FCollectionPropertyFacade::SetFlags(int32 KeyIndex, ECollectionPropertyFlags Flags)
@@ -527,12 +547,19 @@ namespace Chaos::Softs
 
 				GetManagedArrayCollection()->RemoveAttribute(PropertyFacadeNames::KeyName, PropertyFacadeNames::PropertyGroup);
 			}
-			// Add Dirty attribute if missing
-			TManagedArray<bool>& DirtyAttribute = GetManagedArrayCollection()->AddAttribute<bool>(PropertyFacadeNames::DirtyName, PropertyFacadeNames::PropertyGlobalGroup);
-			if (!GetManagedArrayCollection()->NumElements(PropertyFacadeNames::PropertyGlobalGroup))
+			// Add Dirty attribute if missing from a previously valid schema
+			if (ManagedArrayCollection->HasAttribute(PropertyFacadeNames::KeyNameName, PropertyFacadeNames::PropertyGroup) &&
+				ManagedArrayCollection->HasAttribute(PropertyFacadeNames::LowValueName, PropertyFacadeNames::PropertyGroup) &&
+				ManagedArrayCollection->HasAttribute(PropertyFacadeNames::HighValueName, PropertyFacadeNames::PropertyGroup) &&
+				ManagedArrayCollection->HasAttribute(PropertyFacadeNames::StringValueName, PropertyFacadeNames::PropertyGroup) &&
+				ManagedArrayCollection->HasAttribute(PropertyFacadeNames::FlagsName, PropertyFacadeNames::PropertyGroup))
 			{
-				GetManagedArrayCollection()->AddElements(1, PropertyFacadeNames::PropertyGlobalGroup);
-				DirtyAttribute.Fill(true);
+				TManagedArray<bool>& DirtyAttribute = GetManagedArrayCollection()->AddAttribute<bool>(PropertyFacadeNames::DirtyName, PropertyFacadeNames::PropertyGlobalGroup);
+				if (!GetManagedArrayCollection()->NumElements(PropertyFacadeNames::PropertyGlobalGroup))
+				{
+					GetManagedArrayCollection()->AddElements(1, PropertyFacadeNames::PropertyGlobalGroup);
+					DirtyAttribute.Fill(true);
+				}
 			}
 		}
 	}
