Avoid loop over whole object set when looking for indirectly referenced inners

#rnx

#Synced-CL 52399147

diff --git a/Engine/Source/Runtime/Experimental/JsonObjectGraph/Private/JsonStringifyImpl.cpp b/Engine/Source/Runtime/Experimental/JsonObjectGraph/Private/JsonStringifyImpl.cpp
--- a/Engine/Source/Runtime/Experimental/JsonObjectGraph/Private/JsonStringifyImpl.cpp
+++ b/Engine/Source/Runtime/Experimental/JsonObjectGraph/Private/JsonStringifyImpl.cpp
@@ -493,14 +493,15 @@ bool FJsonStringifyImpl::TryWriteInstanceDataObjectProperties(const UObject* For
 void FJsonStringifyImpl::WriteIndirectlyReferencedContainedObjects(const UObject* ForObject)
 {
 	TArray<const UObject*> UnwrittenInners;
-	// @todo: naive implementation - is looping over every object here too slow?
-	for (const UObject* ObjectToExport : ObjectsToExport)
-	{
-		if (ObjectToExport->GetOuter() == ForObject && !ObjectsExported.Contains(ObjectToExport))
-		{
-			UnwrittenInners.Add(ObjectToExport);
-		}
-	}
+	ForEachObjectWithOuter(ForObject, 
+		[this, &UnwrittenInners](UObject* ObjectToExport)
+		{	
+			if(ObjectsToExport.Contains(ObjectToExport) && !ObjectsExported.Contains(ObjectToExport))
+			{
+				UnwrittenInners.Add(ObjectToExport);
+			}
+		},
+		EGetObjectsFlags::None, RF_Transient);
 
 	if (UnwrittenInners.Num() == 0)
 	{
