Add CasType to OnDemandErrors where relevant.

#Synced-CL 52313638

diff --git a/Engine/Source/Runtime/Experimental/IoStore/OnDemand/Private/OnDemandContentInstallReplay.h b/Engine/Source/Runtime/Experimental/IoStore/OnDemand/Private/OnDemandContentInstallReplay.h
--- a/Engine/Source/Runtime/Experimental/IoStore/OnDemand/Private/OnDemandContentInstallReplay.h
+++ b/Engine/Source/Runtime/Experimental/IoStore/OnDemand/Private/OnDemandContentInstallReplay.h
@@ -7,6 +7,7 @@
 #include "HAL/PlatformTime.h"
 #include "IO/IoStoreOnDemand.h"
 #include "IO/IoStoreOnDemandInternals.h"
+#include "IO/OnDemandError.h"
 #include "IO/PackageId.h"
 
 namespace UE::IoStore
diff --git a/Engine/Source/Runtime/Experimental/IoStore/OnDemand/Private/OnDemandContentInstaller.cpp b/Engine/Source/Runtime/Experimental/IoStore/OnDemand/Private/OnDemandContentInstaller.cpp
--- a/Engine/Source/Runtime/Experimental/IoStore/OnDemand/Private/OnDemandContentInstaller.cpp
+++ b/Engine/Source/Runtime/Experimental/IoStore/OnDemand/Private/OnDemandContentInstaller.cpp
@@ -1430,7 +1430,7 @@ void FOnDemandContentInstaller::ProcessDownloadedChunk(
 						}
 						else if (PutResult.GetValue() == false)
 						{
-							Result = MakeError(UE::UnifiedError::IoStoreOnDemand::InstallerCacheFullError::MakeError());
+							Result = MakeError(UE::UnifiedError::IoStoreOnDemand::InstallerCacheFullError::MakeError(CasType));
 						}
 					}
 				}
diff --git a/Engine/Source/Runtime/Experimental/IoStore/OnDemand/Private/OnDemandInstallCache.cpp b/Engine/Source/Runtime/Experimental/IoStore/OnDemand/Private/OnDemandInstallCache.cpp
--- a/Engine/Source/Runtime/Experimental/IoStore/OnDemand/Private/OnDemandInstallCache.cpp
+++ b/Engine/Source/Runtime/Experimental/IoStore/OnDemand/Private/OnDemandInstallCache.cpp
@@ -1445,19 +1445,19 @@ FResult FCas::Configure(const FOnDemandInstallCasConfig& Config)
 	if (Config.MinBlockSize > MaxSupportedBlockSize)
 	{
 		return MakeCasError<void>(
-			ECasErrorCode::InitializeFailed, EIoErrorCode::InvalidParameter, FString::Printf(TEXT("MinBlockSize must be less than %i"), MaxSupportedBlockSize));
+			CasType, ECasErrorCode::InitializeFailed, EIoErrorCode::InvalidParameter, FString::Printf(TEXT("MinBlockSize must be less than %i"), MaxSupportedBlockSize));
 	}
 
 	if (Config.MaxBlockSize > MaxSupportedBlockSize)
 	{
 		return MakeCasError<void>(
-			ECasErrorCode::InitializeFailed, EIoErrorCode::InvalidParameter, FString::Printf(TEXT("MaxBlockSize must be less than %i"), MaxSupportedBlockSize));
+			CasType, ECasErrorCode::InitializeFailed, EIoErrorCode::InvalidParameter, FString::Printf(TEXT("MaxBlockSize must be less than %i"), MaxSupportedBlockSize));
 	}
 
 	if (Config.MinBlockSize > Config.MaxBlockSize)
 	{
 		return MakeCasError<void>(
-			ECasErrorCode::InitializeFailed, EIoErrorCode::InvalidParameter, FString::Printf(TEXT("MinBlockSize cannot be greater than MaxBlockSize")));
+			CasType, ECasErrorCode::InitializeFailed, EIoErrorCode::InvalidParameter, FString::Printf(TEXT("MinBlockSize cannot be greater than MaxBlockSize")));
 	}
 
 	LastAccessGranularityTicks = FTimespan::FromSeconds(Config.LastAccessGranularitySeconds).GetTicks();
@@ -1487,7 +1487,7 @@ FResult FCas::Initialize(FStringView Directory, bool bDeleteExisting)
 		if (Ifm.DeleteDirectory(*RootDirectory, bRequireExists, bTree) == false)
 		{
 			return MakeCasError<void>(
-				ECasErrorCode::InitializeFailed, EIoErrorCode::DeleteError, FString::Printf(TEXT("Failed to delete CAS blocks directory '%s'"), *RootDirectory));
+				CasType, ECasErrorCode::InitializeFailed, EIoErrorCode::DeleteError, FString::Printf(TEXT("Failed to delete CAS blocks directory '%s'"), *RootDirectory));
 		}
 	}
 
@@ -1497,7 +1497,7 @@ FResult FCas::Initialize(FStringView Directory, bool bDeleteExisting)
 		if (Ifm.MakeDirectory(*RootDirectory, bTree) == false)
 		{
 			return MakeCasError<void>(
-				ECasErrorCode::InitializeFailed, EIoErrorCode::WriteError, FString::Printf(TEXT("Failed to create CAS blocks directory '%s'"), *RootDirectory));
+				CasType, ECasErrorCode::InitializeFailed, EIoErrorCode::WriteError, FString::Printf(TEXT("Failed to create CAS blocks directory '%s'"), *RootDirectory));
 		}
 	}
 
@@ -1581,7 +1581,7 @@ FResult FCas::DeleteBlock(FCasBlockId BlockId, TFunctionRef<void(const FCasAddr&
 			if (BlockState && BlockState->HandleCount > 0)
 			{
 				return MakeCasError<void>(
-					ECasErrorCode::DeleteBlockFailed, EIoErrorCode::FileHandleOpen,
+					CasType, ECasErrorCode::DeleteBlockFailed, EIoErrorCode::FileHandleOpen,
 					FString::Printf(TEXT("Cannot delete CAS block %u while it has open mapped handles"), BlockId.Id));
 			}
 
@@ -1597,7 +1597,7 @@ FResult FCas::DeleteBlock(FCasBlockId BlockId, TFunctionRef<void(const FCasAddr&
 			if (FPlatformTime::ToMilliseconds(FPlatformTime::Cycles() - StartTimeCycles) > DeleteBlockMaxWaitTimeMs)
 			{
 				return MakeCasError<void>(
-					ECasErrorCode::DeleteBlockFailed, EIoErrorCode::Timeout, FString::Printf(TEXT("Timed out waiting for pending read(s) when deleting CAS block %u"), BlockId.Id));
+					CasType, ECasErrorCode::DeleteBlockFailed, EIoErrorCode::Timeout, FString::Printf(TEXT("Timed out waiting for pending read(s) when deleting CAS block %u"), BlockId.Id));
 			}
 
 			BlockReadsDoneEvent->Wait(WaitTimeMs);
@@ -1613,7 +1613,7 @@ FResult FCas::DeleteBlock(FCasBlockId BlockId, TFunctionRef<void(const FCasAddr&
 	if (Ipf.DeleteFile(*Filename) == false)
 	{
 		return MakeCasError<void>(
-			ECasErrorCode::DeleteBlockFailed, EIoErrorCode::DeleteError, FString::Printf(TEXT("Failed to delete CAS block %u"), BlockId.Id));
+			CasType, ECasErrorCode::DeleteBlockFailed, EIoErrorCode::DeleteError, FString::Printf(TEXT("Failed to delete CAS block %u"), BlockId.Id));
 	}
 
 	if (BlockState && BlockState->PendingDefragEvent)
@@ -1657,7 +1657,7 @@ TResult<FSharedFileHandle> FCas::OpenRead(FCasBlockId BlockId)
 	if (CasType == EOnDemandInstallCasType::MMap)
 	{
 		return MakeCasError<FSharedFileHandle>(
-			ECasErrorCode::ReadBlockFailed, EIoErrorCode::InvalidCode,
+			CasType, ECasErrorCode::ReadBlockFailed, EIoErrorCode::InvalidCode,
 			FString::Printf(TEXT("MMap CAS Cannot open CAS block '%s' for read"), *Filename));
 	}
 
@@ -1667,7 +1667,7 @@ TResult<FSharedFileHandle> FCas::OpenRead(FCasBlockId BlockId)
 	if (BlockState == nullptr)
 	{
 		return MakeCasError<FSharedFileHandle>(
-			ECasErrorCode::ReadBlockFailed, EIoErrorCode::NotFound, FString::Printf(TEXT("Could not find block state for '%s'"), *Filename));
+			CasType, ECasErrorCode::ReadBlockFailed, EIoErrorCode::NotFound, FString::Printf(TEXT("Could not find block state for '%s'"), *Filename));
 	}
 
 	FFileOpenResult Result = Ipf.OpenRead(*Filename, IPlatformFile::EOpenReadFlags::AllowWrite);
@@ -1688,7 +1688,7 @@ TResult<FSharedFileHandle> FCas::OpenRead(FCasBlockId BlockId)
 	}
 
 	return MakeCasError<FSharedFileHandle>(
-		ECasErrorCode::ReadBlockFailed, EIoErrorCode::FileOpenFailed, 
+		CasType, ECasErrorCode::ReadBlockFailed, EIoErrorCode::FileOpenFailed,
 		FString::Printf(TEXT("Failed to open CAS block '%s'"), *Filename),
 		Result.StealError());
 }
@@ -1700,7 +1700,7 @@ TResult<FSharedAsyncFileHandle> FCas::OpenAsyncRead(FCasBlockId BlockId)
 	if (CasType == EOnDemandInstallCasType::MMap)
 	{
 		return MakeCasError<FSharedAsyncFileHandle>(
-			ECasErrorCode::ReadBlockFailed, EIoErrorCode::InvalidCode,
+			CasType, ECasErrorCode::ReadBlockFailed, EIoErrorCode::InvalidCode,
 			FString::Printf(TEXT("MMap CAS Cannot open CAS block '%s' for async read"), *GetBlockFilename(BlockId)));
 	}
 
@@ -1710,7 +1710,7 @@ TResult<FSharedAsyncFileHandle> FCas::OpenAsyncRead(FCasBlockId BlockId)
 	if (BlockState == nullptr)
 	{
 		return MakeCasError<FSharedAsyncFileHandle>(
-			ECasErrorCode::ReadBlockFailed, EIoErrorCode::NotFound, FString::Printf(TEXT("Could not find block state for '%s'"), *GetBlockFilename(BlockId)));
+			CasType, ECasErrorCode::ReadBlockFailed, EIoErrorCode::NotFound, FString::Printf(TEXT("Could not find block state for '%s'"), *GetBlockFilename(BlockId)));
 	}
 
 	if (FWeakAsyncFileHandle* MaybeHandle = BlockState->WeakHandle.TryGet<FWeakAsyncFileHandle>())
@@ -1744,7 +1744,7 @@ TResult<FSharedAsyncFileHandle> FCas::OpenAsyncRead(FCasBlockId BlockId)
 	}
 
 	return MakeCasError<FSharedAsyncFileHandle>(
-		ECasErrorCode::ReadBlockFailed, EIoErrorCode::FileOpenFailed,
+		CasType, ECasErrorCode::ReadBlockFailed, EIoErrorCode::FileOpenFailed,
 		FString::Printf(TEXT("Failed to open async CAS block '%s'"), *Filename),
 		HandleResult.StealError());
 }
@@ -1756,7 +1756,7 @@ TResult<FSharedMappedFileHandle> FCas::OpenMapped(FCasBlockId BlockId, int64 Max
 	if (CasType != EOnDemandInstallCasType::MMap)
 	{
 		return MakeCasError<FSharedMappedFileHandle>(
-			ECasErrorCode::OpenMappedFailed, EIoErrorCode::InvalidCode,
+			CasType, ECasErrorCode::OpenMappedFailed, EIoErrorCode::InvalidCode,
 			FString::Printf(TEXT("%s CAS Cannot open CAS block '%s' for mapped read"), LexToString(CasType), *GetBlockFilename(BlockId)));
 	}
 
@@ -1766,7 +1766,7 @@ TResult<FSharedMappedFileHandle> FCas::OpenMapped(FCasBlockId BlockId, int64 Max
 	if (BlockState == nullptr)
 	{
 		return MakeCasError<FSharedMappedFileHandle>(
-			ECasErrorCode::OpenMappedFailed, EIoErrorCode::NotFound, FString::Printf(TEXT("Could not find block state for '%s'"), *GetBlockFilename(BlockId)));
+			CasType, ECasErrorCode::OpenMappedFailed, EIoErrorCode::NotFound, FString::Printf(TEXT("Could not find block state for '%s'"), *GetBlockFilename(BlockId)));
 	}
 
 	// We can't open a mapped block if it is being defragged
@@ -1778,7 +1778,7 @@ TResult<FSharedMappedFileHandle> FCas::OpenMapped(FCasBlockId BlockId, int64 Max
 		}
 
 		return MakeCasError<FSharedMappedFileHandle>(
-			ECasErrorCode::OpenMappedFailed, EIoErrorCode::PendingDelete,
+			CasType, ECasErrorCode::OpenMappedFailed, EIoErrorCode::PendingDelete,
 			FString::Printf(TEXT("Cannot delete CAS block %u while it is being defragged"), BlockId.Id));
 	}
 
@@ -1829,7 +1829,7 @@ TResult<FSharedMappedFileHandle> FCas::OpenMapped(FCasBlockId BlockId, int64 Max
 	}
 
 	return MakeCasError<FSharedMappedFileHandle>(
-		ECasErrorCode::OpenMappedFailed, EIoErrorCode::FileOpenFailed,
+		CasType, ECasErrorCode::OpenMappedFailed, EIoErrorCode::FileOpenFailed,
 		FString::Printf(TEXT("Failed to open mapped CAS block '%s'"), *Filename),
 		HandleResult.StealError());
 }
@@ -1910,7 +1910,7 @@ TResult<FUniqueFileHandle> FCas::OpenWrite(FCasBlockId BlockId, const bool bAppe
 	}
 
 	return MakeCasError<FUniqueFileHandle>(
-		ECasErrorCode::WriteBlockFailed, EIoErrorCode::FileOpenFailed, 
+		CasType, ECasErrorCode::WriteBlockFailed, EIoErrorCode::FileOpenFailed,
 		FString::Printf(TEXT("Failed to open CAS block '%s' for writing"), *Filename),
 		OpenWriteResult.StealError());
 }
@@ -2268,7 +2268,7 @@ FResult FCas::Verify(TArray<FCasAddr>& OutRemovedChunksAddrs, TMap<FCasBlockId,
 		}
 
 		BlockIt.RemoveCurrent();
-		Result = MakeCasError<void>(ECasErrorCode::VerifyFailed, EIoErrorCode::NotFound, MoveTemp(ErrorMessage));
+		Result = MakeCasError<void>(CasType, ECasErrorCode::VerifyFailed, EIoErrorCode::NotFound, MoveTemp(ErrorMessage));
 	}
 
 	UE_LOG(LogIoStoreOnDemand, Log, TEXT("Verified %d CAS blocks of total %.2lf MiB"),
@@ -2301,7 +2301,7 @@ FResult FCas::Verify(TArray<FCasAddr>& OutRemovedChunksAddrs, TMap<FCasBlockId,
 			OutRemovedChunksAddrs.Add(It->Key);
 			It.RemoveCurrent();
 
-			Result = MakeCasError<void>(ECasErrorCode::VerifyFailed, EIoErrorCode::NotFound, MoveTemp(ErrorMessage));
+			Result = MakeCasError<void>(CasType, ECasErrorCode::VerifyFailed, EIoErrorCode::NotFound, MoveTemp(ErrorMessage));
 		}
 	}
 
@@ -3437,7 +3437,7 @@ TResult<bool> FOnDemandInstallCache::PutChunk(EOnDemandInstallCasType CasType, F
 	{
 		using namespace UE::UnifiedError::IoStoreOnDemand;
 		return MakeCasError<bool>(
-			ECasErrorCode::WriteBlockFailed, EIoErrorCode::InvalidParameter, FString::Printf(TEXT("CAS %s is not configured"), LexToString(CasType)));
+			CasType, ECasErrorCode::WriteBlockFailed, EIoErrorCode::InvalidParameter, FString::Printf(TEXT("CAS %s is not configured"), LexToString(CasType)));
 	}
 
 	if (CacheState->PendingChunks.TotalSize > FPendingChunks::MaxPendingBytes)
@@ -3472,7 +3472,7 @@ FResult FOnDemandInstallCache::Purge(EOnDemandInstallCasType CasType, const uint
 	if (CacheState == nullptr)
 	{
 		return MakeCasError<void>(
-			ECasErrorCode::DeleteBlockFailed, EIoErrorCode::InvalidParameter, FString::Printf(TEXT("CAS %s is not configured"), LexToString(CasType)));
+			CasType, ECasErrorCode::DeleteBlockFailed, EIoErrorCode::InvalidParameter, FString::Printf(TEXT("CAS %s is not configured"), LexToString(CasType)));
 	}
 
 	TUniqueLock Lock(CacheState->PurgeDefragMutex);
@@ -3519,7 +3519,6 @@ FResult FOnDemandInstallCache::Purge(EOnDemandInstallCasType CasType, const uint
 		}
 	}
 
-	// TODO: IadMultiCAS: analytics need to track CAS type as well
 	FOnDemandInstallCacheStats::OnCacheUsage(
 		CacheState->Cas.GetType(), CacheState->MaxCacheSize, TotalCachedBytes, TotalReferencedBlockBytes, ReferencedBytes, FragmentedBytes);
 
@@ -3579,7 +3578,7 @@ FResult FOnDemandInstallCache::PurgeAllUnreferenced(EOnDemandInstallCasType CasT
 	if (CacheState == nullptr)
 	{
 		return MakeCasError<void>(
-			ECasErrorCode::DeleteBlockFailed, EIoErrorCode::InvalidParameter, FString::Printf(TEXT("CAS %s is not configured"), LexToString(CasType)));
+			CasType, ECasErrorCode::DeleteBlockFailed, EIoErrorCode::InvalidParameter, FString::Printf(TEXT("CAS %s is not configured"), LexToString(CasType)));
 	}
 
 	TUniqueLock Lock(CacheState->PurgeDefragMutex);
@@ -3650,7 +3649,7 @@ FResult FOnDemandInstallCache::DefragAll(EOnDemandInstallCasType CasType, const
 	if (CacheState == nullptr)
 	{
 		return MakeCasError<void>(
-			ECasErrorCode::DeleteBlockFailed, EIoErrorCode::InvalidParameter, FString::Printf(TEXT("CAS %s is not configured"), LexToString(CasType)));
+			CasType, ECasErrorCode::DeleteBlockFailed, EIoErrorCode::InvalidParameter, FString::Printf(TEXT("CAS %s is not configured"), LexToString(CasType)));
 	}
 
 	TUniqueLock Lock(CacheState->PurgeDefragMutex);
@@ -3912,13 +3911,13 @@ FResult FOnDemandInstallCache::Reset()
 	if (Ifm.DeleteDirectory(*CacheDirectory, false, bTree) == false)
 	{
 		return MakeCasError<void>(
-			ECasErrorCode::InitializeFailed, EIoErrorCode::DeleteError, FString::Printf(TEXT("Failed to delete cache directory '%s'"), *CacheDirectory));
+			EOnDemandInstallCasType::None, ECasErrorCode::InitializeFailed, EIoErrorCode::DeleteError, FString::Printf(TEXT("Failed to delete cache directory '%s'"), *CacheDirectory));
 	}
 
 	if (Ifm.MakeDirectory(*CacheDirectory, bTree) == false)
 	{
 		return MakeCasError<void>(
-			ECasErrorCode::InitializeFailed, EIoErrorCode::WriteError, FString::Printf(TEXT("Failed to create cache directory '%s'"), *CacheDirectory));
+			EOnDemandInstallCasType::None, ECasErrorCode::InitializeFailed, EIoErrorCode::WriteError, FString::Printf(TEXT("Failed to create cache directory '%s'"), *CacheDirectory));
 	}
 
 	for (FCacheState& CacheState : CacheStates)
@@ -4009,7 +4008,8 @@ FResult FOnDemandInstallCache::InitialVerify()
 						ToMiB(TotalBytesToPurge), LexToString(CacheState.Cas.GetType()), ToMiB(TotalPurgedBytes));
 
 					UE_LOG(LogIoStoreOnDemand, Error, TEXT("%s"), *ErrorMessage);
-					return MakeError(UE::UnifiedError::IoStoreOnDemand::InstallCachePurgeError::MakeError());
+					UE::UnifiedError::IoStoreOnDemand::FInstallCacheErrorContext Ctx = MakeInstallCacheErrorContext(CacheState, CacheSize, TotalBytesToPurge);
+					return MakeError(UE::UnifiedError::IoStoreOnDemand::InstallCachePurgeError::MakeError(MoveTemp(Ctx), UE::UnifiedError::EDetailFilter::All));
 				}
 
 				UE_LOG(LogIoStoreOnDemand, Log, TEXT("Successfully purged %.2lf MiB from CAS %s"), ToMiB(TotalPurgedBytes), LexToString(CacheState.Cas.GetType()));
@@ -4499,6 +4499,7 @@ FResult FOnDemandInstallCache::Defrag(
 					if (FileHandle->Seek(ReffedChunk.BlockOffset) == false)
 					{
 						ReadResult = MakeCasError<void>(
+							CacheState.Cas.GetType(),
 							ECasErrorCode::ReadBlockFailed, EIoErrorCode::FileSeekFailed, FString::Printf(TEXT("Failed to seek to offset %u in cache block %u"),
 								ReffedChunk.BlockOffset, DefragBlock.BlockId.Id));
 						continue;
@@ -4506,6 +4507,7 @@ FResult FOnDemandInstallCache::Defrag(
 					if (FileHandle->Read(Buffer.GetData(), Buffer.GetSize()) == false)
 					{
 						ReadResult = MakeCasError<void>(
+							CacheState.Cas.GetType(),
 							ECasErrorCode::ReadBlockFailed, EIoErrorCode::ReadError, FString::Printf(TEXT("Failed to read %" UINT64_FMT " bytes at offset %u in cache block %u"),
 								Buffer.GetSize(), ReffedChunk.BlockOffset, DefragBlock.BlockId.Id));
 						continue;
@@ -4550,7 +4552,8 @@ FResult FOnDemandInstallCache::Defrag(
 						{
 							.ChunkId		= ReffedChunk.ChunkId, // May not be unique, just return the first found for debugging purposes
 							.ExpectedHash	= ReffedChunk.Hash,
-							.ActualHash		= ChunkHash
+							.ActualHash		= ChunkHash,
+							.CasType		= CacheState.Cas.GetType()
 						}));
 					FOnDemandInstallCacheStats::OnDefrag(CacheState.Cas.GetType(), Result, TotalFragmentedBytes, TotalCachedBytes, BlocksToDefragFragmentedSize, BlocksToDefragTotalSize);
 					return Result;
@@ -4627,7 +4630,7 @@ FResult FOnDemandInstallCache::Flush(EOnDemandInstallCasType CasType)
 	if (CacheState == nullptr)
 	{
 		return MakeCasError<void>(
-			ECasErrorCode::WriteBlockFailed, EIoErrorCode::InvalidParameter, FString::Printf(TEXT("CAS %s is not configured"), LexToString(CasType)));
+			CasType, ECasErrorCode::WriteBlockFailed, EIoErrorCode::InvalidParameter, FString::Printf(TEXT("CAS %s is not configured"), LexToString(CasType)));
 	}
 
 	if (FResult Result = FlushPendingChunks(*CacheState, CacheState->PendingChunks); Result.HasError())
@@ -4886,6 +4889,7 @@ FResult FOnDemandInstallCache::FlushPendingChunksImpl(FCacheState& CacheState, c
 			if (CasFileHandle->Write(Ar.GetData(), BytesToWrite) == false)
 			{
 				FResult Result = MakeCasError<void>(
+					CacheState.Cas.GetType(),
 					ECasErrorCode::WriteBlockFailed, EIoErrorCode::WriteError, FString::Printf(TEXT("Failed to write %" INT64_FMT " bytes to CAS block %u"), BytesToWrite, CurrentBlockId.Id));
 				FInstallCacheErrorContext Ctx = MakeInstallCacheErrorContext(CacheState);
 				Result.GetError().PushErrorContext(MoveTemp(Ctx));
@@ -4898,6 +4902,7 @@ FResult FOnDemandInstallCache::FlushPendingChunksImpl(FCacheState& CacheState, c
 			if (CasFileHandle->Flush() == false)
 			{
 				FResult Result = MakeCasError<void>(
+					CacheState.Cas.GetType(),
 					ECasErrorCode::WriteBlockFailed, EIoErrorCode::FileFlushFailed, FString::Printf(TEXT("Failed to flush %" INT64_FMT " bytes to CAS block %u"), BytesToWrite, CurrentBlockId.Id));
 				FInstallCacheErrorContext Ctx = MakeInstallCacheErrorContext(CacheState);
 				Result.GetError().PushErrorContext(MoveTemp(Ctx));
diff --git a/Engine/Source/Runtime/Experimental/IoStore/OnDemand/Private/OnDemandIoStore.cpp b/Engine/Source/Runtime/Experimental/IoStore/OnDemand/Private/OnDemandIoStore.cpp
--- a/Engine/Source/Runtime/Experimental/IoStore/OnDemand/Private/OnDemandIoStore.cpp
+++ b/Engine/Source/Runtime/Experimental/IoStore/OnDemand/Private/OnDemandIoStore.cpp
@@ -1093,7 +1093,7 @@ FIoStatus FOnDemandIoStore::InitializeInstallCache()
 				}
 				else if (PutResult.GetValue() == false)
 				{
-					Result = MakeError(UE::UnifiedError::IoStoreOnDemand::InstallerCacheFullError::MakeError());
+					Result = MakeError(UE::UnifiedError::IoStoreOnDemand::InstallerCacheFullError::MakeError(EOnDemandInstallCasType::General));
 				}
 
 				Seed++;
diff --git a/Engine/Source/Runtime/Experimental/IoStore/OnDemandCore/Public/IO/IoStoreOnDemand.h b/Engine/Source/Runtime/Experimental/IoStore/OnDemandCore/Public/IO/IoStoreOnDemand.h
--- a/Engine/Source/Runtime/Experimental/IoStore/OnDemandCore/Public/IO/IoStoreOnDemand.h
+++ b/Engine/Source/Runtime/Experimental/IoStore/OnDemandCore/Public/IO/IoStoreOnDemand.h
@@ -6,7 +6,7 @@
 #include "Containers/UnrealString.h"
 #include "Containers/StringFwd.h"
 #include "Containers/SharedString.h"
-#include "IO/OnDemandError.h"
+#include "Experimental/UnifiedError/UnifiedError.h"
 #include "IO/OnDemandHostGroup.h"
 #include "IO/OnDemandToc.h"
 #include "IO/PackageId.h"
diff --git a/Engine/Source/Runtime/Experimental/IoStore/OnDemandCore/Public/IO/OnDemandError.h b/Engine/Source/Runtime/Experimental/IoStore/OnDemandCore/Public/IO/OnDemandError.h
--- a/Engine/Source/Runtime/Experimental/IoStore/OnDemandCore/Public/IO/OnDemandError.h
+++ b/Engine/Source/Runtime/Experimental/IoStore/OnDemandCore/Public/IO/OnDemandError.h
@@ -3,10 +3,10 @@
 #pragma once
 
 #include "Experimental/UnifiedError/UnifiedError.h"
-#include "Experimental/UnifiedError/CoreErrorTypes.h"
 #include "GenericPlatform/GenericPlatformFile.h"
 #include "IO/IoChunkId.h"
 #include "IO/IoStatus.h"
+#include "IO/IoStoreOnDemand.h"
 #include "Templates/ValueOrError.h"
 
 #define UE_API IOSTOREONDEMANDCORE_API
@@ -25,7 +25,7 @@ UE_DECLARE_ERROR(UE_API,			InstallCacheVerificationError,		8,	IoStoreOnDemand, N
 UE_DECLARE_ERROR(UE_API,			CasError,							9,	IoStoreOnDemand, NSLOCTEXT("IoStoreOnDemand", "CasError", "Cas error"));
 UE_DECLARE_ERROR(UE_API,			CasJournalError,					10,	IoStoreOnDemand, NSLOCTEXT("IoStoreOnDemand", "CasJournalError", "Cas journal error"));
 UE_DECLARE_ERROR(UE_API,			CasSnapshotError,					11,	IoStoreOnDemand, NSLOCTEXT("IoStoreOnDemand", "CasSnapshotError", "Cas snapshot error"));
-UE_DECLARE_ERROR(UE_API,			InstallerCacheFullError,			12, IoStoreOnDemand, NSLOCTEXT("IoStoreOnDemand", "InstallerCacheFullError", "Cache full error"));
+UE_DECLARE_ERROR_ONEPARAM(UE_API,	InstallerCacheFullError,			12, IoStoreOnDemand, NSLOCTEXT("IoStoreOnDemand", "InstallerCacheFullError", "{CasType} Cache full error"), UE::IoStore::EOnDemandInstallCasType, CasType, UE::IoStore::EOnDemandInstallCasType::None);
 UE_DECLARE_ERROR(UE_API,			InstallReplayLoadError,				13, IoStoreOnDemand, NSLOCTEXT("IoStoreOnDemand", "InstallReplayLoadError", "Failed to load install replay"));
 UE_DECLARE_ERROR(UE_API,			InstallReplaySaveError,				14, IoStoreOnDemand, NSLOCTEXT("IoStoreOnDemand", "InstallReplaySaveError", "Failed to save install replay"));
 UE_DECLARE_ERROR(UE_API,			DownloadRequired,					15, IoStoreOnDemand, NSLOCTEXT("IoStoreOnDemand", "DownloadRequired", "DoNotDownload flag specified when download required"));
@@ -104,10 +104,12 @@ struct FCasErrorContext
 	FString				ErrorMessage;
 	FFileSystemError	FileSystemError;
 	uint32				SystemErrorCode = 0;
+	UE::IoStore::EOnDemandInstallCasType CasType	= UE::IoStore::EOnDemandInstallCasType::None;
 };
 
 template <typename ResultType>
 inline UE::IoStore::TResult<ResultType> MakeCasError(
+	UE::IoStore::EOnDemandInstallCasType CasType,
 	ECasErrorCode ErrorCode,
 	EIoErrorCode IoErrorCode,
 	FString&& ErrorMessage,
@@ -120,7 +122,8 @@ inline UE::IoStore::TResult<ResultType> MakeCasError(
 		.IoErrorCode		= IoErrorCode,
 		.ErrorMessage		= MoveTemp(ErrorMessage),
 		.FileSystemError	= MoveTemp(FileSystemError),
-		.SystemErrorCode	= SystemErrorCode
+		.SystemErrorCode	= SystemErrorCode,
+		.CasType			= CasType
 	}, UE::UnifiedError::EDetailFilter::All));
 }
 
@@ -193,6 +196,7 @@ struct FChunkHashMismatchErrorContext
 	FIoChunkId	ChunkId;
 	FIoHash		ExpectedHash;
 	FIoHash		ActualHash;
+	UE::IoStore::EOnDemandInstallCasType CasType = UE::IoStore::EOnDemandInstallCasType::None;
 };
 
 struct FInstallCacheErrorContext
@@ -223,13 +227,14 @@ inline void SerializeForLog(FCbWriter& Writer, const UE::UnifiedError::IoStoreOn
 {
 	Writer.BeginObject();
 	Writer.AddString(UTF8TEXTVIEW("$type"), TErrorStructFeatures<UE::UnifiedError::IoStoreOnDemand::FCasErrorContext>::GetErrorContextTypeNameAsString());
-	Writer.AddString(UTF8TEXTVIEW("$format"), TEXT("({ErrorCodeText}: {ErrorMessage} [FileSystemError: {FileSystemError}])"));
+	Writer.AddString(UTF8TEXTVIEW("$format"), TEXT("({CasType} CAS - {ErrorCodeText}: {ErrorMessage} [FileSystemError: {FileSystemError}])"));
 	Writer.AddInteger(UTF8TEXTVIEW("ErrorCode"), uint32(Ctx.ErrorCode));
 	Writer.AddString(UTF8TEXTVIEW("ErrorCodeText"), ToString(Ctx.ErrorCode));
 	Writer.AddInteger(UTF8TEXTVIEW("IoErrorCode"), uint32(Ctx.IoErrorCode));
 	Writer.AddInteger(UTF8TEXTVIEW("SystemErrorCode"), Ctx.SystemErrorCode);
 	Writer.AddString(UTF8TEXTVIEW("IoErrorCodeText"), GetIoErrorText(Ctx.IoErrorCode));
 	Writer.AddString(UTF8TEXTVIEW("ErrorMessage"), Ctx.ErrorMessage);
+	Writer.AddString(UTF8TEXTVIEW("CasType"), LexToString(Ctx.CasType));
 
 	TUtf8StringBuilder<512> StringBuilder;
 	StringBuilder << Ctx.FileSystemError;
@@ -251,10 +256,11 @@ inline void SerializeForLog(FCbWriter& Writer, const UE::UnifiedError::IoStoreOn
 {
 	Writer.BeginObject();
 	Writer.AddString(UTF8TEXTVIEW("$type"), TErrorStructFeatures<UE::UnifiedError::IoStoreOnDemand::FChunkHashMismatchErrorContext>::GetErrorContextTypeNameAsString());
-	Writer.AddString(UTF8TEXTVIEW("$format"), TEXT("(ChunkId: {ChunkId}, ExpectedHash: {ExpectedHash}, ActualHash: {ActualHash})"));
+	Writer.AddString(UTF8TEXTVIEW("$format"), TEXT("({CasType} CAS - ChunkId: {ChunkId}, ExpectedHash: {ExpectedHash}, ActualHash: {ActualHash})"));
 	Writer.AddString(UTF8TEXTVIEW("ChunkId"), LexToString(Ctx.ChunkId));
 	Writer.AddString(UTF8TEXTVIEW("ExpectedHash"), LexToString(Ctx.ExpectedHash));
 	Writer.AddString(UTF8TEXTVIEW("ActualHash"), LexToString(Ctx.ActualHash));
+	Writer.AddString(UTF8TEXTVIEW("CasType"), LexToString(Ctx.CasType));
 	Writer.EndObject();
 }
 ////////////////////////////////////////////////////////////////////////////////
