GaussSeidelMainConstraint.cpp: wrap particle failure UE_LOG with #if WITH_EDITOR

#Synced-CL 52338030

diff --git a/Engine/Source/Runtime/Experimental/Chaos/Private/Chaos/Deformable/GaussSeidelMainConstraint.cpp b/Engine/Source/Runtime/Experimental/Chaos/Private/Chaos/Deformable/GaussSeidelMainConstraint.cpp
--- a/Engine/Source/Runtime/Experimental/Chaos/Private/Chaos/Deformable/GaussSeidelMainConstraint.cpp
+++ b/Engine/Source/Runtime/Experimental/Chaos/Private/Chaos/Deformable/GaussSeidelMainConstraint.cpp
@@ -15,6 +15,7 @@ namespace Chaos::Softs
 
 	FAutoConsoleVariableRef CVarClothCGTol(TEXT("p.Chaos.Cloth.CGTol"), CGTol, TEXT("CG Tolerance [def: 1e-4]"));
 
+
 	void FGaussSeidelMainConstraints::AddStaticConstraints(const TArray<TArray<int32>>& ExtraConstraints, const TArray<TArray<int32>>& ExtraIncidentElements, const TArray<TArray<int32>>& ExtraIncidentElementsLocal)
 	{
 		const TArray<TArray<int32>>* IncidentElementsPtr = &ExtraIncidentElements;
@@ -515,10 +516,12 @@ namespace Chaos::Softs
 
 		CurrentIt ++;
 
+	#if WITH_EDITOR
 		if (ParticleFailureCounter.load(std::memory_order_relaxed) > 0)
 		{
 			UE_LOG(LogDeformableGaussSeidelMainConstraint, Warning, TEXT("%d Particle(s) are skipped because of too large dx size"), ParticleFailureCounter.load(std::memory_order_relaxed));
 		}
+#endif
 	}
 
 	void FGaussSeidelMainConstraints::InitStaticColor(const FSolverParticlesRange& ParticlesRange, const TPBDActiveView<FSolverParticles>* InParticleActiveView)
