Widen support for dataflow attachement across the dataflow editor and blueprint library
- Move dataflow attachement to runtime dataflow folder
- Implement and upgraded utility functions to account for attachments
- Update all places in the code that was casting directly to the instance interface to use the new utility functions instead

#Synced-CL 52432806

diff --git a/Engine/Plugins/Experimental/Dataflow/Source/DataflowEditor/Private/Dataflow/AssetDefinition_DataflowAsset.cpp b/Engine/Plugins/Experimental/Dataflow/Source/DataflowEditor/Private/Dataflow/AssetDefinition_DataflowAsset.cpp
--- a/Engine/Plugins/Experimental/Dataflow/Source/DataflowEditor/Private/Dataflow/AssetDefinition_DataflowAsset.cpp
+++ b/Engine/Plugins/Experimental/Dataflow/Source/DataflowEditor/Private/Dataflow/AssetDefinition_DataflowAsset.cpp
@@ -107,7 +107,7 @@ namespace UE::DataflowAssetDefinitionHelpers
 			return false;
 		}
 
-		IDataflowInstanceInterface* Interface = Cast<IDataflowInstanceInterface>(BoundAsset);
+		IDataflowInstanceInterface* Interface = UE::Dataflow::InstanceUtils::GetDataflowInterfaceFromObject(BoundAsset);
 		if (!Interface)
 		{
 			// the bound asset do not support embedding
@@ -144,7 +144,7 @@ namespace UE::DataflowAssetDefinitionHelpers
 			return false;
 		}
 
-		IDataflowInstanceInterface* Interface = Cast<IDataflowInstanceInterface>(BoundAsset);
+		IDataflowInstanceInterface* Interface = UE::Dataflow::InstanceUtils::GetDataflowInterfaceFromObject(BoundAsset);
 		if (!Interface)
 		{
 			// the bound asset do not support embedding
@@ -222,11 +222,15 @@ namespace UE::DataflowAssetDefinitionHelpers
 	 *   - Existing asset : sets the selected shared UDataflow on Asset without duplication.
 	 *   - None / Cancelled : leaves Asset unchanged.
 	 *
-	 * Asset must implement IDataflowInstanceInterface for the Template/ExistingAsset paths to take effect.
 	 * Always returns true — the picker is a single non-repeating step with no sub-dialogs.
 	 */
 	bool SetDataflowFromTemplatePicker(UObject* Asset, bool bShowDefaultBlankOption)
 	{
+		if (!IsValid(Asset))
+		{
+			return true;
+		}
+
 		const TArray<FDataflowTemplateOption> Templates =
 			FDataflowTemplateRegistry::Get().GetTemplateOptions(Asset->GetClass(), bShowDefaultBlankOption);
 
@@ -238,12 +242,6 @@ namespace UE::DataflowAssetDefinitionHelpers
 		{
 		case EDataflowPickerResult::Template:
 		{
-			IDataflowInstanceInterface* Interface = Cast<IDataflowInstanceInterface>(Asset);
-			if (!Interface)
-			{
-				return true;
-			}
-
 			// Load the chosen template; nullptr means Blank (TemplateId == NAME_None).
 			const FString TemplatePath = PickerResult.TemplateId.ToString();
 			UDataflow* const Source = (TemplatePath.IsEmpty() || PickerResult.TemplateId.IsNone())
@@ -263,7 +261,7 @@ namespace UE::DataflowAssetDefinitionHelpers
 					Embedded->ClearFlags(RF_Public | RF_Standalone);
 					Embedded->SetFlags(RF_Transactional);
 				}
-				Interface->GetDataflowInstance().SetDataflowAsset(Embedded);
+				UE::Dataflow::InstanceUtils::SetDataflowAssetOnObject(Embedded, Asset);
 				Asset->Modify();
 			}
 			return true;
@@ -272,11 +270,9 @@ namespace UE::DataflowAssetDefinitionHelpers
 		case EDataflowPickerResult::ExistingAsset:
 		{
 			// Reference the selected asset directly — no duplication.
-			IDataflowInstanceInterface* Interface = Cast<IDataflowInstanceInterface>(Asset);
-			UDataflow* const SharedDataflow = Cast<UDataflow>(PickerResult.SelectedAsset.GetAsset());
-			if (Interface && SharedDataflow)
+			if (UDataflow* const SharedDataflow = Cast<UDataflow>(PickerResult.SelectedAsset.GetAsset()))
 			{
-				Interface->GetDataflowInstance().SetDataflowAsset(SharedDataflow);
+				UE::Dataflow::InstanceUtils::SetDataflowAssetOnObject(SharedDataflow, Asset);
 				Asset->Modify();
 			}
 			return true;
@@ -341,7 +337,7 @@ namespace UE::DataflowAssetDefinitionHelpers
 			{
 				for (UObject* SelectedAsset : SelectedAssets)
 				{
-					if (IDataflowInstanceInterface* Interface = Cast<IDataflowInstanceInterface>(SelectedAsset))
+					if (IDataflowInstanceInterface* Interface = UE::Dataflow::InstanceUtils::GetDataflowInterfaceFromObject(SelectedAsset))
 					{
 						if (UDataflow* DataflowObject = Interface->GetDataflowInstance().GetDataflowAsset())
 						{
@@ -441,7 +437,7 @@ namespace UE::DataflowAssetDefinitionHelpers
 			{
 				for (UObject* SelectedAsset : SelectedAssets)
 				{
-					if (IDataflowInstanceInterface* Interface = Cast<IDataflowInstanceInterface>(SelectedAsset))
+					if (IDataflowInstanceInterface* Interface = UE::Dataflow::InstanceUtils::GetDataflowInterfaceFromObject(SelectedAsset))
 					{
 						if (Interface->GetDataflowInstance().UpdateOwnerAsset(/*bUpdateDependentAsset*/true))
 						{
@@ -462,7 +458,7 @@ namespace UE::DataflowAssetDefinitionHelpers
 			{
 				for (UObject* SelectedAsset : SelectedAssets)
 				{
-					if (IDataflowInstanceInterface* Interface = Cast<IDataflowInstanceInterface>(SelectedAsset))
+					if (IDataflowInstanceInterface* Interface = UE::Dataflow::InstanceUtils::GetDataflowInterfaceFromObject(SelectedAsset))
 					{
 						if (UDataflow* DataflowObject = Interface->GetDataflowInstance().GetDataflowAsset())
 						{
diff --git a/Engine/Plugins/Experimental/Dataflow/Source/DataflowEditor/Private/Dataflow/DataflowEditorToolkit.cpp b/Engine/Plugins/Experimental/Dataflow/Source/DataflowEditor/Private/Dataflow/DataflowEditorToolkit.cpp
--- a/Engine/Plugins/Experimental/Dataflow/Source/DataflowEditor/Private/Dataflow/DataflowEditorToolkit.cpp
+++ b/Engine/Plugins/Experimental/Dataflow/Source/DataflowEditor/Private/Dataflow/DataflowEditorToolkit.cpp
@@ -1482,10 +1482,10 @@ void FDataflowEditorToolkit::OnAssetsSavedAs(const TArray<UObject*>& SavedObject
 		{
 			NewDataflowAsset = Cast<UDataflow>(SavedObj);
 		}
-		else if (Cast<IDataflowInstanceInterface>(SavedObj))
+		else if (IDataflowInstanceInterface* Interface = UE::Dataflow::InstanceUtils::GetDataflowInterfaceFromObject(SavedObj))
 		{
 			NewAssetWithInterface = SavedObj;
-			NewDataflowInterface = Cast<IDataflowInstanceInterface>(SavedObj);
+			NewDataflowInterface = Interface;
 		}
 		else if (SavedObj->IsA<UChaosCacheCollection>())
 		{
@@ -1495,21 +1495,6 @@ void FDataflowEditorToolkit::OnAssetsSavedAs(const TArray<UObject*>& SavedObject
 		{
 			NewGeomCache = Cast<UGeometryCache>(SavedObj);
 		}
-		else if (IInterface_AssetUserData* AssetUserDataProvider = Cast<IInterface_AssetUserData>(SavedObj))
-		{
-			if (const TArray<UAssetUserData*>* AssetUserDataArray = AssetUserDataProvider->GetAssetUserDataArray())
-			{
-				for (UAssetUserData* UserData : *AssetUserDataArray)
-				{
-					if (IDataflowInstanceInterface* DataflowInterface = Cast<IDataflowInstanceInterface>(UserData))
-					{
-						NewDataflowInterface = DataflowInterface;
-						NewAssetWithInterface = SavedObj;
-						break;
-					}
-				}
-			}
-		}
 	}
 
 	TArray<UPackage*> PackagesToSave;
@@ -1805,7 +1790,7 @@ void FDataflowEditorToolkit::ResetEditor()
 	{
 		if (EditorContent)
 		{
-			if (IDataflowInstanceInterface* DataflowInstanceInterface = Cast<IDataflowInstanceInterface>(EditorContent->GetDataflowOwner()))
+			if (IDataflowInstanceInterface* DataflowInstanceInterface = UE::Dataflow::InstanceUtils::GetDataflowInterfaceFromObject(EditorContent->GetDataflowOwner()))
 			{
 				// close all subgraphs tabs and selection to make sure no other widget refer to this graphs data
 				const TArray<TSharedPtr<SDockTab>> Tabs = DocumentManager->GetAllDocumentTabs();
diff --git a/Engine/Plugins/Experimental/Dataflow/Source/DataflowEditor/Private/Dataflow/DataflowMembersWidget.cpp b/Engine/Plugins/Experimental/Dataflow/Source/DataflowEditor/Private/Dataflow/DataflowMembersWidget.cpp
--- a/Engine/Plugins/Experimental/Dataflow/Source/DataflowEditor/Private/Dataflow/DataflowMembersWidget.cpp
+++ b/Engine/Plugins/Experimental/Dataflow/Source/DataflowEditor/Private/Dataflow/DataflowMembersWidget.cpp
@@ -516,7 +516,7 @@ IDataflowInstanceInterface* SDataflowMembersWidget::GetDataflowInstanceInterface
 {
 	if (EditedAssetWeakPtr.IsValid())
 	{
-		return Cast<IDataflowInstanceInterface>(EditedAssetWeakPtr.Get());
+		return UE::Dataflow::InstanceUtils::GetDataflowInterfaceFromObject(EditedAssetWeakPtr.Get());
 	}
 	return nullptr;
 }
diff --git a/Engine/Plugins/Experimental/Dataflow/Source/DataflowEditor/Private/Dataflow/DataflowSimulationScene.cpp b/Engine/Plugins/Experimental/Dataflow/Source/DataflowEditor/Private/Dataflow/DataflowSimulationScene.cpp
--- a/Engine/Plugins/Experimental/Dataflow/Source/DataflowEditor/Private/Dataflow/DataflowSimulationScene.cpp
+++ b/Engine/Plugins/Experimental/Dataflow/Source/DataflowEditor/Private/Dataflow/DataflowSimulationScene.cpp
@@ -641,7 +641,7 @@ void FDataflowSimulationScene::ReadPreviewPropertiesFromDataflowOwner()
 	{
 		if (const UDataflowBaseContent* const EditorContent = GetEditorContent())
 		{
-			if (const IDataflowInstanceInterface* const DataflowInstanceInterface = Cast<IDataflowInstanceInterface>(EditorContent->GetDataflowOwner()))
+			if (const IDataflowInstanceInterface* const DataflowInstanceInterface = UE::Dataflow::InstanceUtils::GetDataflowInterfaceFromObject(EditorContent->GetDataflowOwner()))
 			{
 				const FDataflowInstance& DataflowInstance = DataflowInstanceInterface->GetDataflowInstance();
 				SceneDescription->BlueprintVariables.CopyMatchingValuesByName(DataflowInstance.GetPreviewBlueprintVariableOverrides());  // When reading, match the asset variables
@@ -657,7 +657,7 @@ void FDataflowSimulationScene::WritePreviewPropertiesToDataflowOwner() const
 		if (const UDataflowBaseContent* const EditorContent = GetEditorContent())
 		{
 			UObject* const DataflowOwner = EditorContent->GetDataflowOwner();
-			if (IDataflowInstanceInterface* const DataflowInstanceInterface = Cast<IDataflowInstanceInterface>(DataflowOwner))
+			if (IDataflowInstanceInterface* const DataflowInstanceInterface = UE::Dataflow::InstanceUtils::GetDataflowInterfaceFromObject(DataflowOwner))
 			{
 				FDataflowInstance& DataflowInstance = DataflowInstanceInterface->GetDataflowInstance();
 				DataflowInstance.GetPreviewBlueprintVariableOverrides() = SceneDescription->BlueprintVariables;  // When writing, clobber the asset variables
diff --git a/Engine/Plugins/Experimental/Dataflow/Source/DataflowEnginePlugin/Private/Dataflow/DataflowAttachment.cpp b/Engine/Plugins/Experimental/Dataflow/Source/DataflowEnginePlugin/Private/Dataflow/DataflowAttachment.cpp
deleted file mode 100644
diff --git a/Engine/Plugins/Experimental/Dataflow/Source/DataflowEnginePlugin/Public/Dataflow/DataflowAttachment.h b/Engine/Plugins/Experimental/Dataflow/Source/DataflowEnginePlugin/Public/Dataflow/DataflowAttachment.h
deleted file mode 100644
diff --git a/Engine/Source/Runtime/Experimental/Dataflow/Engine/Private/Dataflow/DataflowAttachment.cpp b/Engine/Source/Runtime/Experimental/Dataflow/Engine/Private/Dataflow/DataflowAttachment.cpp
new file mode 100644
--- /dev/null
+++ b/Engine/Source/Runtime/Experimental/Dataflow/Engine/Private/Dataflow/DataflowAttachment.cpp
@@ -0,0 +1,20 @@
+// Copyright Epic Games, Inc. All Rights Reserved.
+
+#include "Dataflow/DataflowAttachment.h"
+#include "PhysicsEngine/PhysicsAsset.h"
+
+const FDataflowInstance& UDataflowAttachment::GetDataflowInstance() const
+{
+	return Instance;
+}
+
+FDataflowInstance& UDataflowAttachment::GetDataflowInstance()
+{
+	return Instance;
+}
+
+void UDataflowAttachment::PostLoad()
+{
+	Super::PostLoad();
+	Instance.PostLoad();
+}
diff --git a/Engine/Source/Runtime/Experimental/Dataflow/Engine/Private/Dataflow/DataflowBlueprintLibrary.cpp b/Engine/Source/Runtime/Experimental/Dataflow/Engine/Private/Dataflow/DataflowBlueprintLibrary.cpp
--- a/Engine/Source/Runtime/Experimental/Dataflow/Engine/Private/Dataflow/DataflowBlueprintLibrary.cpp
+++ b/Engine/Source/Runtime/Experimental/Dataflow/Engine/Private/Dataflow/DataflowBlueprintLibrary.cpp
@@ -69,7 +69,7 @@ bool UDataflowBlueprintLibrary::EvaluateDataflow(UDataflow* Dataflow, UObject* A
 
 bool UDataflowBlueprintLibrary::RegenerateAssetFromDataflow(UObject* AssetToRegenerate, bool bRegenerateDependentAssets)
 {
-	if (IDataflowInstanceInterface* DataflowInterface = Cast<IDataflowInstanceInterface>(AssetToRegenerate))
+	if (IDataflowInstanceInterface* DataflowInterface = UE::Dataflow::InstanceUtils::GetDataflowInterfaceFromObject(AssetToRegenerate))
 	{
 		FDataflowInstance& DataflowInstance = DataflowInterface->GetDataflowInstance();
 		return DataflowInstance.UpdateOwnerAsset(bRegenerateDependentAssets);
@@ -79,7 +79,7 @@ bool UDataflowBlueprintLibrary::RegenerateAssetFromDataflow(UObject* AssetToRege
 
 bool UDataflowBlueprintLibrary::OverrideDataflowVariableBool(UObject* Asset, FName VariableName, bool VariableValue)
 {
-	if (IDataflowInstanceInterface* DataflowInterface = Cast<IDataflowInstanceInterface>(Asset))
+	if (IDataflowInstanceInterface* DataflowInterface = UE::Dataflow::InstanceUtils::GetDataflowInterfaceFromObject(Asset))
 	{
 		return DataflowInterface->GetDataflowInstance().GetVariableOverrides().OverrideVariableBool(VariableName, VariableValue);
 	}
@@ -88,7 +88,7 @@ bool UDataflowBlueprintLibrary::OverrideDataflowVariableBool(UObject* Asset, FNa
 
 bool UDataflowBlueprintLibrary::OverrideDataflowVariableBoolArray(UObject* Asset, FName VariableName, const TArray<bool>& VariableArrayValue)
 {
-	if (IDataflowInstanceInterface* DataflowInterface = Cast<IDataflowInstanceInterface>(Asset))
+	if (IDataflowInstanceInterface* DataflowInterface = UE::Dataflow::InstanceUtils::GetDataflowInterfaceFromObject(Asset))
 	{
 		return DataflowInterface->GetDataflowInstance().GetVariableOverrides().OverrideVariableBoolArray(VariableName, VariableArrayValue);
 	}
@@ -97,7 +97,7 @@ bool UDataflowBlueprintLibrary::OverrideDataflowVariableBoolArray(UObject* Asset
 
 bool UDataflowBlueprintLibrary::OverrideDataflowVariableInt(UObject* Asset, FName VariableName, int64 VariableValue)
 {
-	if (IDataflowInstanceInterface* DataflowInterface = Cast<IDataflowInstanceInterface>(Asset))
+	if (IDataflowInstanceInterface* DataflowInterface = UE::Dataflow::InstanceUtils::GetDataflowInterfaceFromObject(Asset))
 	{
 		return DataflowInterface->GetDataflowInstance().GetVariableOverrides().OverrideVariableInt(VariableName, VariableValue);
 	}
@@ -106,7 +106,7 @@ bool UDataflowBlueprintLibrary::OverrideDataflowVariableInt(UObject* Asset, FNam
 
 bool UDataflowBlueprintLibrary::OverrideDataflowVariableIntArray(UObject* Asset, FName VariableName, const TArray<int32>& VariableArrayValue)
 {
-	if (IDataflowInstanceInterface* DataflowInterface = Cast<IDataflowInstanceInterface>(Asset))
+	if (IDataflowInstanceInterface* DataflowInterface = UE::Dataflow::InstanceUtils::GetDataflowInterfaceFromObject(Asset))
 	{
 		return DataflowInterface->GetDataflowInstance().GetVariableOverrides().OverrideVariableInt32Array(VariableName, VariableArrayValue);
 	}
@@ -115,7 +115,7 @@ bool UDataflowBlueprintLibrary::OverrideDataflowVariableIntArray(UObject* Asset,
 
 bool UDataflowBlueprintLibrary::OverrideDataflowVariableFloat(UObject* Asset, FName VariableName, float VariableValue)
 {
-	if (IDataflowInstanceInterface* DataflowInterface = Cast<IDataflowInstanceInterface>(Asset))
+	if (IDataflowInstanceInterface* DataflowInterface = UE::Dataflow::InstanceUtils::GetDataflowInterfaceFromObject(Asset))
 	{
 		return DataflowInterface->GetDataflowInstance().GetVariableOverrides().OverrideVariableFloat(VariableName, VariableValue);
 	}
@@ -124,7 +124,7 @@ bool UDataflowBlueprintLibrary::OverrideDataflowVariableFloat(UObject* Asset, FN
 
 bool UDataflowBlueprintLibrary::OverrideDataflowVariableFloatArray(UObject* Asset, FName VariableName, const TArray<float>& VariableArrayValue)
 {
-	if (IDataflowInstanceInterface* DataflowInterface = Cast<IDataflowInstanceInterface>(Asset))
+	if (IDataflowInstanceInterface* DataflowInterface = UE::Dataflow::InstanceUtils::GetDataflowInterfaceFromObject(Asset))
 	{
 		return DataflowInterface->GetDataflowInstance().GetVariableOverrides().OverrideVariableFloatArray(VariableName, VariableArrayValue);
 	}
@@ -133,7 +133,7 @@ bool UDataflowBlueprintLibrary::OverrideDataflowVariableFloatArray(UObject* Asse
 
 bool UDataflowBlueprintLibrary::OverrideDataflowVariableObject(UObject* Asset, FName VariableName, UObject* VariableValue)
 {
-	if (IDataflowInstanceInterface* DataflowInterface = Cast<IDataflowInstanceInterface>(Asset))
+	if (IDataflowInstanceInterface* DataflowInterface = UE::Dataflow::InstanceUtils::GetDataflowInterfaceFromObject(Asset))
 	{
 		return DataflowInterface->GetDataflowInstance().GetVariableOverrides().OverrideVariableObject(VariableName, VariableValue);
 	}
@@ -142,7 +142,7 @@ bool UDataflowBlueprintLibrary::OverrideDataflowVariableObject(UObject* Asset, F
 
 bool UDataflowBlueprintLibrary::OverrideDataflowVariableObjectArray(UObject* Asset, FName VariableName, const TArray<UObject*>& VariableArrayValue)
 {
-	if (IDataflowInstanceInterface* DataflowInterface = Cast<IDataflowInstanceInterface>(Asset))
+	if (IDataflowInstanceInterface* DataflowInterface = UE::Dataflow::InstanceUtils::GetDataflowInterfaceFromObject(Asset))
 	{
 		return DataflowInterface->GetDataflowInstance().GetVariableOverrides().OverrideVariableObjectArray(VariableName, VariableArrayValue);
 	}
@@ -154,7 +154,7 @@ bool UDataflowBlueprintLibrary::OverrideDataflowVariableObjectArray(UObject* Ass
 TArray<FDataflowVariable>  UDataflowBlueprintLibrary::GetDataflowVariableList(UObject* Asset)
 {
 	TArray<FDataflowVariable> PropArray;
-	if (IDataflowInstanceInterface* DataflowInterface = Cast<IDataflowInstanceInterface>(Asset))
+	if (IDataflowInstanceInterface* DataflowInterface = UE::Dataflow::InstanceUtils::GetDataflowInterfaceFromObject(Asset))
 	{
 		FInstancedPropertyBag& AssetBag = DataflowInterface->GetDataflowInstance().GetVariableOverrides().GetVariables();
 		TConstArrayView<FPropertyBagPropertyDesc> Descs = AssetBag.GetPropertyBagStruct()->GetPropertyDescs();
diff --git a/Engine/Source/Runtime/Experimental/Dataflow/Engine/Private/Dataflow/DataflowInstance.cpp b/Engine/Source/Runtime/Experimental/Dataflow/Engine/Private/Dataflow/DataflowInstance.cpp
--- a/Engine/Source/Runtime/Experimental/Dataflow/Engine/Private/Dataflow/DataflowInstance.cpp
+++ b/Engine/Source/Runtime/Experimental/Dataflow/Engine/Private/Dataflow/DataflowInstance.cpp
@@ -1,9 +1,12 @@
 // Copyright Epic Games, Inc. All Rights Reserved.
 
 #include "Dataflow/DataflowInstance.h"
+
+#include "Dataflow/DataflowAttachment.h"
 #include "Dataflow/DataflowObject.h"
 #include "Dataflow/DataflowObjectInterface.h"
 #include "UObject/AssetRegistryTagsContext.h"
+#include "Interfaces/Interface_AssetUserData.h"
 
 #include UE_INLINE_GENERATED_CPP_BY_NAME(DataflowInstance)
 
@@ -12,18 +15,62 @@ namespace UE::Dataflow::InstanceUtils
 	static const FName DataflowAssetPropertyName("DataflowAsset");
 	static const FName DataflowTerminalPropertyName("DataflowTerminal");
 
-	bool HasValidDataflowAsset(UObject* Obj)
+	UDataflowAttachment* GetDataflowAttachment(const UObject* InAsset)
 	{
-		// first check if the object implements IDataflowGeneratedAssetInterface
-		if (const IDataflowInstanceInterface* Interface = Cast<IDataflowInstanceInterface>(Obj))
+		if (const IInterface_AssetUserData* UserDataInterface = Cast<IInterface_AssetUserData>(InAsset))
 		{
-			return true;
+			// The asset user interface does not provide a const API, we need to work around it 
+			IInterface_AssetUserData* MutableUserDataInterface = const_cast<IInterface_AssetUserData*>(UserDataInterface);
+			return MutableUserDataInterface->GetAssetUserData<UDataflowAttachment>();
+		}
+		return nullptr;
+	}
+
+	IDataflowInstanceInterface* GetDataflowInterfaceFromObject(UObject* Obj)
+	{
+		// is the owning asset itself implement the dataflow instance interface
+		if (IDataflowInstanceInterface* Interface = Cast<IDataflowInstanceInterface>(Obj))
+		{
+			return Interface;
+		}
+		// or is that using a dataflow attachement ? 
+		if (UDataflowAttachment* Attachment = GetDataflowAttachment(Obj))
+		{
+			return Attachment;
+		}
+		return nullptr;
+	}
+
+	const IDataflowInstanceInterface* GetDataflowInterfaceFromObject(const UObject* Obj)
+	{
+		// is the owning asset itself implement the dataflow instance interface
+		if (const IDataflowInstanceInterface* Interface = Cast<const IDataflowInstanceInterface>(Obj))
+		{
+			return Interface;
+		}
+		// or is that using a dataflow attachement ? 
+		if (const UDataflowAttachment* Attachment = GetDataflowAttachment(Obj))
+		{
+			return Attachment;
 		}
-		// check if the class contains a dataflow asset and terminal node name properties
-		if (UClass* Class = Obj->GetClass())
+		return nullptr;
+	}
+
+	bool HasValidDataflowAsset(UObject* Obj)
+	{
+		if (Obj)
 		{
-			return Class->FindPropertyByName(DataflowAssetPropertyName) 
-				&& Class->FindPropertyByName(DataflowTerminalPropertyName);
+			// first check if the object implements IDataflowGeneratedAssetInterface
+			if (const IDataflowInstanceInterface* Interface = GetDataflowInterfaceFromObject(Obj))
+			{
+				return true;
+			}
+			// Legacy - check if the class contains a dataflow asset and terminal node name properties
+			if (UClass* Class = Obj->GetClass())
+			{
+				return Class->FindPropertyByName(DataflowAssetPropertyName)
+					&& Class->FindPropertyByName(DataflowTerminalPropertyName);
+			}
 		}
 		return false;
 
@@ -38,12 +85,13 @@ namespace UE::Dataflow::InstanceUtils
 		if (!DataflowObject)
 		{
 			// check if the object implements IDataflowGeneratedAssetInterface
-			if (const IDataflowInstanceInterface* Interface = Cast<IDataflowInstanceInterface>(Obj))
+			if (const IDataflowInstanceInterface* Interface = GetDataflowInterfaceFromObject(Obj))
 			{
 				DataflowObject = Interface->GetDataflowInstance().GetDataflowAsset();
 			}
 		}
 
+		// legacy checks for properties
 		if (!DataflowObject && Obj)
 		{
 			// last search if the object has a property named DataflowAsset
@@ -69,12 +117,23 @@ namespace UE::Dataflow::InstanceUtils
 		return GetDataflowAssetFromObjectTemplate<const UDataflow, const UObject>(Obj);
 	}
 
+	bool SetDataflowAssetOnObject(UDataflow* DataflowAsset, UObject* Obj)
+	{
+		// check if the object implements IDataflowGeneratedAssetInterface
+		if (IDataflowInstanceInterface* Interface = GetDataflowInterfaceFromObject(Obj))
+		{
+			Interface->GetDataflowInstance().SetDataflowAsset(DataflowAsset);
+			return true;
+		}
+		return false;
+	}
+
 	FName GetTerminalNodeNameFromObject(UObject* Obj)
 	{
 		if (Obj)
 		{
 			// check if the object implements IDataflowGeneratedAssetInterface
-			if (const IDataflowInstanceInterface* Interface = Cast<IDataflowInstanceInterface>(Obj))
+			if (IDataflowInstanceInterface* Interface = GetDataflowInterfaceFromObject(Obj))
 			{
 				return Interface->GetDataflowInstance().GetDataflowTerminal();
 			}
diff --git a/Engine/Source/Runtime/Experimental/Dataflow/Engine/Private/Dataflow/DataflowVariableNodes.cpp b/Engine/Source/Runtime/Experimental/Dataflow/Engine/Private/Dataflow/DataflowVariableNodes.cpp
--- a/Engine/Source/Runtime/Experimental/Dataflow/Engine/Private/Dataflow/DataflowVariableNodes.cpp
+++ b/Engine/Source/Runtime/Experimental/Dataflow/Engine/Private/Dataflow/DataflowVariableNodes.cpp
@@ -316,7 +316,7 @@ void FGetDataflowVariableNode::Evaluate(UE::Dataflow::FContext& Context, const F
 		{
 			if (const UE::Dataflow::FEngineContext* EngineContext = Context.AsType<UE::Dataflow::FEngineContext>())
 			{
-				if (IDataflowInstanceInterface* Interface = Cast<IDataflowInstanceInterface>(EngineContext->Owner))
+				if (IDataflowInstanceInterface* Interface = UE::Dataflow::InstanceUtils::GetDataflowInterfaceFromObject(EngineContext->Owner))
 				{
 					const FDataflowInstance& DataflowInstance = Interface->GetDataflowInstance();
 					if (DataflowInstance.GetVariableOverrides().IsVariableOverridden(VariableName))
diff --git a/Engine/Source/Runtime/Experimental/Dataflow/Engine/Public/Dataflow/DataflowAttachment.h b/Engine/Source/Runtime/Experimental/Dataflow/Engine/Public/Dataflow/DataflowAttachment.h
new file mode 100644
--- /dev/null
+++ b/Engine/Source/Runtime/Experimental/Dataflow/Engine/Public/Dataflow/DataflowAttachment.h
@@ -0,0 +1,54 @@
+// Copyright Epic Games, Inc. All Rights Reserved.
+
+#pragma once
+
+#include "Engine/AssetUserData.h"
+#include "Dataflow/DataflowContent.h"
+#include "Dataflow/DataflowInstance.h"
+#include "UObject/ObjectPtr.h"
+
+#include "DataflowAttachment.generated.h"
+
+#define UE_API DATAFLOWENGINE_API
+
+/**
+ * Generic dataflow instance attachment implemented as a user asset data object so it can be embedded inside
+ * any asset that is a user-data container. This allows assets in modules that do not, or cannot know about
+ * dataflow (e.g. Engine module assets) to contain a dataflow attachment.
+ * Using a modular feature implementation this data can be grabbed as an asset opens and a plugin specific
+ * implementation can be used for the asset editor (see FDataflowPhysicsAssetEditorOverride)
+ */
+UCLASS(MinimalApi, Abstract)
+class UDataflowAttachment : public UAssetUserData, public IDataflowContentOwner, public IDataflowInstanceInterface
+{
+	GENERATED_BODY()
+
+public:
+
+	// IDataflowContentOwner
+	void WriteDataflowContent(const TObjectPtr<UDataflowBaseContent>& DataflowContent) const override PURE_VIRTUAL( UDataflowAttachment::WriteDataflowContent, );
+	void ReadDataflowContent(const TObjectPtr<UDataflowBaseContent>& DataflowContent) override PURE_VIRTUAL( UDataflowAttachment::ReadDataflowContent, );
+	////////////////////////////////////////////////////////////////////////////////////////////////////////////////
+	
+	// IDataflowInstanceInterface
+	UE_API const FDataflowInstance& GetDataflowInstance() const override;
+	UE_API FDataflowInstance& GetDataflowInstance() override;
+	////////////////////////////////////////////////////////////////////////////////////////////////////////////////
+
+	// UObject
+	UE_API virtual void PostLoad() override;
+	////////////////////////////////////////////////////////////////////////////////////////////////////////////////
+private:
+
+	// Dataflow instance for the attachment when editing.
+	UPROPERTY()
+	FDataflowInstance Instance;
+
+protected:
+	
+	// IDataflowContentOwner
+	TObjectPtr<UDataflowBaseContent> CreateDataflowContent() override PURE_VIRTUAL( UDataflowAttachment::CreateDataflowContent, return nullptr; );
+	////////////////////////////////////////////////////////////////////////////////////////////////////////////////
+};
+
+#undef UE_API
diff --git a/Engine/Source/Runtime/Experimental/Dataflow/Engine/Public/Dataflow/DataflowInstance.h b/Engine/Source/Runtime/Experimental/Dataflow/Engine/Public/Dataflow/DataflowInstance.h
--- a/Engine/Source/Runtime/Experimental/Dataflow/Engine/Public/Dataflow/DataflowInstance.h
+++ b/Engine/Source/Runtime/Experimental/Dataflow/Engine/Public/Dataflow/DataflowInstance.h
@@ -331,15 +331,28 @@ public:
 
 namespace UE::Dataflow::InstanceUtils
 {
-	/** Weather or not a dataflow asset with a non empty terminal node name is assigned to an UObject */
+	/** Weather or not a dataflow asset is assigned to an UObject */
 	DATAFLOWENGINE_API bool HasValidDataflowAsset(UObject* Obj);
 
-	/** Get the dataflow asset from a UObject if available */
+	/** Get the Dataflow asset from a UObject if available */
 	DATAFLOWENGINE_API UDataflow* GetDataflowAssetFromObject(UObject* Obj);
 
-	/** Get the dataflow asset from a UObject if available (const version) */
+	/** Get the Dataflow asset from a UObject if available (const version) */
 	DATAFLOWENGINE_API const UDataflow* GetDataflowAssetFromObject(const UObject* Obj);
 
+	/** Get the Dataflow instance interface from a UObject */
+	DATAFLOWENGINE_API IDataflowInstanceInterface* GetDataflowInterfaceFromObject(UObject* Obj);
+
+	/** Get the Dataflow instance interface from a const UObject */
+	DATAFLOWENGINE_API const IDataflowInstanceInterface* GetDataflowInterfaceFromObject(const UObject* Obj);
+
+	/** 
+	* Set a Dataflow asset/object on a specific object
+	* The object type must be compatible with dataflow
+	* Returns true if the dataflow object was set successfuly
+	*/
+	DATAFLOWENGINE_API bool SetDataflowAssetOnObject(UDataflow* DataflowAsset, UObject* Obj);
+
 	/** Get the terminal node name from a UObject if available */
 	DATAFLOWENGINE_API FName GetTerminalNodeNameFromObject(UObject* Obj);
 
