FOnDemandInstallCacheUsage now contains usage for all CAS.

#Synced-CL 52313638

diff --git a/Engine/Source/Runtime/Experimental/IoStore/OnDemand/Private/OnDemandInstallCache.cpp b/Engine/Source/Runtime/Experimental/IoStore/OnDemand/Private/OnDemandInstallCache.cpp
--- a/Engine/Source/Runtime/Experimental/IoStore/OnDemand/Private/OnDemandInstallCache.cpp
+++ b/Engine/Source/Runtime/Experimental/IoStore/OnDemand/Private/OnDemandInstallCache.cpp
@@ -2583,7 +2583,7 @@ public:
 	virtual void								UpdateLastAccess(
 													TConstArrayView<FIoHash> ChunkHashes, 
 													TStaticArray<bool, EOnDemandInstallCasType::Count>& bInOutLastAccessDirty) override;
-	virtual FOnDemandInstallCacheUsage			GetCacheUsage(EOnDemandInstallCasType CasType) override;
+	virtual FOnDemandInstallCacheUsage			GetCacheUsage() override;
 	virtual bool								IsCasEnabled(EOnDemandInstallCasType CasType) const override;
 
 private:
@@ -4712,56 +4712,63 @@ void FOnDemandInstallCache::UpdateLastAccess(TConstArrayView<FIoHash> ChunkHashe
 	}
 }
 
-FOnDemandInstallCacheUsage FOnDemandInstallCache::GetCacheUsage(EOnDemandInstallCasType CasType)
+FOnDemandInstallCacheUsage FOnDemandInstallCache::GetCacheUsage()
 {
 	// If this is called from a thread other than the OnDemandIoStore tick thread
 	// then its possible the block info and containers may not be in sync with each other
 	// or the current state of the tick thread.
 	// This should only be used for debugging and telemetry purposes.
 
-	FCacheState* CacheState = GetCacheState(CasType);
-	if (CacheState == nullptr)
+	FOnDemandInstallCacheUsage CacheUsage;
+
+	for (EOnDemandInstallCasType CasType : TEnumRange<EOnDemandInstallCasType>())
 	{
-		return FOnDemandInstallCacheUsage();
-	}
+		FOnDemandInstallCacheUsage::FCasUsage& CasUsage = CacheUsage.CasUsage[CasType];
+		CasUsage.CasType = CasType;
 
-	FCasBlockInfoMap	BlockInfo;
-	const uint64		TotalCachedBytes = CacheState->Cas.GetBlockInfo(BlockInfo);
+		FCacheState* CacheState = GetCacheState(CasType);
+		if (CacheState == nullptr)
+		{
+			continue;
+		}
 
-	TArray<FSharedOnDemandContainer>	Containers;
-	TArray<TBitArray<>>					ChunkEntryIndices;
-	IoStore.GetReferencedContent(Containers, ChunkEntryIndices);
-	check(Containers.Num() == ChunkEntryIndices.Num());
+		FCasBlockInfoMap	BlockInfo;
+		const uint64		TotalCachedBytes = CacheState->Cas.GetBlockInfo(BlockInfo);
 
-	uint64 ReferencedBytes = 0;
-	AddReferencesToBlocks(*CacheState, Containers, ChunkEntryIndices, BlockInfo, ReferencedBytes);
+		TArray<FSharedOnDemandContainer>	Containers;
+		TArray<TBitArray<>>					ChunkEntryIndices;
+		IoStore.GetReferencedContent(Containers, ChunkEntryIndices);
+		check(Containers.Num() == ChunkEntryIndices.Num());
 
-	uint64 FragmentedBytes = 0;
-	uint64 ReferencedBlockBytes = 0;
-	for (const TPair<FCasBlockId, FCasBlockInfo>& Kv : BlockInfo)
-	{
-		const FCasBlockId BlockId = Kv.Key;
-		const FCasBlockInfo& Info = Kv.Value;
+		uint64 ReferencedBytes = 0;
+		AddReferencesToBlocks(*CacheState, Containers, ChunkEntryIndices, BlockInfo, ReferencedBytes);
 
-		if (Info.RefSize < Info.FileSize)
+		uint64 FragmentedBytes = 0;
+		uint64 ReferencedBlockBytes = 0;
+		for (const TPair<FCasBlockId, FCasBlockInfo>& Kv : BlockInfo)
 		{
-			FragmentedBytes += (Info.FileSize - Info.RefSize);
-		}
+			const FCasBlockId BlockId = Kv.Key;
+			const FCasBlockInfo& Info = Kv.Value;
 
-		if (Info.RefSize > 0)
-		{
-			ReferencedBlockBytes += Info.FileSize;
+			if (Info.RefSize < Info.FileSize)
+			{
+				FragmentedBytes += (Info.FileSize - Info.RefSize);
+			}
+
+			if (Info.RefSize > 0)
+			{
+				ReferencedBlockBytes += Info.FileSize;
+			}
 		}
+
+		CasUsage.MaxSize = CacheState->MaxCacheSize;
+		CasUsage.TotalSize = TotalCachedBytes;
+		CasUsage.ReferencedBlockSize = ReferencedBlockBytes;
+		CasUsage.ReferencedSize = ReferencedBytes;
+		CasUsage.FragmentedChunksSize = FragmentedBytes;
 	}
 
-	return FOnDemandInstallCacheUsage
-	{
-		.MaxSize = CacheState->MaxCacheSize,
-		.TotalSize = TotalCachedBytes,
-		.ReferencedBlockSize = ReferencedBlockBytes,
-		.ReferencedSize = ReferencedBytes,
-		.FragmentedChunksSize = FragmentedBytes,
-	};
+	return CacheUsage;
 }
 
 bool FOnDemandInstallCache::IsCasEnabled(EOnDemandInstallCasType CasType) const
diff --git a/Engine/Source/Runtime/Experimental/IoStore/OnDemand/Private/OnDemandInstallCache.h b/Engine/Source/Runtime/Experimental/IoStore/OnDemand/Private/OnDemandInstallCache.h
--- a/Engine/Source/Runtime/Experimental/IoStore/OnDemand/Private/OnDemandInstallCache.h
+++ b/Engine/Source/Runtime/Experimental/IoStore/OnDemand/Private/OnDemandInstallCache.h
@@ -38,7 +38,7 @@ public:
 	virtual FResult								Flush(EOnDemandInstallCasType CasType) = 0;
 	virtual FResult								FlushLastAccess(const TStaticArray<bool, EOnDemandInstallCasType::Count>& bFlushCasLastAccess) = 0;
 	virtual void								UpdateLastAccess(TConstArrayView<FIoHash> ChunkHashes, TStaticArray<bool, EOnDemandInstallCasType::Count>& bInOutLastAccessDirty) = 0;
-	virtual FOnDemandInstallCacheUsage			GetCacheUsage(EOnDemandInstallCasType CasType) = 0;
+	virtual FOnDemandInstallCacheUsage			GetCacheUsage() = 0;
 	virtual bool								IsCasEnabled(EOnDemandInstallCasType CasType) const = 0;
 };
 
diff --git a/Engine/Source/Runtime/Experimental/IoStore/OnDemand/Private/OnDemandIoStore.cpp b/Engine/Source/Runtime/Experimental/IoStore/OnDemand/Private/OnDemandIoStore.cpp
--- a/Engine/Source/Runtime/Experimental/IoStore/OnDemand/Private/OnDemandIoStore.cpp
+++ b/Engine/Source/Runtime/Experimental/IoStore/OnDemand/Private/OnDemandIoStore.cpp
@@ -1863,7 +1863,7 @@ FOnDemandCacheUsage FOnDemandIoStore::GetCacheUsage(const FOnDemandGetCacheUsage
 
 	if (InstallCache.IsValid())
 	{
-		CacheUsage.InstallCache = InstallCache->GetCacheUsage(Args.CasType);
+		CacheUsage.InstallCache = InstallCache->GetCacheUsage();
 
 		if (EnumHasAnyFlags(Args.Options, EOnDemandGetCacheUsageOptions::DumpHandlesToLog | EOnDemandGetCacheUsageOptions::DumpHandlesToResults))
 		{
diff --git a/Engine/Source/Runtime/Experimental/IoStore/OnDemandCore/Private/IoStoreOnDemand.cpp b/Engine/Source/Runtime/Experimental/IoStore/OnDemandCore/Private/IoStoreOnDemand.cpp
--- a/Engine/Source/Runtime/Experimental/IoStore/OnDemandCore/Private/IoStoreOnDemand.cpp
+++ b/Engine/Source/Runtime/Experimental/IoStore/OnDemandCore/Private/IoStoreOnDemand.cpp
@@ -185,14 +185,6 @@ static void PrintCacheUsage(const TArray<FString>& InArgs)
 	{
 		FOnDemandGetCacheUsageArgs UsageArgs;
 		UsageArgs.Options = EOnDemandGetCacheUsageOptions::DumpHandlesToLog;
-		for (const FString& Arg : InArgs)
-		{
-			EOnDemandInstallCasType CasType = EOnDemandInstallCasType::None;
-			if (LexFromString(CasType, Arg); CasType != EOnDemandInstallCasType::None)
-			{
-				UsageArgs.CasType = CasType;
-			}
-		}
 
 		TIoStatusOr<FOnDemandCacheUsage> MaybeUsage = IoStore->GetCacheUsage(UsageArgs);
 		if (!MaybeUsage.IsOk())
@@ -202,11 +194,16 @@ static void PrintCacheUsage(const TArray<FString>& InArgs)
 		}
 
 		const FOnDemandCacheUsage CacheUsage = MaybeUsage.ConsumeValueOrDie();
-		UE_LOG(LogIoStoreOnDemand, Display, TEXT("iostore.CacheUsage"));
-		TStringBuilder<512> Sb;
-		Sb	<< TEXTVIEW("InstallCache ") << LexToString(UsageArgs.CasType) << TEXTVIEW(": ") << CacheUsage.InstallCache 
-			<< TEXTVIEW(", StreamingCache: ") << CacheUsage.StreamingCache;
-		UE_LOG(LogIoStoreOnDemand, Display, TEXT("%s"), Sb.ToString());
+
+		UE_LOGF(LogIoStoreOnDemand, Display, "iostore.CacheUsage");
+
+		TUtf8StringBuilder<256> Sb;
+		Sb << UTF8TEXTVIEW("InstallCache: ") << CacheUsage.InstallCache;
+		UE_LOGF(LogIoStoreOnDemand, Display, "%s", Sb.ToString());
+
+		Sb.Reset();
+		Sb << UTF8TEXTVIEW("StreamingCache: ") << CacheUsage.StreamingCache;
+		UE_LOGF(LogIoStoreOnDemand, Display, "%s", Sb.ToString());
 
 		IoStore->DumpMountedContainersToLog();
 	}
@@ -468,22 +465,38 @@ static FAutoConsoleCommand VerifyCacheCommand(
 static const TCHAR* NotInitializedError = TEXT("I/O store on-demand not initialized");
 
 ////////////////////////////////////////////////////////////////////////////////
-FStringBuilderBase& operator<<(FStringBuilderBase& Sb, const FOnDemandInstallCacheUsage& CacheUsage)
+FUtf8StringBuilderBase& operator<<(FUtf8StringBuilderBase& Sb, const FOnDemandInstallCacheUsage::FCasUsage& CasUsage)
 {
-	Sb.Appendf(TEXT("MaxSize=%.2lf MiB, TotalSize=%.2lf MiB, ReferencedBlockSize=%.2lf MiB, ReferencedSize=%.2lf MiB, FragmentedChunksSize=%.2lf MiB"),
-		double(CacheUsage.MaxSize) / 1024.0 / 1024.0,
-		double(CacheUsage.TotalSize) / 1024.0 / 1024.0,
-		double(CacheUsage.ReferencedBlockSize) / 1024.0 / 1024.0,
-		double(CacheUsage.ReferencedSize) / 1024.0 / 1024.0,
-		double(CacheUsage.FragmentedChunksSize) / 1024.0 / 1024.0);
+	if (!CasUsage.IsValid())
+	{
+		return Sb;
+	}
+
+	Sb.Appendf("CAS=%ls, MaxSize=%.2lf MiB, TotalSize=%.2lf MiB, ReferencedBlockSize=%.2lf MiB, ReferencedSize=%.2lf MiB, FragmentedChunksSize=%.2lf MiB",
+		LexToString(CasUsage.CasType),
+		double(CasUsage.MaxSize) / 1024.0 / 1024.0,
+		double(CasUsage.TotalSize) / 1024.0 / 1024.0,
+		double(CasUsage.ReferencedBlockSize) / 1024.0 / 1024.0,
+		double(CasUsage.ReferencedSize) / 1024.0 / 1024.0,
+		double(CasUsage.FragmentedChunksSize) / 1024.0 / 1024.0);
+
+	return Sb;
+}
+
+FUtf8StringBuilderBase& operator<<(FUtf8StringBuilderBase& Sb, const FOnDemandInstallCacheUsage& CacheUsage)
+{
+	for (const FOnDemandInstallCacheUsage::FCasUsage& CasUsage : CacheUsage.CasUsage)
+	{
+		Sb << CasUsage << TEXTVIEW("\n");
+	}
 
 	return Sb;
 }
 
 ////////////////////////////////////////////////////////////////////////////////
-FStringBuilderBase& operator<<(FStringBuilderBase& Sb, const FOnDemandStreamingCacheUsage& CacheUsage)
+FUtf8StringBuilderBase& operator<<(FUtf8StringBuilderBase& Sb, const FOnDemandStreamingCacheUsage& CacheUsage)
 {
-	Sb.Appendf(TEXT("MaxSize=%.2lf MiB, TotalSize=%.2lf MiB"),
+	Sb.Appendf("MaxSize=%.2lf MiB, TotalSize=%.2lf MiB",
 		double(CacheUsage.MaxSize) / 1024.0 / 1024.0, double(CacheUsage.TotalSize) / 1024.0 / 1024.0);
 	return Sb;
 }
diff --git a/Engine/Source/Runtime/Experimental/IoStore/OnDemandCore/Public/IO/IoStoreOnDemand.h b/Engine/Source/Runtime/Experimental/IoStore/OnDemandCore/Public/IO/IoStoreOnDemand.h
--- a/Engine/Source/Runtime/Experimental/IoStore/OnDemandCore/Public/IO/IoStoreOnDemand.h
+++ b/Engine/Source/Runtime/Experimental/IoStore/OnDemandCore/Public/IO/IoStoreOnDemand.h
@@ -6,6 +6,7 @@
 #include "Containers/UnrealString.h"
 #include "Containers/StringFwd.h"
 #include "Containers/SharedString.h"
+#include "Containers/StaticArray.h"
 #include "Experimental/UnifiedError/UnifiedError.h"
 #include "IO/OnDemandHostGroup.h"
 #include "IO/OnDemandToc.h"
@@ -525,7 +526,6 @@ ENUM_CLASS_FLAGS(EOnDemandGetCacheUsageOptions);
 /** Arguments for getting the install cache usage. */
 struct FOnDemandGetCacheUsageArgs
 {
-	EOnDemandInstallCasType CasType = EOnDemandInstallCasType::General;
 	EOnDemandGetCacheUsageOptions Options = EOnDemandGetCacheUsageOptions::None;
 };
 
@@ -539,16 +539,26 @@ struct FOnDemandInstallHandleCacheUsage
 /** Holds information about the install cache usage */
 struct FOnDemandInstallCacheUsage
 {
-	uint64 MaxSize = 0;
-	uint64 TotalSize = 0;
-	uint64 ReferencedBlockSize = 0;
-	uint64 ReferencedSize = 0;
-	uint64 FragmentedChunksSize = 0;
+	struct FCasUsage
+	{
+		uint64 MaxSize = 0;
+		uint64 TotalSize = 0;
+		uint64 ReferencedBlockSize = 0;
+		uint64 ReferencedSize = 0;
+		uint64 FragmentedChunksSize = 0;
+		EOnDemandInstallCasType CasType = EOnDemandInstallCasType::None;
+
+		bool IsValid() const { return CasType != EOnDemandInstallCasType::None && MaxSize != 0; }
+
+		UE_API friend FUtf8StringBuilderBase& operator<<(FUtf8StringBuilderBase& Sb, const FCasUsage& CasUsage);
+	};
+
+	TStaticArray<FCasUsage, EOnDemandInstallCasType::Count> CasUsage;
 
 	// Only populated if DumpHandlesToResults is specified
 	TArray<FOnDemandInstallHandleCacheUsage> ReferencedBytesByHandle;
 
-	UE_API friend FStringBuilderBase& operator<<(FStringBuilderBase& Sb, const FOnDemandInstallCacheUsage& CacheUsage);
+	UE_API friend FUtf8StringBuilderBase& operator<<(FUtf8StringBuilderBase& Sb, const FOnDemandInstallCacheUsage& CacheUsage);
 };
 
 /** Holds information about the streaming cache usage */
@@ -557,7 +567,7 @@ struct FOnDemandStreamingCacheUsage
 	uint64 MaxSize = 0;
 	uint64 TotalSize = 0;
 
-	UE_API friend FStringBuilderBase& operator<<(FStringBuilderBase& Sb, const FOnDemandStreamingCacheUsage& CacheUsage);
+	UE_API friend FUtf8StringBuilderBase& operator<<(FUtf8StringBuilderBase& Sb, const FOnDemandStreamingCacheUsage& CacheUsage);
 };
 
 /** Holds information about install and streaming cache usage */
