[Chaos] Add one dimensional friction along the velocity
We can enable it with a cVar

#Synced-CL 52412319

diff --git a/Engine/Source/Runtime/Experimental/Chaos/Private/Chaos/Collision/PBDCollisionContainerSolver.cpp b/Engine/Source/Runtime/Experimental/Chaos/Private/Chaos/Collision/PBDCollisionContainerSolver.cpp
--- a/Engine/Source/Runtime/Experimental/Chaos/Private/Chaos/Collision/PBDCollisionContainerSolver.cpp
+++ b/Engine/Source/Runtime/Experimental/Chaos/Private/Chaos/Collision/PBDCollisionContainerSolver.cpp
@@ -26,6 +26,7 @@ namespace Chaos
 		extern bool bChaos_PBDCollisionSolver_Velocity_SolveEnabled;
 		extern float Chaos_PBDCollisionSolver_Position_MinInvMassScale;
 		extern float Chaos_PBDCollisionSolver_Velocity_MinInvMassScale;
+		extern bool bChaos_PBDCollisionSolver_Use1DFriction;
 
 		// If one body is more than MassRatio1 times the mass of the other, adjust the solver stiffness when the lighter body is underneath.
 		// Solver stiffness will be equal to 1 when the mass ratio is MassRatio1.
@@ -103,18 +104,16 @@ namespace Chaos
 		// World-space normal
 		const FVec3f WorldContactNormal = FVec3f(ShapeWorldTransform1.TransformVectorNoScale(FVec3(ManifoldPoint.ContactPoint.ShapeContactNormal)));
 
-		// World-space tangents
+		// World-space tangent U
 		FVec3f WorldContactTangentU = FVec3f::CrossProduct(FVec3f(0, 1, 0), WorldContactNormal);
 		if (!WorldContactTangentU.Normalize(FRealSingle(UE_KINDA_SMALL_NUMBER)))
 		{
 			WorldContactTangentU = FVec3f::CrossProduct(FVec3f(1, 0, 0), WorldContactNormal);
 			WorldContactTangentU = WorldContactTangentU.GetUnsafeNormal();
 		}
-		const FVec3f WorldContactTangentV = FVec3f::CrossProduct(WorldContactNormal, WorldContactTangentU);
-
 
 		// Calculate contact velocity if we will need it below (restitution and/or first-contact for friction or initial depenetration)
-		const bool bNeedContactVelocity = (!ManifoldPoint.Flags.bHasStaticFrictionAnchor) || (Restitution > FRealSingle(0)) || ManifoldPoint.Flags.bInitialContact;
+		const bool bNeedContactVelocity = (!ManifoldPoint.Flags.bHasStaticFrictionAnchor) || (Restitution > FRealSingle(0)) || ManifoldPoint.Flags.bInitialContact || CVars::bChaos_PBDCollisionSolver_Use1DFriction;
 		FVec3f ContactVel = FVec3(0);
 		FRealSingle ContactVelocityNormal = FRealSingle(0);
 		if (bNeedContactVelocity)
@@ -123,8 +122,18 @@ namespace Chaos
 			const FVec3f ContactVel1 = Body1.V() + FVec3f::CrossProduct(Body1.W(), WorldRelativeContact1);
 			ContactVel = ContactVel0 - ContactVel1;
 			ContactVelocityNormal = FVec3f::DotProduct(ContactVel, WorldContactNormal);
+			// Align TangentU with the relative sliding velocity direction,
+			// Falls back to the geometric axis when sliding velocity is near zero
+			FVec3f ContactVelTangent = ContactVel - ContactVelocityNormal * WorldContactNormal;
+			if (ContactVelTangent.Normalize(FRealSingle(UE_KINDA_SMALL_NUMBER)))
+			{
+				WorldContactTangentU = ContactVelTangent;
+			}
 		}
 
+		// TangentV is always perpendicular to both normal and TangentU
+		const FVec3f WorldContactTangentV = FVec3f::CrossProduct(WorldContactNormal, WorldContactTangentU);
+
 		// If we have contact data from a previous tick, use it to calculate the lateral position delta we need
 		// to apply to move the contacts back to their original relative locations (i.e., to enforce static friction).
 		// Otherwise, estimate the friction correction from the contact velocity.
diff --git a/Engine/Source/Runtime/Experimental/Chaos/Private/Chaos/Collision/PBDCollisionSolver.cpp b/Engine/Source/Runtime/Experimental/Chaos/Private/Chaos/Collision/PBDCollisionSolver.cpp
--- a/Engine/Source/Runtime/Experimental/Chaos/Private/Chaos/Collision/PBDCollisionSolver.cpp
+++ b/Engine/Source/Runtime/Experimental/Chaos/Private/Chaos/Collision/PBDCollisionSolver.cpp
@@ -53,6 +53,9 @@ namespace Chaos
 		FAutoConsoleVariableRef CVarChaos_PBDCollisionSolver_Velocity_FrictionEnabled(TEXT("p.Chaos.PBDCollisionSolver.Velocity.FrictionEnabled"), bChaos_PBDCollisionSolver_Velocity_FrictionEnabled, TEXT(""));
 		FAutoConsoleVariableRef CVarChaos_PBDCollisionSolver_Velocity_StaticFrictionStiffness(TEXT("p.Chaos.PBDCollisionSolver.Velocity.StaticFriction.Stiffness"), Chaos_PBDCollisionSolver_Velocity_StaticFrictionStiffness, TEXT(""));
 		FAutoConsoleVariableRef CVarChaos_PBDCollisionSolver_Velocity_AveragePointEnabled(TEXT("p.Chaos.PBDCollisionSolver.Velocity.AveragePointEnabled"), bChaos_PBDCollisionSolver_Velocity_AveragePointEnabled, TEXT(""));
+
+		bool bChaos_PBDCollisionSolver_Use1DFriction = false;
+		FAutoConsoleVariableRef CVarChaos_PBDCollisionSolver_CVarUse1DFriction(TEXT("p.Chaos.Solver.Use1DFriction"), bChaos_PBDCollisionSolver_Use1DFriction, TEXT("Use 1 dimensional friction along the velocity axis."));
 	}
 	using namespace CVars;
 
diff --git a/Engine/Source/Runtime/Experimental/Chaos/Public/Chaos/Collision/PBDCollisionSolver.h b/Engine/Source/Runtime/Experimental/Chaos/Public/Chaos/Collision/PBDCollisionSolver.h
--- a/Engine/Source/Runtime/Experimental/Chaos/Public/Chaos/Collision/PBDCollisionSolver.h
+++ b/Engine/Source/Runtime/Experimental/Chaos/Public/Chaos/Collision/PBDCollisionSolver.h
@@ -22,6 +22,7 @@ namespace Chaos
 		extern bool bChaos_PBDCollisionSolver_Velocity_FrictionEnabled;
 		extern float Chaos_PBDCollisionSolver_Position_StaticFrictionStiffness;
 		extern float Chaos_PBDCollisionSolver_Velocity_StaticFrictionStiffness;
+		extern bool bChaos_PBDCollisionSolver_Use1DFriction;
 	}
 
 	namespace Private
@@ -1200,6 +1201,12 @@ namespace Chaos
 						FSolverReal ContactVelocityDeltaNormal, ContactVelocityDeltaTangentU, ContactVelocityDeltaTangentV;
 						CalculateContactVelocityError(SolverManifoldPoint, DynamicFriction, Dt, ContactVelocityDeltaNormal, ContactVelocityDeltaTangentU, ContactVelocityDeltaTangentV);
 
+						// 1D friction: zero out the V component. TangentU is already velocity-aligned
+						if (CVars::bChaos_PBDCollisionSolver_Use1DFriction)
+						{
+							ContactVelocityDeltaTangentV = FSolverReal(0);
+						}
+
 						if (!bShouldSolveNormalVelocity)
 						{
 							ContactVelocityDeltaNormal = 0;
