ErrorOnRequest Cheat working case-insensitive

#rnx

#Synced-CL 52419526

diff --git a/Engine/Source/Runtime/Experimental/IoStore/OnDemand/Private/OnDemandContentInstaller.cpp b/Engine/Source/Runtime/Experimental/IoStore/OnDemand/Private/OnDemandContentInstaller.cpp
--- a/Engine/Source/Runtime/Experimental/IoStore/OnDemand/Private/OnDemandContentInstaller.cpp
+++ b/Engine/Source/Runtime/Experimental/IoStore/OnDemand/Private/OnDemandContentInstaller.cpp
@@ -907,8 +907,8 @@ void FOnDemandContentInstaller::ProcessInstallRequest(FRequest& Request)
 #if !UE_BUILD_SHIPPING
 	if (!UE::IoStore::CVars::IoStoreErrorOnRequest.IsEmpty())
 	{
-		if (FCString::Strstr(*ContentHandle->DebugName, *UE::IoStore::CVars::IoStoreErrorOnRequest) != nullptr || 
-			FCString::Strstr(*InstallRequest.Args.DebugName, *UE::IoStore::CVars::IoStoreErrorOnRequest) != nullptr)
+		if (FCString::Stristr(*ContentHandle->DebugName, *UE::IoStore::CVars::IoStoreErrorOnRequest) != nullptr || 
+			FCString::Stristr(*InstallRequest.Args.DebugName, *UE::IoStore::CVars::IoStoreErrorOnRequest) != nullptr)
 		{
 			Request.Result = MakeError(
 				UE::UnifiedError::Core::ArgumentError::MakeError(
