[Parallel constraint solver] Hide SolverPartitionManager constructor from non-friend classes by using the passkey idiom.

#Synced-CL 52354501

diff --git a/Engine/Source/Runtime/Experimental/Chaos/Private/Chaos/Evolution/SolverPartitionManager.cpp b/Engine/Source/Runtime/Experimental/Chaos/Private/Chaos/Evolution/SolverPartitionManager.cpp
--- a/Engine/Source/Runtime/Experimental/Chaos/Private/Chaos/Evolution/SolverPartitionManager.cpp
+++ b/Engine/Source/Runtime/Experimental/Chaos/Private/Chaos/Evolution/SolverPartitionManager.cpp
@@ -25,7 +25,7 @@ namespace Chaos
 
 	namespace Private
 	{
-		FSolverPartitionManager::FSolverPartitionManager(Chaos::FConstraintContainerSolver& ContainerIn)
+		FSolverPartitionManager::FSolverPartitionManager(Chaos::FConstraintContainerSolver& ContainerIn, FPasskey Passkey)
 			: Container(ContainerIn)
 		{
 		}
@@ -46,8 +46,8 @@ namespace Chaos
 			return ColorStartIndices.Num();
 		}
 
-		FGreedyConstraintColoring::FGreedyConstraintColoring(Chaos::FConstraintContainerSolver& ContainerIn)
-			: FSolverPartitionManager(ContainerIn)
+		FGreedyConstraintColoring::FGreedyConstraintColoring(Chaos::FConstraintContainerSolver& ContainerIn, FPasskey Passkey)
+			: FSolverPartitionManager(ContainerIn, Passkey)
 		{
 		}
 
diff --git a/Engine/Source/Runtime/Experimental/Chaos/Public/Chaos/Evolution/SolverConstraintContainer.h b/Engine/Source/Runtime/Experimental/Chaos/Public/Chaos/Evolution/SolverConstraintContainer.h
--- a/Engine/Source/Runtime/Experimental/Chaos/Public/Chaos/Evolution/SolverConstraintContainer.h
+++ b/Engine/Source/Runtime/Experimental/Chaos/Public/Chaos/Evolution/SolverConstraintContainer.h
@@ -165,9 +165,9 @@ namespace Chaos
 			if (CVars::Chaos_SolverPartitionManager_Partitioning == 1 && SupportsPartitioning())
 			{
 				// Create a solver partition manager to apply greedy graph coloring
-				if (!SolverPartitionManager.IsValid() || SolverPartitionManager->GetType() != Private::ESolverPartitionManagerType::GreedyConstraintColoring)
+				if (!SolverPartitionManager.IsValid() || SolverPartitionManager->GetType(Private::FSolverPartitionManager::FPasskey()) != Private::ESolverPartitionManagerType::GreedyConstraintColoring)
 				{
-					SolverPartitionManager = MakeUnique<Private::FGreedyConstraintColoring>(*this);
+					SolverPartitionManager = TUniquePtr<Private::FGreedyConstraintColoring>(new	Private::FGreedyConstraintColoring(*this, Private::FSolverPartitionManager::FPasskey()));
 				}
 			}
 			else
diff --git a/Engine/Source/Runtime/Experimental/Chaos/Public/Chaos/Evolution/SolverPartitionManager.h b/Engine/Source/Runtime/Experimental/Chaos/Public/Chaos/Evolution/SolverPartitionManager.h
--- a/Engine/Source/Runtime/Experimental/Chaos/Public/Chaos/Evolution/SolverPartitionManager.h
+++ b/Engine/Source/Runtime/Experimental/Chaos/Public/Chaos/Evolution/SolverPartitionManager.h
@@ -17,11 +17,24 @@ namespace Chaos
 
 		class FSolverPartitionManager
 		{
-
 		public:
-			FSolverPartitionManager(Chaos::FConstraintContainerSolver& ContrainerIn);
+			// Passkey idiom: public class with private constructor.
+			// Anyone can name this type, but only friend classes can construct it.
+			class FPasskey {
+				friend class Chaos::FConstraintContainerSolver;
+				explicit FPasskey() {}
+			};
+
+			/**
+			 * Not publicly callable because of passkey with private constructor
+			*/
+			FSolverPartitionManager(Chaos::FConstraintContainerSolver& ContrainerIn, FPasskey Passkey);
 			virtual ~FSolverPartitionManager();
-			virtual ESolverPartitionManagerType GetType() const = 0;
+
+			/**
+			 * Not publicly callable because of passkey with private constructor
+			*/
+			virtual ESolverPartitionManagerType GetType(FPasskey Passkey) const = 0;
 
 			/**
 			 * Global start indices into solver's constraint container.
@@ -40,16 +53,17 @@ namespace Chaos
 
 		protected:
 			TArray<TArray<int32>> ConstraintsBatchIndices;
-
 			FConstraintContainerSolver& Container;
 		};
 
 		class FGreedyConstraintColoring : public FSolverPartitionManager
 		{
 		public:
-			FGreedyConstraintColoring(Chaos::FConstraintContainerSolver& ContrainerIn);
+			FGreedyConstraintColoring(Chaos::FConstraintContainerSolver& ContrainerIn, FPasskey Passkey);
+			virtual ESolverPartitionManagerType GetType(FPasskey Passkey) const override final { return ESolverPartitionManagerType::GreedyConstraintColoring; }
+
+		public:
 			UE_INTERNAL virtual int32 PreparePartitions() override final;
-			UE_INTERNAL virtual ESolverPartitionManagerType GetType() const override final { return ESolverPartitionManagerType::GreedyConstraintColoring; }
 
 		private:
 			/**
