ChaosVD / Chaos Runtime: particle extra data trace API

Foundational data types, serialization helpers, and trace function
for the Chaos runtime module (no CVD plugin dependency):

- ChaosVDParticleExtraData.h/.cpp: FChaosVDExtraDataStructEntry,
FChaosVDExtraDataCategory (with SourceChannelId for optional channel tagging),
FChaosVDParticleExtraData; EChaosVDExtraDataSerializationMode enum
(SerializeBin default, NativeSerialization opt-in);
TraceChaosVDParticleExtraData() entry point;
FChaosVDExtraDataStructEntry::SerializeFrom() -- records the chosen mode
so the reader can apply symmetric deserialization
- ChaosVDStructCollectionMemWriterReader.h: FChaosVDStructCollectionMemoryWriter/Reader
-- serializes UScriptStruct refs as path strings, drops other UObject refs
- ChaosVDTraceMacros.h: CVD_TRACE_PARTICLE_EXTRA_DATA and
CVD_TRACE_PARTICLE_EXTRA_DATA_IN_CHANNEL macros (channel-aware variant
sets SourceChannelId on the category for display in load-time dialogs)

#Synced-CL 52152484

diff --git a/Engine/Source/Runtime/Experimental/Chaos/Private/ChaosVisualDebugger/ChaosVDParticleExtraData.cpp b/Engine/Source/Runtime/Experimental/Chaos/Private/ChaosVisualDebugger/ChaosVDParticleExtraData.cpp
new file mode 100644
--- /dev/null
+++ b/Engine/Source/Runtime/Experimental/Chaos/Private/ChaosVisualDebugger/ChaosVDParticleExtraData.cpp
@@ -0,0 +1,129 @@
+// Copyright Epic Games, Inc. All Rights Reserved.
+
+#include "ChaosVisualDebugger/ChaosVDParticleExtraData.h"
+
+#include "ChaosVisualDebugger/ChaosVDMemWriterReader.h"
+#include "ChaosVisualDebugger/ChaosVDStructCollectionMemWriterReader.h"
+#include "ChaosVisualDebugger/ChaosVisualDebuggerTrace.h"
+#include "Logging/LogMacros.h"
+#include "UObject/Class.h"
+
+DEFINE_LOG_CATEGORY_STATIC(LogChaosVDParticleExtraData, Log, All);
+
+namespace Chaos::VisualDebugger
+{
+	FStringView FChaosVDParticleExtraData::WrapperTypeName = TEXT("FChaosVDParticleExtraData");
+
+	bool FChaosVDExtraDataStructEntry::Serialize(FArchive& Ar)
+	{
+		Ar << StructTypePath;
+		uint8 ModeVal = static_cast<uint8>(SerializationMode);
+		Ar << ModeVal;
+		if (Ar.IsLoading())
+		{
+			SerializationMode = static_cast<EChaosVDExtraDataSerializationMode>(ModeVal);
+		}
+		Ar << Bytes;
+		return !Ar.IsError();
+	}
+
+	void FChaosVDExtraDataStructEntry::SerializeFrom(UScriptStruct* Struct, const void* Data,
+		const TSharedRef<FChaosVDSerializableNameTable>& NameTable,
+		EChaosVDExtraDataSerializationMode Mode)
+	{
+		SerializationMode = Mode;
+		FChaosVDStructCollectionMemoryWriter Writer(Bytes, NameTable);
+		if (Mode == EChaosVDExtraDataSerializationMode::SerializeBin)
+		{
+			Struct->SerializeBin(Writer, const_cast<void*>(Data));
+		}
+		else
+		{
+			Struct->SerializeItem(Writer, const_cast<void*>(Data), nullptr);
+		}
+	}
+
+	void FChaosVDExtraDataCategory::AddEntry(UScriptStruct* Struct, const void* Data, EChaosVDExtraDataSerializationMode Mode)
+	{
+#if WITH_CHAOS_VISUAL_DEBUGGER
+		FChaosVDExtraDataStructEntry& Entry = Entries.AddDefaulted_GetRef();
+		Entry.StructTypePath = FName(Struct->GetPathName());
+		Entry.SerializeFrom(Struct, Data, FChaosVisualDebuggerTrace::GetNameTableInstance(), Mode);
+#endif
+	}
+
+	bool FChaosVDExtraDataCategory::Serialize(FArchive& Ar)
+	{
+		Ar << CategoryName;
+		Ar << SourceChannelId;
+
+		int32 NumEntries = Entries.Num();
+		Ar << NumEntries;
+		if (Ar.IsLoading())
+		{
+			constexpr int32 MaxExpectedEntries = 256;
+			if (NumEntries < 0 || NumEntries > MaxExpectedEntries)
+			{
+				UE_LOG(LogChaosVDParticleExtraData, Warning,
+					TEXT("FChaosVDExtraDataCategory::Serialize - Invalid NumEntries (%d) for category '%s', expected 0..%d. Data may be corrupt."),
+					NumEntries, *CategoryName.ToString(), MaxExpectedEntries);
+				Ar.SetError();
+				return false;
+			}
+			Entries.SetNum(NumEntries);
+		}
+		for (FChaosVDExtraDataStructEntry& Entry : Entries)
+		{
+			Entry.Serialize(Ar);
+		}
+		return !Ar.IsError();
+	}
+
+	bool FChaosVDParticleExtraData::Serialize(FArchive& Ar)
+	{
+		Ar << SolverID;
+		Ar << ParticleID;
+
+		int32 NumCategories = Categories.Num();
+		Ar << NumCategories;
+		if (Ar.IsLoading())
+		{
+			constexpr int32 MaxExpectedCategories = 256;
+			if (NumCategories < 0 || NumCategories > MaxExpectedCategories)
+			{
+				UE_LOG(LogChaosVDParticleExtraData, Warning,
+					TEXT("FChaosVDParticleExtraData::Serialize - Invalid NumCategories (%d), expected 0..%d. Data may be corrupt."),
+					NumCategories, MaxExpectedCategories);
+				Ar.SetError();
+				return false;
+			}
+			Categories.SetNum(NumCategories);
+		}
+		for (FChaosVDExtraDataCategory& Category : Categories)
+		{
+			Category.Serialize(Ar);
+		}
+		return !Ar.IsError();
+	}
+
+	void TraceChaosVDParticleExtraData(int32 SolverID, int32 ParticleID, FChaosVDParticleExtraData& InExtraData)
+	{
+#if WITH_CHAOS_VISUAL_DEBUGGER
+		if (!FChaosVisualDebuggerTrace::IsTracing())
+		{
+			return;
+		}
+
+		InExtraData.SolverID = SolverID;
+		InExtraData.ParticleID = ParticleID;
+
+		TArray<uint8> Buffer;
+		const TSharedRef<FChaosVDSerializableNameTable> NameTable = FChaosVisualDebuggerTrace::GetNameTableInstance();
+		FChaosVDMemoryWriter Writer(Buffer, NameTable);
+
+		InExtraData.Serialize(Writer);
+
+		FChaosVisualDebuggerTrace::TraceBinaryData(Buffer, FChaosVDParticleExtraData::WrapperTypeName);
+#endif
+	}
+}
diff --git a/Engine/Source/Runtime/Experimental/Chaos/Public/ChaosVisualDebugger/ChaosVDParticleExtraData.h b/Engine/Source/Runtime/Experimental/Chaos/Public/ChaosVisualDebugger/ChaosVDParticleExtraData.h
new file mode 100644
--- /dev/null
+++ b/Engine/Source/Runtime/Experimental/Chaos/Public/ChaosVisualDebugger/ChaosVDParticleExtraData.h
@@ -0,0 +1,113 @@
+// Copyright Epic Games, Inc. All Rights Reserved.
+
+#pragma once
+
+#include "Containers/Array.h"
+#include "Containers/Map.h"
+#include "Containers/StringFwd.h"
+#include "HAL/Platform.h"
+#include "Templates/SharedPointer.h"
+#include "UObject/NameTypes.h"
+
+class FArchive;
+class UScriptStruct;
+
+namespace Chaos::VisualDebugger
+{
+	class FChaosVDSerializableNameTable;
+
+	/** Controls how the struct content bytes inside FChaosVDExtraDataStructEntry were written.
+	 *  Recorded in the trace so the reader can apply the symmetric deserialization. */
+	enum class EChaosVDExtraDataSerializationMode : uint8
+	{
+		/** Default. Uses UStruct::SerializeBin — property-tagged, handles field additions/removals gracefully. */
+		SerializeBin = 0,
+		/** Opt-in. Uses UScriptStruct::SerializeItem — invokes native Serialize if present.
+		 *  More expressive but potentially binary-incompatible across builds. */
+		NativeSerialization = 1,
+	};
+
+	/** A single serialized struct entry: the reflected type path and the struct's serialized bytes. */
+	struct FChaosVDExtraDataStructEntry
+	{
+		/** Full path of the UScriptStruct, e.g. "/Script/Mover.FMoverDefaultSyncState".
+		 *  Stored as FName so it deduplicates via the CVD name table when serialized
+		 *  with FChaosVDMemoryWriter/Reader. */
+		FName StructTypePath;
+
+		/** Struct content serialized with FChaosVDStructCollectionMemoryWriter. */
+		TArray<uint8> Bytes;
+
+		/** How the bytes were written. Read back at load time to apply symmetric deserialization. */
+		EChaosVDExtraDataSerializationMode SerializationMode = EChaosVDExtraDataSerializationMode::SerializeBin;
+
+		CHAOS_API bool Serialize(FArchive& Ar);
+
+		/** Serialize the struct into Bytes using the given mode and record the mode.
+		 *  Callers should use this instead of writing Bytes manually. */
+		CHAOS_API void SerializeFrom(UScriptStruct* Struct, const void* Data,
+			const TSharedRef<FChaosVDSerializableNameTable>& NameTable,
+			EChaosVDExtraDataSerializationMode Mode = EChaosVDExtraDataSerializationMode::SerializeBin);
+	};
+
+	/** A named group of struct entries.  Entries sharing a CategoryName are displayed together
+	 *  in the CVD particle details panel under that category. */
+	struct FChaosVDExtraDataCategory
+	{
+		FName CategoryName;
+		/** Optional: ID of the CVD optional data channel this category was traced through.
+		 *  NAME_None when the caller did not specify a channel. */
+		FName SourceChannelId;
+		TArray<FChaosVDExtraDataStructEntry> Entries;
+
+		CHAOS_API bool Serialize(FArchive& Ar);
+
+		/** Serialize a struct instance into a new entry. Fetches the CVD name table internally.
+		 *  This is a no-op when WITH_CHAOS_VISUAL_DEBUGGER is not set. */
+		CHAOS_API void AddEntry(UScriptStruct* Struct, const void* Data,
+			EChaosVDExtraDataSerializationMode Mode = EChaosVDExtraDataSerializationMode::SerializeBin);
+	};
+
+	/**
+	 * Per-particle extra data payload.  Contains an ordered list of named categories,
+	 * each holding zero or more serialized struct entries.
+	 *
+	 * Written at runtime via TraceChaosVDParticleExtraData() and stored in
+	 * FChaosVDParticleExtraDataContainer inside FChaosVDSolverFrameData::CustomData.
+	 * On the editor side, FChaosVDParticleExtraDataComponent reads this container and
+	 * injects the reconstructed structs into the particle details panel via the
+	 * FChaosVDScene::OnParticleDetailsEnrichmentRequested delegate.
+	 */
+	struct FChaosVDParticleExtraData
+	{
+		CHAOS_API static FStringView WrapperTypeName;
+
+		int32 SolverID = INDEX_NONE;
+		int32 ParticleID = INDEX_NONE;
+
+		TArray<FChaosVDExtraDataCategory> Categories;
+
+		CHAOS_API bool Serialize(FArchive& Ar);
+
+		bool HasValidData() const
+		{
+			return SolverID != INDEX_NONE && ParticleID != INDEX_NONE;
+		}
+	};
+
+	/**
+	 * Traces per-particle extra struct data for the current solver frame.
+	 *
+	 * Call this from any module (after the particle has been processed) to associate
+	 * arbitrary named categories of UStructs with a specific particle. The data is
+	 * keyed by SolverID + ParticleID and made available in the CVD particle details
+	 * panel without requiring changes to the particle trace call site.
+	 *
+	 * This is a no-op when CVD tracing is disabled (WITH_CHAOS_VISUAL_DEBUGGER=0).
+	 *
+	 * @param SolverID       ID of the solver owning the particle.
+	 * @param ParticleID     Index of the particle within the solver.
+	 * @param InExtraData    Populated extra data payload to trace.
+	 */
+	CHAOS_API void TraceChaosVDParticleExtraData(int32 SolverID, int32 ParticleID, FChaosVDParticleExtraData& InExtraData);
+}
diff --git a/Engine/Source/Runtime/Experimental/Chaos/Public/ChaosVisualDebugger/ChaosVDStructCollectionMemWriterReader.h b/Engine/Source/Runtime/Experimental/Chaos/Public/ChaosVisualDebugger/ChaosVDStructCollectionMemWriterReader.h
new file mode 100644
--- /dev/null
+++ b/Engine/Source/Runtime/Experimental/Chaos/Public/ChaosVisualDebugger/ChaosVDStructCollectionMemWriterReader.h
@@ -0,0 +1,130 @@
+// Copyright Epic Games, Inc. All Rights Reserved.
+
+#pragma once
+
+#include "ChaosVDMemWriterReader.h"
+#include "Serialization/MemoryReader.h"
+#include "Serialization/MemoryWriter.h"
+#include "UObject/Class.h"
+#include "UObject/ObjectPtr.h"
+#include "UObject/UObjectGlobals.h"
+#include "UObject/WeakObjectPtr.h"
+
+namespace Chaos::VisualDebugger
+{
+	/**
+	 * Memory writer that combines CVD name-table FName serialization with object-reference skipping.
+	 *
+	 * - FNames are written as uint64 name table IDs (same as FChaosVDMemoryWriter).
+	 * - UScriptStruct object references are preserved as their full path string so they can
+	 *   be resolved on the reading side via FindObject. This allows FInstancedStruct and
+	 *   similar struct-typed UPROPERTY members to round-trip correctly.
+	 * - All other UObject and FWeakObjectPtr references are silently dropped (left null on read).
+	 *
+	 * Intended for serializing arbitrary UStruct instances into CVD particle extra data.
+	 */
+	class FChaosVDStructCollectionMemoryWriter : public FMemoryWriter
+	{
+	public:
+		FChaosVDStructCollectionMemoryWriter(TArray<uint8>& InBytes, const TSharedRef<FChaosVDSerializableNameTable>& InNameTable)
+			: FMemoryWriter(InBytes), NameTableInstance(InNameTable)
+		{
+		}
+
+		virtual FArchive& operator<<(FName& Name) override
+		{
+			NameTableFNameSerializer(*this, Name, NameTableInstance);
+			return *this;
+		}
+
+		virtual FArchive& operator<<(FObjectPtr& Value) override
+		{
+			UObject* Obj = Value.Get();
+			SerializeObjectRef(Obj);
+			return *this;
+		}
+
+		virtual FArchive& operator<<(UObject*& Res) override
+		{
+			SerializeObjectRef(Res);
+			return *this;
+		}
+
+		// Weak object refs are always dropped — nothing written to the buffer.
+		virtual FArchive& operator<<(FWeakObjectPtr& Value) override
+		{
+			return *this;
+		}
+
+		TSharedRef<FChaosVDSerializableNameTable> NameTableInstance;
+
+	private:
+		void SerializeObjectRef(UObject* Obj)
+		{
+			const bool bIsScriptStruct = Obj && Obj->IsA<UScriptStruct>();
+			*this << const_cast<bool&>(bIsScriptStruct);
+			if (bIsScriptStruct)
+			{
+				// Strip the "ScriptStruct " prefix (13 chars) — same convention as FSkipObjectRefsMemoryWriter.
+				FString StructName = Obj->GetFullName(nullptr).RightChop(13);
+				*this << StructName;
+			}
+		}
+	};
+
+	/**
+	 * Memory reader counterpart of FChaosVDStructCollectionMemoryWriter.
+	 *
+	 * - FNames are read back via the name table.
+	 * - UScriptStruct refs are resolved via FindObject using the stored path string.
+	 * - All other object refs are set to null (they were not written).
+	 * - FWeakObjectPtr refs are left default-constructed (nothing was written).
+	 */
+	class FChaosVDStructCollectionMemoryReader : public FMemoryReader
+	{
+	public:
+		FChaosVDStructCollectionMemoryReader(const TArray<uint8>& InBytes, const TSharedRef<FChaosVDSerializableNameTable>& InNameTable)
+			: FMemoryReader(InBytes), NameTableInstance(InNameTable)
+		{
+		}
+
+		virtual FArchive& operator<<(FName& Name) override
+		{
+			NameTableFNameSerializer(*this, Name, NameTableInstance);
+			return *this;
+		}
+
+		virtual FArchive& operator<<(FObjectPtr& Value) override
+		{
+			UObject* Obj = nullptr;
+			*this << Obj;
+			Value = Obj;
+			return *this;
+		}
+
+		virtual FArchive& operator<<(UObject*& Res) override
+		{
+			bool bIsScriptStruct = false;
+			*this << bIsScriptStruct;
+			if (bIsScriptStruct)
+			{
+				FString StructName;
+				*this << StructName;
+				Res = FindObject<UScriptStruct>(nullptr, *StructName);
+			}
+			else
+			{
+				Res = nullptr;
+			}
+			return *this;
+		}
+
+		// Nothing was written for weak refs; leave value default-constructed.
+		virtual FArchive& operator<<(FWeakObjectPtr& Value) override
+		{
+			return *this;
+		}
+
+		TSharedRef<FChaosVDSerializableNameTable> NameTableInstance;
+	};
+}
diff --git a/Engine/Source/Runtime/Experimental/Chaos/Public/ChaosVisualDebugger/ChaosVDTraceMacros.h b/Engine/Source/Runtime/Experimental/Chaos/Public/ChaosVisualDebugger/ChaosVDTraceMacros.h
--- a/Engine/Source/Runtime/Experimental/Chaos/Public/ChaosVisualDebugger/ChaosVDTraceMacros.h
+++ b/Engine/Source/Runtime/Experimental/Chaos/Public/ChaosVisualDebugger/ChaosVDTraceMacros.h
@@ -229,6 +229,62 @@
 		FChaosVisualDebuggerTrace::SetTraceRelevancyVolume(VolumeBox)
 	#endif
 
+	/**
+	 * Trace a single UStruct instance as particle extra data in one line.
+	 * StructType must be a USTRUCT (provides ::StaticStruct()). Mode defaults to SerializeBin.
+	 *
+	 * Unlike the rest of the particle trace data, extra data is NOT guaranteed to be backwards
+	 * compatible. This is by design: it trades version safety for convenience — callers can pass
+	 * any USTRUCT directly without writing a dedicated data wrapper.
+	 *
+	 * The default mode (SerializeBin) uses property-tagged serialization and handles field additions
+	 * and removals gracefully. Passing EChaosVDExtraDataSerializationMode::NativeSerialization invokes
+	 * the struct's native Serialize function, which may be binary-incompatible across builds and can
+	 * cause a crash if the struct layout has changed. Native serialization generally requires the CVD
+	 * viewer binary to match the one that produced the recording; it is safe to use when CVD runs
+	 * from inside the editor against a session recorded by the same editor build.
+	 *
+	 * If this causes issues when loading older recordings, enable Safe Loading Mode in CVD General
+	 * Settings (Project > Plugins > Chaos Visual Debugger > General). When active, CVD skips all
+	 * non-backwards-compatible data rather than risking a crash.
+	 *
+	 * Usage:
+	 *   CVD_TRACE_PARTICLE_EXTRA_DATA(SolverID, ParticleID, "MyCategory", FMyStruct, MyData);
+	 *   CVD_TRACE_PARTICLE_EXTRA_DATA(SolverID, ParticleID, "MyCategory", FMyStruct, MyData, EChaosVDExtraDataSerializationMode::NativeSerialization);
+	 */
+	#ifndef CVD_TRACE_PARTICLE_EXTRA_DATA
+		#define CVD_TRACE_PARTICLE_EXTRA_DATA(SolverID, ParticleID, CategoryName, StructType, Data, ...) \
+		{ \
+			if (FChaosVisualDebuggerTrace::IsTracing()) \
+			{ \
+				Chaos::VisualDebugger::FChaosVDParticleExtraData CVDParticleExtraData__; \
+				Chaos::VisualDebugger::FChaosVDExtraDataCategory& CVDCategory__ = CVDParticleExtraData__.Categories.AddDefaulted_GetRef(); \
+				CVDCategory__.CategoryName = (CategoryName); \
+				CVDCategory__.AddEntry(StructType::StaticStruct(), &(Data), ##__VA_ARGS__); \
+				Chaos::VisualDebugger::TraceChaosVDParticleExtraData((SolverID), (ParticleID), CVDParticleExtraData__); \
+			} \
+		}
+	#endif
+
+	/** Channel-aware variant of CVD_TRACE_PARTICLE_EXTRA_DATA.
+	 *  ChannelId should be passed as FName(TEXT("ChannelName")) or NAME_None.
+	 *  Records which CVD optional data channel this entry was traced from so the
+	 *  post-load native-serialization dialog can name the source channel. */
+	#ifndef CVD_TRACE_PARTICLE_EXTRA_DATA_IN_CHANNEL
+		#define CVD_TRACE_PARTICLE_EXTRA_DATA_IN_CHANNEL(SolverID, ParticleID, ChannelId, CategoryName, StructType, Data, ...) \
+		{ \
+			if (FChaosVisualDebuggerTrace::IsTracing()) \
+			{ \
+				Chaos::VisualDebugger::FChaosVDParticleExtraData CVDParticleExtraData__; \
+				Chaos::VisualDebugger::FChaosVDExtraDataCategory& CVDCategory__ = CVDParticleExtraData__.Categories.AddDefaulted_GetRef(); \
+				CVDCategory__.CategoryName = (CategoryName); \
+				CVDCategory__.SourceChannelId = (ChannelId); \
+				CVDCategory__.AddEntry(StructType::StaticStruct(), &(Data), ##__VA_ARGS__); \
+				Chaos::VisualDebugger::TraceChaosVDParticleExtraData((SolverID), (ParticleID), CVDParticleExtraData__); \
+			} \
+		}
+	#endif
+
 #else // WITH_CHAOS_VISUAL_DEBUGGER
 
 	#ifndef CVD_TRACE_PARTICLE
@@ -369,4 +425,12 @@
 		#define CVD_SET_RELEVANCY_VOLUME(VolumeBox)
 	#endif
 
+	#ifndef CVD_TRACE_PARTICLE_EXTRA_DATA
+		#define CVD_TRACE_PARTICLE_EXTRA_DATA(SolverID, ParticleID, CategoryName, StructType, Data, ...)
+	#endif
+
+	#ifndef CVD_TRACE_PARTICLE_EXTRA_DATA_IN_CHANNEL
+		#define CVD_TRACE_PARTICLE_EXTRA_DATA_IN_CHANNEL(SolverID, ParticleID, ChannelId, CategoryName, StructType, Data, ...)
+	#endif
+
 #endif // WITH_CHAOS_VISUAL_DEBUGGER
\ No newline at end of file
