Chaos Visual Debugger: Clarifying serialization code to prevent the static analizer from freaking out

#Synced-CL 52152484

diff --git a/Engine/Source/Runtime/Experimental/Chaos/Public/ChaosVisualDebugger/ChaosVDStructCollectionMemWriterReader.h b/Engine/Source/Runtime/Experimental/Chaos/Public/ChaosVisualDebugger/ChaosVDStructCollectionMemWriterReader.h
--- a/Engine/Source/Runtime/Experimental/Chaos/Public/ChaosVisualDebugger/ChaosVDStructCollectionMemWriterReader.h
+++ b/Engine/Source/Runtime/Experimental/Chaos/Public/ChaosVisualDebugger/ChaosVDStructCollectionMemWriterReader.h
@@ -61,12 +61,13 @@ namespace Chaos::VisualDebugger
 	private:
 		void SerializeObjectRef(UObject* Obj)
 		{
-			const bool bIsScriptStruct = Obj && Obj->IsA<UScriptStruct>();
-			*this << const_cast<bool&>(bIsScriptStruct);
-			if (bIsScriptStruct)
+			UScriptStruct* ScriptStruct = Obj ? Cast<UScriptStruct>(Obj) : nullptr;
+			bool bIsScriptStruct = (ScriptStruct != nullptr);
+			*this << bIsScriptStruct;
+			if (ScriptStruct)
 			{
-				// Strip the "ScriptStruct " prefix (13 chars) — same convention as FSkipObjectRefsMemoryWriter.
-				FString StructName = Obj->GetFullName(nullptr).RightChop(13);
+				// Strip the "ScriptStruct " prefix (13 chars) - same convention as FSkipObjectRefsMemoryWriter.
+				FString StructName = ScriptStruct->GetFullName(nullptr).RightChop(13);
 				*this << StructName;
 			}
 		}
