Add CVar to enable storing shader code and shader code libraries in memory mapped on-demand cache.

#rnx

#Synced-CL 52465190

diff --git a/Engine/Source/Runtime/Experimental/IoStore/OnDemand/Private/OnDemandContentInstaller.cpp b/Engine/Source/Runtime/Experimental/IoStore/OnDemand/Private/OnDemandContentInstaller.cpp
--- a/Engine/Source/Runtime/Experimental/IoStore/OnDemand/Private/OnDemandContentInstaller.cpp
+++ b/Engine/Source/Runtime/Experimental/IoStore/OnDemand/Private/OnDemandContentInstaller.cpp
@@ -16,6 +16,7 @@
 #include "IO/OnDemandError.h"
 #include "Logging/StructuredLog.h"
 #include "Misc/App.h"
+#include "Misc/CoreDelegates.h" // USE_MMAPPED_SHADERARCHIVE
 #include "Misc/CoreDelegatesInternal.h"
 #include "Misc/Timespan.h"
 #include "OnDemandContentInstallReplay.h"
@@ -96,11 +97,28 @@ namespace UE::IoStore
 			ECVF_SaveForNextBoot
 			);
 
+		static TAutoConsoleVariable<bool> CVar_MMapShadersEnabled(
+			TEXT("iostore.MMapShadersEnabled"),
+			false,
+			TEXT("Whether shader code and shader libraries should be stored in the memory mapped CAS backend."),
+			ECVF_SaveForNextBoot
+			);
+
 		static bool EnableConcurentInstallRequests()
 		{
 			static bool bEnableConcurentInstallRequests = CVar_EnableConcurrentRequests.GetValueOnAnyThread();
 			return bEnableConcurentInstallRequests;
 		}
+
+		static bool MemoryMappedShadersEnabled()
+		{
+#if USE_MMAPPED_SHADERARCHIVE
+			static bool bMmapShadersEnabled = CVar_MMapShadersEnabled.GetValueOnAnyThread();
+			return bMmapShadersEnabled;
+#else
+			return false;
+#endif
+		}
 	}
 
 namespace Private
@@ -642,9 +660,16 @@ void FOnDemandContentInstaller::ReportAnalytics(TArray<FAnalyticsEventAttribute>
 EOnDemandInstallCasType FOnDemandContentInstaller::GetCasType(EIoChunkType ChunkType) const
 {
 	EOnDemandInstallCasType CasType = EOnDemandInstallCasType::General;
-	if (bMMapCasEnabled && ChunkType == EIoChunkType::MemoryMappedBulkData)
+	if (bMMapCasEnabled)
 	{
-		CasType = EOnDemandInstallCasType::MMap;
+		if (ChunkType == EIoChunkType::MemoryMappedBulkData)
+		{
+			CasType = EOnDemandInstallCasType::MMap;
+		}
+		else if (CVars::MemoryMappedShadersEnabled() && (ChunkType == EIoChunkType::ShaderCode || ChunkType == EIoChunkType::ShaderCodeLibrary))
+		{
+			CasType = EOnDemandInstallCasType::MMap;
+		}
 	}
 	return CasType;
 }
