[NetPhysics] Fix dangling particle handle in RewindData::DirtyParticlePreResimState during resim
- RemoveObject now also cleans DirtyParticlePreResimState and DirtyParticleErrors
- Moved CleanupInternalClustersForProxies inside DestroyOnStep guard for GC proxies

#Synced-CL 52284408

diff --git a/Engine/Source/Runtime/Experimental/Chaos/Public/RewindData.h b/Engine/Source/Runtime/Experimental/Chaos/Public/RewindData.h
--- a/Engine/Source/Runtime/Experimental/Chaos/Public/RewindData.h
+++ b/Engine/Source/Runtime/Experimental/Chaos/Public/RewindData.h
@@ -1773,6 +1773,8 @@ public:
 	void RemoveObject(const FGeometryParticleHandle* Particle, const EAllowShrinking AllowShrinking=EAllowShrinking::Default)
 	{
 		DirtyParticles.Remove(Particle, AllowShrinking);
+		DirtyParticlePreResimState.Remove(Particle, AllowShrinking);
+		DirtyParticleErrors.Remove(Particle, AllowShrinking);
 	}
 
 	void RemoveObject(const FPBDJointConstraintHandle* Joint, const EAllowShrinking AllowShrinking = EAllowShrinking::Default)
