Chaos API: Adding PhysicsService and removing the dependecy of the async plugin on dataflow.

#Synced-CL 52508799

diff --git a/Engine/Plugins/Experimental/ChaosDataflowSolver/Source/ChaosDataflowSolver/Private/ChaosDataflowSolverActor.cpp b/Engine/Plugins/Experimental/ChaosDataflowSolver/Source/ChaosDataflowSolver/Private/ChaosDataflowSolverActor.cpp
--- a/Engine/Plugins/Experimental/ChaosDataflowSolver/Source/ChaosDataflowSolver/Private/ChaosDataflowSolverActor.cpp
+++ b/Engine/Plugins/Experimental/ChaosDataflowSolver/Source/ChaosDataflowSolver/Private/ChaosDataflowSolverActor.cpp
@@ -19,6 +19,7 @@
 #include "RigidPhysics/RigidBody.h"
 #include "RigidPhysics/RigidContext.h"
 #include "RigidPhysics/RigidModifier.h"
+#include "RigidPhysics/RigidPhysicsService.h"
 #include "RigidPhysics/RigidScene.h"
 #include "RigidAssetUserData.h"
 
@@ -74,10 +75,8 @@ void AChaosDataflowSolverActor::EndPlay(const EEndPlayReason::Type ReasonEnd)
 				RegisteredComponents.Empty();
 			}
 
-			if (State->Factory)
-			{
-				State->Factory->DestroyScene(State->Handle);
-			}
+			UE::Physics::FRigidPhysicsService& Service = UE::Physics::FRigidPhysicsService::GetInstance();
+			Service.DestroyScene(State->Handle);
 		}
 	}
 #endif
diff --git a/Engine/Plugins/Experimental/ChaosRigidPhysicsAsync/Source/ChaosRigidPhysicsAsync/ChaosRigidPhysicsAsync.Build.cs b/Engine/Plugins/Experimental/ChaosRigidPhysicsAsync/Source/ChaosRigidPhysicsAsync/ChaosRigidPhysicsAsync.Build.cs
--- a/Engine/Plugins/Experimental/ChaosRigidPhysicsAsync/Source/ChaosRigidPhysicsAsync/ChaosRigidPhysicsAsync.Build.cs
+++ b/Engine/Plugins/Experimental/ChaosRigidPhysicsAsync/Source/ChaosRigidPhysicsAsync/ChaosRigidPhysicsAsync.Build.cs
@@ -15,9 +15,6 @@ namespace UnrealBuildTool.Rules
                 "CoreUObject",
 				"Engine",
                 "RigidPhysics",
-				"DataflowCore",
-				"DataflowEngine",
-				"DataflowSimulation",
 				}
 			);
 
diff --git a/Engine/Plugins/Experimental/ChaosRigidPhysicsAsync/Source/ChaosRigidPhysicsAsync/ChaosRigidPhysicsAsyncModule.cpp b/Engine/Plugins/Experimental/ChaosRigidPhysicsAsync/Source/ChaosRigidPhysicsAsync/ChaosRigidPhysicsAsyncModule.cpp
--- a/Engine/Plugins/Experimental/ChaosRigidPhysicsAsync/Source/ChaosRigidPhysicsAsync/ChaosRigidPhysicsAsyncModule.cpp
+++ b/Engine/Plugins/Experimental/ChaosRigidPhysicsAsync/Source/ChaosRigidPhysicsAsync/ChaosRigidPhysicsAsyncModule.cpp
@@ -2,17 +2,29 @@
 
 #include "Modules/ModuleManager.h"
 
-#include "ChaosRigidPhysicsAsync/Dataflow/RigidPhysicsAsyncDataflowNodes.h"
+#include "ChaosRigidPhysicsAsync/RigidFactoryAsync.h"
+#include "ChaosRigidPhysicsAsync/RigidSceneSettingsAsync.h"
+#include "RigidPhysics/RigidPhysicsService.h"
 
 class FChaosRigidPhysicsAsyncModule : public FDefaultModuleImpl
 {
 public:
 	virtual void StartupModule() override
 	{
-		RegisterRigidPhysicsAsyncDataflowSimulationNodes();
+#if UE_RIGIDPHYSICS_API_ENABLED
+		TUniquePtr<Chaos::Rigids::Async::FRigidFactoryAsync> Factory = MakeUnique<Chaos::Rigids::Async::FRigidFactoryAsync>();
+		TUniquePtr<Chaos::Rigids::Async::FRigidSceneSettingsAsync> SceneSettings = MakeUnique<Chaos::Rigids::Async::FRigidSceneSettingsAsync>();
+		UE::Physics::FRigidPhysicsService& Service = UE::Physics::FRigidPhysicsService::GetInstance();
+		Service.RegisterFactory(MoveTemp(Factory), MoveTemp(SceneSettings));
+#endif // UE_RIGIDPHYSICS_API_ENABLED
 	}
+
 	virtual void ShutdownModule() override
 	{
+#if UE_RIGIDPHYSICS_API_ENABLED
+		UE::Physics::FRigidPhysicsService& Service = UE::Physics::FRigidPhysicsService::GetInstance();
+		Service.Reset();
+#endif // UE_RIGIDPHYSICS_API_ENABLED
 	}
 };
 IMPLEMENT_MODULE(FChaosRigidPhysicsAsyncModule, ChaosRigidPhysicsAsync)
\ No newline at end of file
diff --git a/Engine/Plugins/Experimental/ChaosRigidPhysicsAsync/Source/ChaosRigidPhysicsAsync/Private/ChaosRigidPhysicsAsync/Dataflow/RigidPhysicsAsyncDataflowNodes.cpp b/Engine/Plugins/Experimental/ChaosRigidPhysicsAsync/Source/ChaosRigidPhysicsAsync/Private/ChaosRigidPhysicsAsync/Dataflow/RigidPhysicsAsyncDataflowNodes.cpp
deleted file mode 100644
diff --git a/Engine/Plugins/Experimental/ChaosRigidPhysicsAsync/Source/ChaosRigidPhysicsAsync/Public/ChaosRigidPhysicsAsync/Dataflow/RigidPhysicsAsyncDataflowNodes.h b/Engine/Plugins/Experimental/ChaosRigidPhysicsAsync/Source/ChaosRigidPhysicsAsync/Public/ChaosRigidPhysicsAsync/Dataflow/RigidPhysicsAsyncDataflowNodes.h
deleted file mode 100644
diff --git a/Engine/Source/Runtime/Experimental/Dataflow/Simulation/Private/Dataflow/DataflowSimulationNodes.cpp b/Engine/Source/Runtime/Experimental/Dataflow/Simulation/Private/Dataflow/DataflowSimulationNodes.cpp
--- a/Engine/Source/Runtime/Experimental/Dataflow/Simulation/Private/Dataflow/DataflowSimulationNodes.cpp
+++ b/Engine/Source/Runtime/Experimental/Dataflow/Simulation/Private/Dataflow/DataflowSimulationNodes.cpp
@@ -5,6 +5,7 @@
 #include "Dataflow/DataflowNodeFactory.h"
 #include "Dataflow/DataflowNodeColorsRegistry.h"
 #include "Dataflow/Interfaces/DataflowPhysicsSolver.h"
+#include "Dataflow/RigidPhysicsDataflowNodes.h"
 
 #include "RigidPhysics/RigidContext.h"
 #include "RigidPhysics/RigidScene.h"
@@ -32,6 +33,7 @@ namespace UE::Dataflow
 
 		DATAFLOW_NODE_REGISTER_CREATION_FACTORY(FAdvanceRigidPhysicsSolverDataflowNode);
 		DATAFLOW_NODE_REGISTER_CREATION_FACTORY(FSceneSimulationTerminalDataflowNode);
+		RegisterRigidPhysicsDataflowSimulationNodes();
 #endif
 		
 		static constexpr FLinearColor CDefaultNodeBodyTintColor = FLinearColor(0.f, 0.f, 0.0f, 0.5f);
diff --git a/Engine/Source/Runtime/Experimental/Dataflow/Simulation/Private/Dataflow/RigidPhysicsDataflowNodes.cpp b/Engine/Source/Runtime/Experimental/Dataflow/Simulation/Private/Dataflow/RigidPhysicsDataflowNodes.cpp
new file mode 100644
--- /dev/null
+++ b/Engine/Source/Runtime/Experimental/Dataflow/Simulation/Private/Dataflow/RigidPhysicsDataflowNodes.cpp
@@ -0,0 +1,42 @@
+// Copyright Epic Games, Inc. All Rights Reserved.
+
+#include "Dataflow/RigidPhysicsDataflowNodes.h"
+#include "Dataflow/Interfaces/DataflowPhysicsSolver.h"
+#include "Dataflow/DataflowSimulationNodes.h"
+#include "Dataflow/DataflowSimulationManager.h"
+#include "Dataflow/DataflowNodeFactory.h"
+#include "Dataflow/DataflowNodeColorsRegistry.h"
+#include "RigidPhysics/RigidBody.h"
+#include "RigidPhysics/RigidContext.h"
+#include "RigidPhysics/RigidModifier.h"
+#include "RigidPhysics/RigidScene.h"
+#include "RigidPhysics/RigidPhysicsService.h"
+
+#include UE_INLINE_GENERATED_CPP_BY_NAME(RigidPhysicsDataflowNodes)
+void RegisterRigidPhysicsDataflowSimulationNodes()
+{
+#if UE_RIGIDPHYSICS_API_ENABLED
+	DATAFLOW_NODE_REGISTER_CREATION_FACTORY(FRigidSceneSetupDataflowNode);
+#endif
+}
+
+void FRigidSceneSetupDataflowNode::Evaluate(UE::Dataflow::FContext& Context, const FDataflowOutput* Output) const
+{
+#if UE_RIGIDPHYSICS_API_ENABLED
+	using namespace Chaos;
+	using namespace Chaos::Rigids::Async;
+	using namespace UE::Physics;
+
+	if (Context.IsA(UE::Dataflow::FDataflowSimulationContext::StaticType()))
+	{
+		UE::Dataflow::FDataflowSimulationContext& SimulationContext = StaticCast<UE::Dataflow::FDataflowSimulationContext&>(Context);
+		if (Output->IsA(&State))
+		{
+			FRigidPhysicsSceneDataflowState InState;
+			FRigidPhysicsService& Service = FRigidPhysicsService::GetInstance();
+			InState.Handle = Service.CreateScene(FRigidDebugName(SimulationContext.GetRigidSceneName()));
+			SetValue(SimulationContext, InState, &State);
+		}
+	}
+#endif
+}
\ No newline at end of file
diff --git a/Engine/Source/Runtime/Experimental/Dataflow/Simulation/Public/Dataflow/RigidPhysicsDataflowNodes.h b/Engine/Source/Runtime/Experimental/Dataflow/Simulation/Public/Dataflow/RigidPhysicsDataflowNodes.h
new file mode 100644
--- /dev/null
+++ b/Engine/Source/Runtime/Experimental/Dataflow/Simulation/Public/Dataflow/RigidPhysicsDataflowNodes.h
@@ -0,0 +1,31 @@
+// Copyright Epic Games, Inc. All Rights Reserved.
+#pragma once
+#include "Dataflow/DataflowObject.h"
+#include "Dataflow/DataflowSimulationProxy.h"
+#include "Dataflow/DataflowSimulationContext.h"
+#include "Dataflow/DataflowSimulationNodes.h"
+#include "Dataflow/RigidPhysicsSceneState.h"
+#include "RigidPhysics/RigidScene.h"
+#include "RigidPhysicsDataflowNodes.generated.h"
+
+/** Create a rigid scene */
+USTRUCT(meta = (DataflowSimulation))
+struct FRigidSceneSetupDataflowNode : public FDataflowNode
+{
+	GENERATED_USTRUCT_BODY()
+	DATAFLOW_NODE_DEFINE_INTERNAL(FRigidSceneSetupDataflowNode, "CreateRigidScene", "Physics|Scene", UDataflow::SimulationTag)
+public:
+
+	FRigidSceneSetupDataflowNode(const UE::Dataflow::FNodeParameters& InParam, FGuid InGuid = FGuid::NewGuid())
+		: FDataflowNode(InParam, InGuid)
+	{
+		RegisterOutputConnection(&State);
+	}
+
+	UPROPERTY(meta = (DataflowOutput))
+	FRigidPhysicsSceneDataflowState State;
+	virtual void Evaluate(UE::Dataflow::FContext& Context, const FDataflowOutput* Output) const override;
+	//~ End FDataflowNode interface
+};
+
+void RegisterRigidPhysicsDataflowSimulationNodes();
diff --git a/Engine/Source/Runtime/Experimental/Dataflow/Simulation/Public/Dataflow/RigidPhysicsSceneState.h b/Engine/Source/Runtime/Experimental/Dataflow/Simulation/Public/Dataflow/RigidPhysicsSceneState.h
--- a/Engine/Source/Runtime/Experimental/Dataflow/Simulation/Public/Dataflow/RigidPhysicsSceneState.h
+++ b/Engine/Source/Runtime/Experimental/Dataflow/Simulation/Public/Dataflow/RigidPhysicsSceneState.h
@@ -17,7 +17,6 @@ struct FRigidPhysicsSceneDataflowState
 	GENERATED_BODY()
 public:
 #if UE_RIGIDPHYSICS_API_ENABLED
-	TSharedPtr<UE::Physics::IRigidFactory> Factory;
 	UE::Physics::FRigidSceneHandle Handle;
 #endif
 };
diff --git a/Engine/Source/Runtime/Experimental/RigidPhysics/Private/RigidPhysics/RigidPhysicsService.cpp b/Engine/Source/Runtime/Experimental/RigidPhysics/Private/RigidPhysics/RigidPhysicsService.cpp
new file mode 100644
--- /dev/null
+++ b/Engine/Source/Runtime/Experimental/RigidPhysics/Private/RigidPhysics/RigidPhysicsService.cpp
@@ -0,0 +1,44 @@
+// Copyright Epic Games, Inc. All Rights Reserved.
+#include "RigidPhysics/RigidPhysicsService.h"
+
+#if UE_RIGIDPHYSICS_API_ENABLED
+
+namespace UE::Physics
+{
+	FRigidPhysicsService& FRigidPhysicsService::GetInstance()
+	{
+		static FRigidPhysicsService Instance;
+		return Instance;
+	}
+
+	void FRigidPhysicsService::RegisterFactory(TUniquePtr<IRigidFactory>&& InFactory, TUniquePtr<IRigidSceneSettings>&& InDefaultSettings)
+	{
+		Factory = MoveTemp(InFactory);
+		SceneSettings = MoveTemp(InDefaultSettings);
+	}
+
+	FRigidSceneHandle FRigidPhysicsService::CreateScene(const FRigidDebugName& InName) const
+	{
+		if (ensure(Factory.IsValid() && SceneSettings.IsValid()))
+		{
+			return Factory->CreateScene(InName, SceneSettings.Get());
+		}
+		return FRigidSceneHandle();
+	}
+
+	void FRigidPhysicsService::DestroyScene(const FRigidSceneHandle& InHandle) const
+	{
+		if (ensure(Factory.IsValid()))
+		{
+			return Factory->DestroyScene(InHandle);
+		}
+	}
+
+	void FRigidPhysicsService::Reset()
+	{
+		Factory = nullptr;
+		SceneSettings = nullptr;
+	}
+} // namespace UE::Physics
+
+#endif // UE_RIGIDPHYSICS_API_ENABLED
diff --git a/Engine/Source/Runtime/Experimental/RigidPhysics/Public/RigidPhysics/RigidPhysicsService.h b/Engine/Source/Runtime/Experimental/RigidPhysics/Public/RigidPhysics/RigidPhysicsService.h
new file mode 100644
--- /dev/null
+++ b/Engine/Source/Runtime/Experimental/RigidPhysics/Public/RigidPhysics/RigidPhysicsService.h
@@ -0,0 +1,32 @@
+// Copyright Epic Games, Inc. All Rights Reserved.
+#pragma once
+
+#include "RigidPhysics/RigidFwd.h"
+
+#if UE_RIGIDPHYSICS_API_ENABLED
+
+#include "RigidPhysics/RigidFactory.h"
+#include "RigidPhysics/RigidSceneHandle.h"
+#include "RigidPhysics/RigidSceneSettings.h"
+
+UE_EXPERIMENTAL(5.8, "The new Chaos API is experimental")
+namespace UE::Physics
+{
+	class FRigidPhysicsService
+	{
+	public:
+		RIGIDPHYSICS_API static FRigidPhysicsService& GetInstance();
+
+		RIGIDPHYSICS_API void RegisterFactory(TUniquePtr<IRigidFactory>&& InFactory, TUniquePtr<IRigidSceneSettings>&& InDefaultSettings);
+
+		RIGIDPHYSICS_API FRigidSceneHandle CreateScene(const FRigidDebugName& InName) const;
+		RIGIDPHYSICS_API void DestroyScene(const FRigidSceneHandle& InHandle) const;
+
+		RIGIDPHYSICS_API void Reset();
+	private:
+		TUniquePtr<IRigidFactory> Factory;
+		TUniquePtr<IRigidSceneSettings> SceneSettings;
+	};
+} // namespace UE::Physics
+
+#endif // UE_RIGIDPHYSICS_API_ENABLED
