Resim Optimizations 1. During resim, only FullResim particles need caching skip ResimAsFollower and FrozenDuringResim.

#Synced-CL 52476441

diff --git a/Engine/Source/Runtime/Experimental/Chaos/Private/RewindData.cpp b/Engine/Source/Runtime/Experimental/Chaos/Private/RewindData.cpp
--- a/Engine/Source/Runtime/Experimental/Chaos/Private/RewindData.cpp
+++ b/Engine/Source/Runtime/Experimental/Chaos/Private/RewindData.cpp
@@ -2890,10 +2890,17 @@ void FRewindData::CacheCurrentDirtyData(const FFrameAndPhase::EParticleHistoryPh
 		return;
 	}
 
+	const bool bIsResim = IsResim();
 	for (FDirtyParticleInfo& DirtyParticleInfo : DirtyParticles)
 	{
 		if (FGeometryParticleHandle* Geometry = DirtyParticleInfo.GetObjectPtr())
 		{
+			// During resim, only FullResim particles need caching skip ResimAsFollower and FrozenDuringResim
+			// to avoid unnecessary function calls and FindOrAddDirtyObj hash lookups
+			if (bIsResim && Geometry->ResimType() != EResimType::FullResim)
+			{
+				continue;
+			}
 			CacheDirtyParticleData(Geometry, CurrentPhase, /*bDirty*/false);
 		}
 	}
