Dataflow: Fixed a race condition that would cause random crashes when async loading assets.
- Removed the redundant ToolTip self-assignment in NewNodeFromRegisteredType that was triggering the corruption.
- Changed GetToolTip() to use a const reference instead of copying the full FFactoryParameters struct.

[RN] Dataflow: Fixed a race condition that would cause random crashes when async loading assets.
#Synced-CL 52466474

diff --git a/Engine/Source/Runtime/Experimental/Dataflow/Core/Private/Dataflow/DataflowNode.cpp b/Engine/Source/Runtime/Experimental/Dataflow/Core/Private/Dataflow/DataflowNode.cpp
--- a/Engine/Source/Runtime/Experimental/Dataflow/Core/Private/Dataflow/DataflowNode.cpp
+++ b/Engine/Source/Runtime/Experimental/Dataflow/Core/Private/Dataflow/DataflowNode.cpp
@@ -1355,8 +1355,7 @@ TUniquePtr<const FStructOnScope> FDataflowNode::NewStructOnScopeConst() const
 
 FString FDataflowNode::GetToolTip() const
 {
-	UE::Dataflow::FFactoryParameters FactoryParameters = UE::Dataflow::FNodeFactory::GetInstance()->GetParameters(GetType());
-
+	const UE::Dataflow::FFactoryParameters& FactoryParameters = UE::Dataflow::FNodeFactory::GetInstance()->GetParameters(GetType());
 	return FactoryParameters.ToolTip;
 }
 
diff --git a/Engine/Source/Runtime/Experimental/Dataflow/Core/Private/Dataflow/DataflowNodeFactory.cpp b/Engine/Source/Runtime/Experimental/Dataflow/Core/Private/Dataflow/DataflowNodeFactory.cpp
--- a/Engine/Source/Runtime/Experimental/Dataflow/Core/Private/Dataflow/DataflowNodeFactory.cpp
+++ b/Engine/Source/Runtime/Experimental/Dataflow/Core/Private/Dataflow/DataflowNodeFactory.cpp
@@ -93,8 +93,6 @@ namespace UE::Dataflow
 			TUniquePtr<FDataflowNode> NewNode = ClassMap[Param.Type](Param);
 			if(NewNode->HasValidConnections())
 			{
-				ParametersMap[Param.Type].ToolTip = NewNode->GetToolTip();
-
 				return Graph.AddNode(TSharedPtr<FDataflowNode>(NewNode.Release()));
 			}
 			
