IAD - Fix CurrentBlock getting lost when a snapshot is loaded.

#Synced-CL 52456494

diff --git a/Engine/Source/Runtime/Experimental/IoStore/OnDemand/Private/OnDemandInstallCache.cpp b/Engine/Source/Runtime/Experimental/IoStore/OnDemand/Private/OnDemandInstallCache.cpp
--- a/Engine/Source/Runtime/Experimental/IoStore/OnDemand/Private/OnDemandInstallCache.cpp
+++ b/Engine/Source/Runtime/Experimental/IoStore/OnDemand/Private/OnDemandInstallCache.cpp
@@ -2745,9 +2745,10 @@ FOnDemandInstallCache::FOnDemandInstallCache(const FOnDemandInstallCacheConfig&
 				FCasSnapshot::FCasState* SnapshotCasState = Algo::FindBy(Snapshot.CasState, CacheState.Cas.GetType(), &FCasSnapshot::FCasState::CasType);
 				if (SnapshotCasState)
 				{
-					UE_LOG(LogIoStoreOnDemand, Log, TEXT("Loaded CAS state from snapshot with %d blocks and %d chunk locations"),
-						SnapshotCasState->Blocks.Num(), SnapshotCasState->ChunkLocations.Num());
+					UE_LOGF(LogIoStoreOnDemand, Log, "Loaded %ls CAS state from snapshot with %d blocks and %d chunk locations",
+						LexToString(CacheState.Cas.GetType()), SnapshotCasState->Blocks.Num(), SnapshotCasState->ChunkLocations.Num());
 
+					CacheState.CurrentBlock = SnapshotCasState->CurrentBlockId;
 					CacheState.Cas.LoadSnapshot(MoveTemp(*SnapshotCasState));
 				}
 			}
@@ -4088,7 +4089,7 @@ void FOnDemandInstallCache::AddReferencesToBlocks(
 					BlockInfo->RefSize += ChunkDiskSize;
 				}
 			}
-			// This chek is no longer valid with multiple CAS because referenced chunk is likely just in a different CAS
+			// This check is no longer valid with multiple CAS because referenced chunk is likely just in a different CAS
 			//else
 			//{
 			//	// Check pending list
@@ -4874,7 +4875,9 @@ FResult FOnDemandInstallCache::FlushPendingChunksImpl(FCacheState& CacheState, c
 
 			if (!ensure(Chunks.Chunks[ChunkIdx].GetSize() <= CacheState.Cas.GetMaxBlockSize()))
 			{
-				UE_LOGFMT(LogIoStoreOnDemand, Error, "Chunk is larger then block size ({ChunkId})", ("ChunkId",Chunks.ChunkHashes[ChunkIdx]));
+				UE_LOGFMT(LogIoStoreOnDemand, Error, "{CasType} CAS - Chunk is larger then block size ({ChunkId})", 
+					("CasType", LexToString(CacheState.Cas.GetType())), 
+					("ChunkId", Chunks.ChunkHashes[ChunkIdx]));
 			}
 
 			const FIoBuffer& Chunk = Chunks.Chunks[ChunkIdx];
@@ -4889,8 +4892,8 @@ FResult FOnDemandInstallCache::FlushPendingChunksImpl(FCacheState& CacheState, c
 
 		if (BytesToWrite > 0)
 		{
-			UE_LOG(LogIoStoreOnDemand, Log, TEXT("Writing %.2lf MiB to CAS block %u"),
-				ToMiB(Ar.Tell()), CurrentBlockId.Id);
+			UE_LOGF(LogIoStoreOnDemand, Log, "Writing %.2lf MiB to %ls CAS block %u",
+				ToMiB(Ar.Tell()), LexToString(CacheState.Cas.GetType()), CurrentBlockId.Id);
 
 			check(CasBlockOffset + BytesToWrite <= CacheState.Cas.GetMaxBlockSize());
 
