[Chaos Modular Vehicle]
Expose non-const GetHitProxy in SuspensionBaseInterface.

#Synced-CL 52518673

diff --git a/Engine/Source/Runtime/Experimental/ChaosVehicles/ChaosVehiclesCore/Public/SimModule/SuspensionBaseInterface.h b/Engine/Source/Runtime/Experimental/ChaosVehicles/ChaosVehiclesCore/Public/SimModule/SuspensionBaseInterface.h
--- a/Engine/Source/Runtime/Experimental/ChaosVehicles/ChaosVehiclesCore/Public/SimModule/SuspensionBaseInterface.h
+++ b/Engine/Source/Runtime/Experimental/ChaosVehicles/ChaosVehiclesCore/Public/SimModule/SuspensionBaseInterface.h
@@ -90,6 +90,7 @@ public:
 
 	void SetHitProxy(IPhysicsProxyBase* NewValue) { TargetPoint.HitProxy = NewValue; }
 	const IPhysicsProxyBase* GetHitProxy() const { return TargetPoint.HitProxy; }
+	IPhysicsProxyBase* GetHitProxy() { return TargetPoint.HitProxy; }
 
 protected:
 
