Dataflow : Add dataflow attachment factory and to automatically create the attachment when setting a dataflow asset via the utils methods

#Synced-CL 52497132

diff --git a/Engine/Source/Runtime/Experimental/Dataflow/Engine/Private/Dataflow/DataflowAttachment.cpp b/Engine/Source/Runtime/Experimental/Dataflow/Engine/Private/Dataflow/DataflowAttachment.cpp
--- a/Engine/Source/Runtime/Experimental/Dataflow/Engine/Private/Dataflow/DataflowAttachment.cpp
+++ b/Engine/Source/Runtime/Experimental/Dataflow/Engine/Private/Dataflow/DataflowAttachment.cpp
@@ -18,3 +18,29 @@ void UDataflowAttachment::PostLoad()
 	Super::PostLoad();
 	Instance.PostLoad();
 }
+
+//==============================================================================================
+
+/*static*/ FDataflowAttachmentFactory& FDataflowAttachmentFactory::Get()
+{
+	static FDataflowAttachmentFactory Instance;
+	return Instance;
+}
+
+void FDataflowAttachmentFactory::Register(FName ClassName, FCreateAttachmentFunction CreateFunction)
+{
+	CreateFunctionsByName.Add(ClassName, CreateFunction);
+}
+
+UDataflowAttachment* FDataflowAttachmentFactory::FDataflowAttachmentFactory::Create(UObject* Owner)
+{
+	const FName OwnerClassName = Owner->GetClass()->GetFName();
+	if (FCreateAttachmentFunction* Function = CreateFunctionsByName.Find(OwnerClassName))
+	{
+		if (Function->IsSet())
+		{
+			return (*Function)(Owner);
+		}
+	}
+	return nullptr;
+}
diff --git a/Engine/Source/Runtime/Experimental/Dataflow/Engine/Private/Dataflow/DataflowInstance.cpp b/Engine/Source/Runtime/Experimental/Dataflow/Engine/Private/Dataflow/DataflowInstance.cpp
--- a/Engine/Source/Runtime/Experimental/Dataflow/Engine/Private/Dataflow/DataflowInstance.cpp
+++ b/Engine/Source/Runtime/Experimental/Dataflow/Engine/Private/Dataflow/DataflowInstance.cpp
@@ -26,6 +26,31 @@ namespace UE::Dataflow::InstanceUtils
 		return nullptr;
 	}
 
+	void AddDataflowAttachmentIfNeeded(UObject* InAsset)
+	{
+		if (const IDataflowInstanceInterface* DataflowInstanceInterface = Cast<IDataflowInstanceInterface>(InAsset))
+		{
+			// no attachement needed
+			return;
+		}
+
+		if (IInterface_AssetUserData* UserDataInterface = Cast<IInterface_AssetUserData>(InAsset))
+		{
+			if (!UserDataInterface->HasAssetUserDataOfClass(UDataflowAttachment::StaticClass()))
+			{
+				if (UDataflowAttachment* Attachement = FDataflowAttachmentFactory::Get().Create(InAsset))
+				{
+					UserDataInterface->AddAssetUserData(Attachement);
+				}
+			}
+		}
+	}
+
+	UDataflowAttachment* GetDataflowAttachmentFromObject(const UObject* Obj)
+	{
+		return GetDataflowAttachment(Obj);
+	}
+
 	IDataflowInstanceInterface* GetDataflowInterfaceFromObject(UObject* Obj)
 	{
 		// is the owning asset itself implement the dataflow instance interface
@@ -119,6 +144,8 @@ namespace UE::Dataflow::InstanceUtils
 
 	bool SetDataflowAssetOnObject(UDataflow* DataflowAsset, UObject* Obj)
 	{
+		AddDataflowAttachmentIfNeeded(Obj);
+
 		// check if the object implements IDataflowGeneratedAssetInterface
 		if (IDataflowInstanceInterface* Interface = GetDataflowInterfaceFromObject(Obj))
 		{
diff --git a/Engine/Source/Runtime/Experimental/Dataflow/Engine/Public/Dataflow/DataflowAttachment.h b/Engine/Source/Runtime/Experimental/Dataflow/Engine/Public/Dataflow/DataflowAttachment.h
--- a/Engine/Source/Runtime/Experimental/Dataflow/Engine/Public/Dataflow/DataflowAttachment.h
+++ b/Engine/Source/Runtime/Experimental/Dataflow/Engine/Public/Dataflow/DataflowAttachment.h
@@ -2,6 +2,7 @@
 
 #pragma once
 
+#include "Containers/Map.h"
 #include "Engine/AssetUserData.h"
 #include "Dataflow/DataflowContent.h"
 #include "Dataflow/DataflowInstance.h"
@@ -51,4 +52,21 @@ protected:
 	////////////////////////////////////////////////////////////////////////////////////////////////////////////////
 };
 
+class UClass;
+
+struct FDataflowAttachmentFactory
+{
+public:
+	UE_API static FDataflowAttachmentFactory& Get();
+
+	using FCreateAttachmentFunction = TFunction<UDataflowAttachment*(UObject*)>;
+
+	UE_API void Register(FName ClassName, FCreateAttachmentFunction CreateFunction);
+
+	UE_API UDataflowAttachment* Create(UObject* Owner);
+
+private:
+	TMap<FName, FCreateAttachmentFunction> CreateFunctionsByName;
+};
+
 #undef UE_API
diff --git a/Engine/Source/Runtime/Experimental/Dataflow/Engine/Public/Dataflow/DataflowInstance.h b/Engine/Source/Runtime/Experimental/Dataflow/Engine/Public/Dataflow/DataflowInstance.h
--- a/Engine/Source/Runtime/Experimental/Dataflow/Engine/Public/Dataflow/DataflowInstance.h
+++ b/Engine/Source/Runtime/Experimental/Dataflow/Engine/Public/Dataflow/DataflowInstance.h
@@ -9,6 +9,7 @@
 #include "DataflowInstance.generated.h"
 
 class FAssetRegistryTagsContext;
+class UDataflowAttachment;
 
 namespace UE::Dataflow::Private
 {
@@ -343,6 +344,9 @@ namespace UE::Dataflow::InstanceUtils
 	/** Get the Dataflow instance interface from a UObject */
 	DATAFLOWENGINE_API IDataflowInstanceInterface* GetDataflowInterfaceFromObject(UObject* Obj);
 
+	/** Get the Dataflow attachment from a UObject */
+	DATAFLOWENGINE_API UDataflowAttachment* GetDataflowAttachmentFromObject(const UObject* Obj);
+
 	/** Get the Dataflow instance interface from a const UObject */
 	DATAFLOWENGINE_API const IDataflowInstanceInterface* GetDataflowInterfaceFromObject(const UObject* Obj);
 
