Move FGenericHash to OnDemandCore module.

#rnx

#Synced-CL 52465190

diff --git a/Engine/Source/Runtime/Experimental/IoStore/OnDemand/Private/GenericHash.cpp b/Engine/Source/Runtime/Experimental/IoStore/OnDemand/Private/GenericHash.cpp
deleted file mode 100644
diff --git a/Engine/Source/Runtime/Experimental/IoStore/OnDemand/Private/GenericHash.h b/Engine/Source/Runtime/Experimental/IoStore/OnDemand/Private/GenericHash.h
deleted file mode 100644
diff --git a/Engine/Source/Runtime/Experimental/IoStore/OnDemand/Private/OnDemandInstallCache.cpp b/Engine/Source/Runtime/Experimental/IoStore/OnDemand/Private/OnDemandInstallCache.cpp
--- a/Engine/Source/Runtime/Experimental/IoStore/OnDemand/Private/OnDemandInstallCache.cpp
+++ b/Engine/Source/Runtime/Experimental/IoStore/OnDemand/Private/OnDemandInstallCache.cpp
@@ -17,13 +17,13 @@
 #include "Async/SharedLock.h"
 #include "Async/AsyncFileHandle.h"
 #include "Containers/UnrealString.h"
-#include "GenericHash.h"
 #include "HAL/FileManager.h"
 #include "HAL/PlatformFile.h"
 #include "HAL/PlatformFileManager.h"
 #include "IO/IoContainerHeader.h"
 #include "IO/IoChunkId.h"
 #include "IO/IoChunkEncoding.h"
+#include "IO/GenericHash.h"
 #include "Logging/StructuredLog.h"
 #include "IO/OnDemandError.h"
 #include "Misc/DateTime.h"
diff --git a/Engine/Source/Runtime/Experimental/IoStore/OnDemandCore/Internal/IO/GenericHash.cpp b/Engine/Source/Runtime/Experimental/IoStore/OnDemandCore/Internal/IO/GenericHash.cpp
new file mode 100644
--- /dev/null
+++ b/Engine/Source/Runtime/Experimental/IoStore/OnDemandCore/Internal/IO/GenericHash.cpp
@@ -0,0 +1,33 @@
+// Copyright Epic Games, Inc. All Rights Reserved.
+
+#include "GenericHash.h"
+
+#include "Containers/ArrayView.h"
+#include "IO/IoHash.h"
+#include "Misc/StringBuilder.h"
+#include "String/BytesToHex.h"
+
+namespace UE::GenericHash
+{
+
+static_assert(sizeof(FIoHash) == 20);
+
+FStringBuilderBase& ToHex(FMemoryView Memory, FStringBuilderBase& Out)
+{
+	UE::String::BytesToHexLower(
+		MakeArrayView<const uint8>(
+			reinterpret_cast<const uint8*>(Memory.GetData()),
+			IntCastChecked<int32>(Memory.GetSize())),
+		Out);
+
+	return Out;
+}
+
+FString ToHex(FMemoryView Memory)
+{
+	TStringBuilder<256> Sb;
+	ToHex(Memory, Sb);
+	return FString(Sb);
+}
+
+} // namesapce UE::GenericHash
diff --git a/Engine/Source/Runtime/Experimental/IoStore/OnDemandCore/Internal/IO/GenericHash.h b/Engine/Source/Runtime/Experimental/IoStore/OnDemandCore/Internal/IO/GenericHash.h
new file mode 100644
--- /dev/null
+++ b/Engine/Source/Runtime/Experimental/IoStore/OnDemandCore/Internal/IO/GenericHash.h
@@ -0,0 +1,120 @@
+// Copyright Epic Games, Inc. All Rights Reserved.
+
+#pragma once
+
+#include "CoreTypes.h"
+#include "Containers/StringFwd.h"
+#include "Containers/UnrealString.h"
+#include "HAL/UnrealMemory.h"
+#include "Math/UnrealMathUtility.h"
+#include "Memory/MemoryView.h"
+#include "Serialization/Archive.h"
+
+#define UE_API IOSTOREONDEMANDCORE_API
+
+struct FIoHash;
+
+namespace UE
+{
+
+namespace GenericHash
+{
+
+UE_API FString				ToHex(FMemoryView Memory);
+UE_API FStringBuilderBase&	ToHex(FMemoryView Memory, FStringBuilderBase& Out);
+
+} // namespace UE::GenericHash
+
+template<SIZE_T BitCount>
+struct THash
+{
+	static constexpr SIZE_T Size	= BitCount >> 3;
+	using ByteArray					= uint8[Size];
+
+						THash() = default;
+	const uint8*		GetData() const { return Hash; }
+	SIZE_T				GetSize() const { return Size; }
+	FMemoryView			GetView() const { return FMemoryView(Hash, Size); }
+
+	static THash		From(FMemoryView Memory);
+	static THash		From(const uint8* Data, SIZE_T DataSize);
+	static THash		From(const FIoHash& IoHash);
+	static const		THash Zero;
+
+	friend inline bool operator==(const THash& A, const THash& B)
+	{
+		return FMemory::Memcmp(A.Hash, B.Hash, Size) == 0;
+	}
+
+	friend inline bool operator!=(const THash& A, const THash& B)
+	{
+		return FMemory::Memcmp(A.Hash, B.Hash, Size) != 0;
+	}
+
+	friend inline bool operator<(const THash& A, const THash& B)
+	{
+		return FMemory::Memcmp(A.Hash, B.Hash, Size) < 0;
+	}
+
+	friend inline uint32 GetTypeHash(const THash& H)
+	{
+		return *reinterpret_cast<const uint32*>(H.Hash);
+	}
+
+	friend inline FString LexToString(const THash& H)
+	{
+		return GenericHash::ToHex(H.GetView());
+	}
+
+	friend inline FStringBuilderBase& LexToString(const THash& H, FStringBuilderBase& Out)
+	{
+		return GenericHash::ToHex(H.GetView(), Out);
+	}
+
+	friend FArchive& operator<<(FArchive& Ar, THash& H)
+	{
+		Ar.Serialize(const_cast<uint8*>(H.GetData()), H.GetSize());
+		return Ar;
+	}
+
+private:
+	FMutableMemoryView GetMutableView() { return FMutableMemoryView(Hash, Size); }
+
+	alignas(uint32) ByteArray Hash{};
+};
+
+template<SIZE_T BitCount>
+const THash<BitCount> THash<BitCount>::Zero;
+
+template<SIZE_T BitCount>
+THash<BitCount> THash<BitCount>::From(FMemoryView Memory)
+{
+	THash<BitCount> Hash;
+	Hash.GetMutableView().CopyFrom(Memory.Left(THash<BitCount>::Size));
+	return Hash;
+}
+
+template<SIZE_T BitCount>
+THash<BitCount> THash<BitCount>::From(const uint8* Data, SIZE_T DataSize)
+{
+	return From(MakeMemoryView(Data, DataSize));
+}
+
+template<SIZE_T BitCount>
+THash<BitCount> THash<BitCount>::From(const FIoHash& IoHash)
+{
+	THash<BitCount> Hash;
+	FMemory::Memcpy(&Hash, &IoHash, FMath::Min(THash<BitCount>::Size, SIZE_T(20)));
+	return Hash;
+}
+
+using FHash32	= THash<32>;
+using FHash64	= THash<64>;
+using FHash96	= THash<96>;
+using FHash128	= THash<128>;
+using FHash160	= THash<160>;
+using FHash256	= THash<256>;
+
+} // namespace UE
+
+#undef UE_API
