[TraceBasedDebuggers] Fix race condition: synchronize trace channel enablement with PIE BeginPlay

When RewindDebugger auto-record is enabled, StartRecordingLocalSession triggers an async chain:
MessageBus -> FRuntimeModule::StartRecording -> Trace worker thread -> HandleTraceConnectionEstablished -> EnableRequiredTraceChannels
The main thread would continue to UWorld::BeginPlay before trace channels were enabled, causing early lifecycle events (OnBeginPlay, BeginSimulation) to be discarded.

Adds FEventRef to FRuntimeModule that is signaled after EnableRequiredTraceChannels completes.
FRewindDebugger::OnPIEStarted now waits (with timeout) for this event before returning control to the PIE startup flow.

#rnx
#Synced-CL 52551117

diff --git a/Engine/Plugins/Animation/GameplayInsights/Source/RewindDebugger/Private/RewindDebugger.cpp b/Engine/Plugins/Animation/GameplayInsights/Source/RewindDebugger/Private/RewindDebugger.cpp
--- a/Engine/Plugins/Animation/GameplayInsights/Source/RewindDebugger/Private/RewindDebugger.cpp
+++ b/Engine/Plugins/Animation/GameplayInsights/Source/RewindDebugger/Private/RewindDebugger.cpp
@@ -228,6 +228,18 @@ void FRewindDebugger::OnPIEStarted(bool bSimulating)
 	if (ShouldAutoRecordOnPIE())
 	{
 		StartRecordingLocalSession();
+
+#if WITH_TRACE_BASED_DEBUGGERS
+		// Wait for trace channels to be enabled before PIE continues to BeginPlay.
+		// StartRecordingLocalSession triggers an async chain:
+		//   FRuntimeModule::StartRecording -> Trace worker thread -> HandleTraceConnectionEstablished -> EnableRequiredTraceChannels
+		// Without this wait, early lifecycle trace events (OnBeginPlay, etc.) are lost.
+		if (::RewindDebugger::FRewindDebuggerRuntime* Runtime = ::RewindDebugger::FRewindDebuggerRuntime::Instance())
+		{
+			constexpr float MaxWaitTimeSeconds = 3.0f;
+			Runtime->WaitForTraceChannelsEnabled(MaxWaitTimeSeconds);
+		}
+#endif // WITH_TRACE_BASED_DEBUGGERS
 	}
 }
 
diff --git a/Engine/Source/Runtime/Experimental/TraceBasedDebuggers/Private/TraceBasedDebuggerRuntime.cpp b/Engine/Source/Runtime/Experimental/TraceBasedDebuggers/Private/TraceBasedDebuggerRuntime.cpp
--- a/Engine/Source/Runtime/Experimental/TraceBasedDebuggers/Private/TraceBasedDebuggerRuntime.cpp
+++ b/Engine/Source/Runtime/Experimental/TraceBasedDebuggers/Private/TraceBasedDebuggerRuntime.cpp
@@ -153,6 +153,11 @@ void FRuntimeModule::StartRecording(const TConstArrayView<FString> Args)
 void FRuntimeModule::StartRecording(const FStartRecordingCommandMessage& InRecordingStartCommand)
 {
 	using namespace UE::TraceBasedDebuggers;
+
+	// Reset the channels-enabled event so WaitForTraceChannelsEnabled will block
+	// until HandleTraceConnectionEstablished completes
+	TraceChannelsEnabledEvent->Reset();
+
 	if (!ensureAlwaysMsgf(!IsRecording(), TEXT("Received a recording start command while a recording is already active")))
 	{
 		UE_LOG_REF(LogCategory, Log, TEXT("[%hs] There is an active recording. Attempting to stop it..."), __func__)
@@ -321,6 +326,9 @@ void FRuntimeModule::StopRecording()
 	FTraceAuxiliary::OnTraceStopped.RemoveAll(this);
 	OnTraceStoppedDelegateHandle = FDelegateHandle();
 
+	// Reset the channels-enabled event since we are stopping
+	TraceChannelsEnabledEvent->Reset();
+
 	CurrentTransportMode = ETraceTransportMode::Invalid;
 
 	if (!IsRecording())
@@ -541,6 +549,10 @@ void FRuntimeModule::HandleTraceConnectionEstablished()
 
 	TraceConnectionDetailsUpdatedDelegate.Broadcast();
 
+	// Signal that trace channels are now enabled so any thread waiting
+	// in WaitForTraceChannelsEnabled can proceed
+	TraceChannelsEnabledEvent->Trigger();
+
 	// Add ticker to keep track of accumulated recording time
 	RecordingTimerHandle = FTSTicker::GetCoreTicker().AddTicker(FTickerDelegate::CreateRaw(this, &FRuntimeModule::RecordingTimerTick));
 }
@@ -598,6 +610,58 @@ bool FRuntimeModule::WaitForTraceSessionDisconnect()
 	return !FTraceAuxiliary::IsConnected();
 }
 
+bool FRuntimeModule::WaitForTraceChannelsEnabled(float MaxWaitTimeSeconds)
+{
+	// Fast path: channels already enabled
+	if (TraceChannelsEnabledEvent->Wait(0))
+	{
+		return true;
+	}
+
+	UE_LOG_REF(LogCategory, Log, TEXT("[%hs] Waiting for trace channels to be enabled..."), __func__);
+
+#if WITH_EDITOR
+	FScopedSlowTask WaitForChannelsSlowTask(MaxWaitTimeSeconds,
+		LOCTEXT("WaitForTraceChannelsMessage", "Waiting for trace channels to be enabled..."));
+	WaitForChannelsSlowTask.MakeDialog(/*bShowCancelButton*/true, /*bAllowInPIE*/true);
+#endif
+
+	float CurrentWaitTime = 0.0f;
+	while (CurrentWaitTime < MaxWaitTimeSeconds)
+	{
+		constexpr float WaitIntervalSeconds = 0.1f;
+		constexpr uint32 WaitIntervalMs = 100;
+
+		if (TraceChannelsEnabledEvent->Wait(WaitIntervalMs))
+		{
+			UE_LOG_REF(LogCategory, Log, TEXT("[%hs] Trace channels enabled after %.1f seconds"), __func__, CurrentWaitTime);
+			return true;
+		}
+
+		CurrentWaitTime += WaitIntervalSeconds;
+
+#if WITH_EDITOR
+		WaitForChannelsSlowTask.EnterProgressFrame(WaitIntervalSeconds);
+		if (WaitForChannelsSlowTask.ShouldCancel())
+		{
+			UE_LOG_REF(LogCategory, Log, TEXT("[%hs] Cancelled waiting for trace channels to be enabled per user request after %.1f seconds"), __func__, CurrentWaitTime);
+			return false;
+		}
+#endif
+	}
+
+	// There's a small race window where the event could have been triggered between the last Wait(WaitIntervalMs) timing out and the loop condition re-evaluating.
+	// Let's do one last instantaneous check in case the event fired just after our last poll.
+	// Without it, we could get a false negative where the channels ARE enabled but the function returns false because the timing of the last poll just barely missed the trigger.
+	const bool bResult = TraceChannelsEnabledEvent->Wait(0);
+	if (!bResult)
+	{
+		UE_LOG_REF(LogCategory, Warning, TEXT("[%hs] Timed out waiting for trace channels to be enabled after %.1f seconds"), __func__, MaxWaitTimeSeconds);
+	}
+
+	return bResult;
+}
+
 void FRuntimeModule::SaveAndDisabledCurrentEnabledTraceChannels()
 {
 	// Until we support allowing other channels, indicate in the logs that we are disabling everything else
diff --git a/Engine/Source/Runtime/Experimental/TraceBasedDebuggers/Public/TraceBasedDebuggerRuntime.h b/Engine/Source/Runtime/Experimental/TraceBasedDebuggers/Public/TraceBasedDebuggerRuntime.h
--- a/Engine/Source/Runtime/Experimental/TraceBasedDebuggers/Public/TraceBasedDebuggerRuntime.h
+++ b/Engine/Source/Runtime/Experimental/TraceBasedDebuggers/Public/TraceBasedDebuggerRuntime.h
@@ -4,9 +4,10 @@
 
 #if WITH_TRACE_BASED_DEBUGGERS
 #include "Containers/Ticker.h"
+#include "HAL/Event.h"
+#include "Misc/Guid.h"
 #include "ProfilingDebugging/TraceAuxiliary.h"
 #include "TraceDataRelayTransport.h"
-#include "Misc/Guid.h"
 
 #define UE_API TRACEBASEDDEBUGGERS_API
 
@@ -103,6 +104,14 @@ public:
 	/** Returns the accumulated current send buffer size used by the relay transport system - This should be used for logging or UI feedback only */
 	UE_API int64 GetBufferedDataBytesSize() const;
 
+	/**
+	 * Blocks the calling thread until trace channels have been enabled after a recording has started, or until the timeout expires.
+	 * Intended for use on the game thread to synchronize PIE startup with asynchronous trace channel enablement.
+	 * @param MaxWaitTimeSeconds Maximum time to wait in seconds before giving up.
+	 * @return true if channels were enabled within the timeout, false otherwise.
+	 */
+	UE_API bool WaitForTraceChannelsEnabled(float MaxWaitTimeSeconds);
+
 	FSimpleMulticastDelegate& OnTraceConnectionDetailsUpdated()
 	{
 		return TraceConnectionDetailsUpdatedDelegate;
@@ -240,6 +249,9 @@ private:
 
 	ETraceTransportMode CurrentTransportMode = ETraceTransportMode::Invalid;
 	bool bRequestedStop = false;
+
+	/** Event signaled when HandleTraceConnectionEstablished completes and trace channels are enabled */
+	FEventRef TraceChannelsEnabledEvent{EEventMode::ManualReset};
 };
 
 } // UE::TraceBasedDebuggers
