Add initial support for multirange http requests for IAX

#rnx

- At the moment we will only make use of the multi range request feature if iax.HttpIoDispatcherMaxRangeKiB is set to -1 so that the higher level code makes no attempt to merge requests.
- The code will work with our CDNs at the moment although there is additional validation we may wish to add in future work against different servers that might return different things.
- We currently do not have any fallbacks if a CDN does not support the feature, that will come in a future work item.
- FMultiEndpointHttpClient::Get now takes multiple ranges in the form of an array. If the array is empty then the entire body will be requested. If one request we make a ranged request just for that range. If more than one request is made we first sort them into sequential order and attempt to merge and concurrent or overlapping requests. This allows us to send as few requests as possible and prevents inconsistent behaviour from different CDNs. If all the requests are merged then we will just make a ranged request, if multiple requests remain then we make a multiranged request.
- We not do the logging at the end of CompleteRequest as we miight encounter errors when parsing the body for a multiranged request and will want that recorded.
- We now wrap the returned data in a new structrure FResponseBody. This will allow each request to ask for its range or the entire body and be served the correct data without needing to know how the CDN returned the data to use.

### Misc
- JoinIoBuffers is no longer declared in the header file as it is only used by our http client.
- Renamed GAllowPartialContentRequests to GAllowRangeRequests as it now controls if we make any sort of ranged request via our http requests.

#Synced-CL 52050999

diff --git a/Engine/Source/Runtime/Experimental/IoStore/OnDemand/Private/OnDemandHttpClient.cpp b/Engine/Source/Runtime/Experimental/IoStore/OnDemand/Private/OnDemandHttpClient.cpp
--- a/Engine/Source/Runtime/Experimental/IoStore/OnDemand/Private/OnDemandHttpClient.cpp
+++ b/Engine/Source/Runtime/Experimental/IoStore/OnDemand/Private/OnDemandHttpClient.cpp
@@ -33,10 +33,10 @@ static FAutoConsoleVariableRef CVar_IaxWarnOnFailure(
  * FIoOffsetAndLength is provided with the request. When disabled we will always request the full chunk and then use the FIoOffsetAndLength to
  * return a partial slice of the data to the caller.
  */
-bool GAllowPartialContentRequests = true;
+bool GAllowRangeRequests = true;
 static FAutoConsoleVariableRef CVar_IaxAllowPartialContentRequests(
 	TEXT("iax.HttpAllowPartialContentRequests"),
-	GAllowPartialContentRequests,
+	GAllowRangeRequests,
 	TEXT("Enable/disable the use of partial content http requests")
 );
 
@@ -52,7 +52,7 @@ static FAutoConsoleVariableRef CVar_IaxUseSystemCaStore(
 
 ///////////////////////////////////////////////////////////////////////////////
 
-FIoBuffer JoinIoBuffers(const FIoBuffer& LHS, const FIoBuffer& RHS)
+static FIoBuffer JoinIoBuffers(const FIoBuffer& LHS, const FIoBuffer& RHS)
 {
 	const uint64 TotalSize = LHS.GetSize() + RHS.GetSize();
 	FIoBuffer Output(TotalSize);
@@ -82,6 +82,274 @@ static int8 TrackCdnCacheStats(const HTTP::FResponse& Response)
 	return Result;
 }
 
+/**
+ * Attempts to parse the multirange boundary marker from a http response header.
+ * This marker will be found on the 'Content-Type:' line if a multirange request
+ * was made.
+ * If no boundary marker can be found we can just return an empty string.
+ */
+static FAnsiString ParseMultiRangeBoundary(UE::IoStore::HTTP::FResponse& Response)
+{
+	FAnsiStringView ContentType = Response.GetHeader(ANSITEXTVIEW("content-type"));
+	if (ContentType.IsEmpty())
+	{
+		return FAnsiString();
+	}
+
+	constexpr FAnsiStringView MultipartType = ANSITEXTVIEW("multipart/byteranges;");
+	if (!ContentType.StartsWith(MultipartType))
+	{
+		return FAnsiString();
+	}
+
+	// Search for the first instance of 'boundary=' after 'multipart/byteranges;'
+	constexpr FAnsiStringView BoundaryKey = ANSITEXTVIEW("boundary=");
+	int32 StartIndex = ContentType.Find(BoundaryKey, MultipartType.Len(), ESearchCase::IgnoreCase);
+	if (StartIndex == INDEX_NONE)
+	{
+		return FAnsiString();
+	}
+
+	StartIndex += BoundaryKey.Len(); // Advance to the start of the value
+
+	const int32 EndIndex = ContentType.Find(";", StartIndex);
+	const int32 NumChars = (EndIndex != INDEX_NONE ? EndIndex : ContentType.Len()) - StartIndex;
+
+	FAnsiStringView BoundaryValue = ContentType.SubStr(StartIndex, NumChars);
+	BoundaryValue.TrimStartAndEndInline();
+
+	// When the boundary is used it will always be prefixed with '--' so we can add that prefix now.
+	constexpr FAnsiStringView BoundaryPrefix = ANSITEXTVIEW("--");
+
+	FAnsiString Output;
+	Output.Reserve(BoundaryPrefix.Len() + BoundaryValue.Len());
+	Output.Append(BoundaryPrefix);
+	Output.Append(BoundaryValue);
+
+	return Output;
+}
+
+FAnsiStringView SkipToNextLine(FAnsiStringView View, int32 Offset = 0)
+{
+	constexpr FAnsiStringView CRLF = ANSITEXTVIEW("\r\n");
+
+	int32 LineEndingIndex = View.Find(CRLF, Offset, ESearchCase::IgnoreCase);
+	if (LineEndingIndex != INDEX_NONE)
+	{
+		return View.RightChop(LineEndingIndex + 2);
+	}
+	else
+	{
+		return FAnsiStringView();
+	}
+}
+
+FAnsiStringView ExtractTextLine(FAnsiStringView View)
+{
+	constexpr FAnsiStringView CRLF = ANSITEXTVIEW("\r\n");
+
+	int32 LineEndingIndex = View.Find(CRLF, 0, ESearchCase::IgnoreCase);
+	if (LineEndingIndex != INDEX_NONE)
+	{
+		return View.Left(LineEndingIndex);
+	}
+	else
+	{
+		return View;
+	}
+}
+
+FAnsiStringView ExtractData(FAnsiStringView View, FAnsiStringView Boundary, TPair<uint32, uint32> ByteRange)
+{
+	int32 DataLength = (int32)(ByteRange.Value - ByteRange.Key) + 1; // TODO -  Safer casting/truncation?
+	FAnsiStringView Data = View.Left(DataLength);
+
+	// TODO - Validate?
+
+	return Data;
+}
+
+// TODO - LexFromString (need FAnsiStringView support)
+void ExtractFromString(uint32& OutValue, const FAnsiStringView& InString)
+{
+	TAnsiStringBuilder<64> Builder;
+	LexFromString(OutValue, *(Builder << InString));
+}
+
+FIoStatus ParseContentRange(FAnsiStringView Line, TPair<uint32,uint32>& ByteRange)
+{
+	// Content-Range: bytes 1000-1023/
+
+	constexpr FAnsiStringView Bytes = ANSITEXTVIEW("bytes ");
+
+	int32 BytesIndex = Line.Find(Bytes);
+	if (BytesIndex == INDEX_NONE)
+	{
+		return FIoStatus(EIoErrorCode::HttpReadError, TEXT("Multirange request: Unable to bytes for 'Content-Range'"));
+	}
+
+	FAnsiStringView DataView = Line.RightChop(BytesIndex + Bytes.Len());
+	DataView.TrimEndInline();
+
+	ExtractFromString(ByteRange.Key, DataView);
+
+	int32 SecondByteOffset = 0;
+	if (!DataView.FindChar('-', SecondByteOffset))
+	{
+		return FIoStatus(EIoErrorCode::HttpReadError, TEXT("Multirange request: Malformed byte range for 'Content-Range'"));
+	}
+
+	DataView.RightChopInline(SecondByteOffset + 1);
+
+	ExtractFromString(ByteRange.Value, DataView);
+
+	// TODO - Better validation that we parsed byte range values
+
+	return EIoErrorCode::Ok;
+}
+
+static FIoStatus ParseHeader(FAnsiStringView& BodyView, FAnsiStringView Boundary, TPair<uint32, uint32>& OutByteRange)
+{
+	constexpr FAnsiStringView ContentRange = ANSITEXTVIEW("Content-Range");
+
+	bool bReadRange = false;
+	while (!BodyView.IsEmpty())
+	{
+		FAnsiStringView Line = ExtractTextLine(BodyView);
+		BodyView = SkipToNextLine(BodyView, Line.Len());
+
+		// TODO: Clean up the parsing, should be less logic to determine that we read the correct format.
+		if (Line.StartsWith(ContentRange))
+		{
+			if (FIoStatus Status = ParseContentRange(Line.RightChop(ContentRange.Len()), OutByteRange); !Status.IsOk())
+			{
+				return Status;
+			}
+
+			bReadRange = true;
+		}
+		else if (Line.IsEmpty())
+		{
+			if (bReadRange)
+			{
+				return EIoErrorCode::Ok;
+			}
+			else
+			{
+				return FIoStatus(EIoErrorCode::HttpReadError, TEXT("Multirange request:Ranged header did not contain a 'Content-Range'"));
+			}
+			
+		}
+		else if (Line.StartsWith(Boundary))
+		{
+			return FIoStatus(EIoErrorCode::HttpReadError, TEXT("Multirange request: End of ranged header was unexpectedly reached"));
+		}
+	}
+
+	return FIoStatus(EIoErrorCode::HttpReadError, TEXT("Multirange request: Ranged header was truncated"));
+}
+
+FResponseBody::FResponseBody(FIoBuffer&& InBody, const FIoOffsetAndLength& Range)
+	: Body(MoveTemp(InBody))
+{
+	// TODO: Ignore range if whole body?
+	// TODO: Improve casting/truncation.
+	DataRanges.Add(FResponseBody::FRange
+	{
+		.BodyOffset = 0,
+		.StartByte = (uint32)Range.GetOffset(),
+		.EndByte = (uint32)(Range.GetOffset() + Range.GetLength() - 1),
+	});
+}
+
+FIoStatus FResponseBody::Parse(FIoBuffer&& InBody, FAnsiStringView Boundary)
+{
+	constexpr FAnsiStringView DelimEnd = ANSITEXTVIEW("--");
+
+	// Note that in theory the body data can contain '\0' if the content type is binary, but the string view
+	// code will ignore that if we provide the full size and makes it easier to parse the headers.
+	FAnsiStringView BodyView((const char*)InBody.GetData(), (int32)InBody.GetSize()); // TODO - Safer cast
+
+	while (!BodyView.IsEmpty())
+	{
+		FAnsiStringView Line = ExtractTextLine(BodyView);
+		BodyView = SkipToNextLine(BodyView, Line.Len());
+
+		if (Line.StartsWith(Boundary))
+		{
+			FAnsiStringView PostBoundary = Line.SubStr(Boundary.Len(), 2);
+			if (PostBoundary == DelimEnd)
+			{
+				Body = MoveTemp(InBody);
+				return EIoErrorCode::Ok;
+			}
+
+			TPair<uint32, uint32> ByteRange;
+			if (FIoStatus Status = ParseHeader(BodyView, Boundary, ByteRange); !Status.IsOk())
+			{
+				Reset();
+				return Status;
+			}
+
+			// Now parse the data
+			FAnsiStringView DataView = ExtractData(BodyView, Boundary, ByteRange);
+			BodyView = SkipToNextLine(BodyView, DataView.Len());
+
+			// TODO: Safer cast
+			const uint32 DataOffset = (uint32)((uint8*)DataView.GetData() - InBody.GetData());
+
+			// TODO: Validate the boundary and data with ByteRange?
+			DataRanges.Add(FResponseBody::FRange
+				{
+					.BodyOffset = DataOffset,
+					.StartByte = ByteRange.Key,
+					.EndByte = ByteRange.Value,
+				});
+		}
+		else if (!Line.IsEmpty())
+		{
+			Reset();
+			return FIoStatus(EIoErrorCode::HttpReadError, TEXT("Multirange request: Unexpected data found in the body"));
+		}
+	}
+
+	Reset();
+	return FIoStatus(EIoErrorCode::HttpReadError, TEXT("Multirange request: Truncated data"));
+}
+
+FIoBuffer FResponseBody::GetData(const FIoOffsetAndLength& Range) const
+{
+	if (!Range.IsValid())
+	{
+		return Body;
+	}
+	else if (DataRanges.IsEmpty())
+	{
+		return FIoBuffer(Body.GetView().Mid(Range.GetOffset(), Range.GetLength()), Body);
+	}
+	else
+	{
+		for (const FRange& DataRange : DataRanges)
+		{
+			if (Range.GetOffset() >= DataRange.StartByte && Range.GetOffset() <= DataRange.EndByte) // Does the requested range start within the range
+			{
+				const uint64 EndByte = (Range.GetOffset() + Range.GetLength()) - 1;
+				if (EndByte >= DataRange.StartByte && EndByte <= DataRange.EndByte) // Does the requested range end within the range
+				{
+					// TODO - Worth returning DataRange.DataView if exact match?
+					const uint64 Offset = Range.GetOffset() - DataRange.StartByte;
+					return FIoBuffer(Body.GetView().Mid(DataRange.BodyOffset + Offset, Range.GetLength()), Body);
+				}
+			}
+		}
+	}
+
+	checkNoEntry();
+
+	// TODO - Error
+	return FIoBuffer();
+}
+
 static const TCHAR* CDNCacheStatusToString(int8 Status)
 {
 #if UE_TRACK_CDN_HIT_STATUS
@@ -91,6 +359,54 @@ static const TCHAR* CDNCacheStatusToString(int8 Status)
 #endif //UE_TRACK_CDN_HIT_STATUS
 }
 
+// Assumes sorted order
+bool MergeOverlappingRanges(FIoOffsetAndLength& MergeDst, const FIoOffsetAndLength& PotentialOverlap)
+{
+	const uint64 DstLastByte = MergeDst.GetOffset() + MergeDst.GetLength();
+	if (PotentialOverlap.GetOffset() <= DstLastByte)
+	{
+		const uint64 PotentialLastByte = PotentialOverlap.GetOffset() + PotentialOverlap.GetLength();
+		if (DstLastByte > PotentialLastByte)
+		{
+			return true;
+		}
+
+		const uint64 MergedLength = PotentialLastByte - MergeDst.GetOffset();
+		MergeDst.SetLength(MergedLength);
+
+		return true;
+	}
+	else
+	{
+		return false;
+	}
+}
+
+TArray<FIoOffsetAndLength> SortAndMergeRanges(TArray<FIoOffsetAndLength>& UnsortedRanges)
+{
+	if (UnsortedRanges.IsEmpty())
+	{
+		return TArray<FIoOffsetAndLength>();
+	}
+
+	UnsortedRanges.Sort([](const FIoOffsetAndLength& LHS, const FIoOffsetAndLength& RHS)
+	{
+		return LHS.GetOffset() < RHS.GetOffset();
+	});
+
+	TArray<FIoOffsetAndLength> SortedRanges;
+	SortedRanges.Add(UnsortedRanges[0]);
+
+	for (int32 Index = 1; Index < UnsortedRanges.Num(); ++Index)
+	{
+		if (!MergeOverlappingRanges(SortedRanges.Last(), UnsortedRanges[Index]))
+		{
+			SortedRanges.Add(UnsortedRanges[Index]);
+		}
+	}
+	return SortedRanges;
+}
+
 FIoStatus LoadDefaultHttpCertificates(bool& bWasLoaded)
 {
 	using namespace UE::IoStore::HTTP;
@@ -185,6 +501,46 @@ FMultiEndpointHttpClient::FHttpTicketId GenerateTicketId()
 
 ///////////////////////////////////////////////////////////////////////////////
 
+uint64 FMultiEndpointHttpClient::FRequest::GetTotalRangeLength() const
+{
+	uint64 TotalLength = 0;
+
+	for (const FIoOffsetAndLength& Range : Ranges)
+	{
+		if (Range.IsValid())
+		{
+			TotalLength += Range.GetLength();
+		}
+	}
+
+	return TotalLength;
+}
+
+bool FMultiEndpointHttpClient::FRequest::BuildRangeHeaderValue(FAnsiStringBuilderBase& OutString) const
+{
+	OutString = "bytes=";
+
+	bool bHasAddedRange = false;
+
+	for (const FIoOffsetAndLength& Range : Ranges)
+	{
+		if (Range.IsValid())
+		{
+			if (bHasAddedRange)
+			{
+				OutString << ",";
+			}
+
+			OutString << Range.GetOffset() << "-" << (Range.GetOffset() + Range.GetLength() - 1);
+			bHasAddedRange = true;
+		}
+	}
+
+	return bHasAddedRange;
+}
+
+///////////////////////////////////////////////////////////////////////////////
+
 FMultiEndpointHttpClient::FMultiEndpointHttpClient(const FMultiEndpointHttpClientConfig& InConfig)
 	: Config(InConfig)
 {
@@ -255,7 +611,7 @@ TIoStatusOr<FMultiEndpointHttpClientResponse> FMultiEndpointHttpClient::Get(FAns
 		{
 			FMultiEndpointHttpClientResponse Response
 			{
-				.Body = MoveTemp(Body),
+				.Body = FResponseBody(MoveTemp(Body)),
 				.Sample = PerfSample,
 				.StatusCode = StatusCode,
 				.RetryCount = Attempt,
@@ -275,31 +631,32 @@ TIoStatusOr<FMultiEndpointHttpClientResponse> FMultiEndpointHttpClient::Get(FAns
 
 FMultiEndpointHttpClient::FHttpTicketId FMultiEndpointHttpClient::Get(const FOnDemandHostGroup& HostGroup, FAnsiStringView RelativeUrl, FOnHttpResponse&& OnResponse)
 {
-	FIoOffsetAndLength DefaultChunkRange(0,0);
-	return Get(HostGroup, RelativeUrl, DefaultChunkRange, MoveTemp(OnResponse));
+	return Get(HostGroup, RelativeUrl, TArray<FIoOffsetAndLength>(), MoveTemp(OnResponse));
 }
 
-FMultiEndpointHttpClient::FHttpTicketId FMultiEndpointHttpClient::Get(const FOnDemandHostGroup& HostGroup, FAnsiStringView RelativeUrl, const FIoOffsetAndLength& ChunkRange, FOnHttpResponse&& OnResponse)
+FMultiEndpointHttpClient::FHttpTicketId FMultiEndpointHttpClient::Get(const FOnDemandHostGroup& HostGroup, FAnsiStringView RelativeUrl, TArray<FIoOffsetAndLength>&& ChunkRanges, FOnHttpResponse&& OnResponse)
 {
-	return Get(HostGroup, RelativeUrl, ChunkRange, TArray<FAnsiString>(), EMultiEndpointRequestFlags::None, MoveTemp(OnResponse));
+	return Get(HostGroup, RelativeUrl, MoveTemp(ChunkRanges), TArray<FAnsiString>(), EMultiEndpointRequestFlags::None, MoveTemp(OnResponse));
 }
 
 FMultiEndpointHttpClient::FHttpTicketId FMultiEndpointHttpClient::Get(
 	const FOnDemandHostGroup& HostGroup,
 	FAnsiStringView RelativeUrl,
-	const FIoOffsetAndLength& ChunkRange,
+	TArray<FIoOffsetAndLength>&& ChunkRanges,
 	TArray<FAnsiString>&& Headers,
 	EMultiEndpointRequestFlags Flags,
 	FOnHttpResponse&& OnResponse)
 {
 	FConnection& Connection = GetConnection(HostGroup);
 
+	TArray<FIoOffsetAndLength> SortedRanged = SortAndMergeRanges(ChunkRanges);
+	
 	FHttpTicketId TicketId = IssueRequest(FRequest
 	{
 		.OnResponse		= MoveTemp(OnResponse),
 		.RequestHeaders	= MoveTemp(Headers),
 		.RelativeUrl	= FAnsiString(RelativeUrl),
-		.Range			= ChunkRange,
+		.Ranges			= MoveTemp(SortedRanged),
 		.Connection		= Connection,
 		.StartTime		= FPlatformTime::Cycles64(),
 		.Host			= Connection.CurrentHost,
@@ -430,7 +787,7 @@ FMultiEndpointHttpClient::FHttpTicketId FMultiEndpointHttpClient::IssueRequest(F
 
 	FEventLoop::FRequestParams RequestParams = FEventLoop::FRequestParams
 	{
-		.ContentSizeEst = uint32(Request.Range.GetLength()),
+		.ContentSizeEst	= (uint32)Request.GetTotalRangeLength(), // TODO - Safer cast?
 		.bAutoRedirect	= Config.Redirects == EHttpRedirects::Follow,
 		.bAllowChunked	= Config.bAllowChunkedTransfer
 	};
@@ -450,10 +807,14 @@ FMultiEndpointHttpClient::FHttpTicketId FMultiEndpointHttpClient::IssueRequest(F
 		}
 	}
 
-	if (GAllowPartialContentRequests && Request.Range.IsValid())
+	if (GAllowRangeRequests)
 	{
-		HttpRequest.Header(ANSITEXTVIEW("range"),
-			WriteToAnsiString<64>(ANSITEXTVIEW("bytes="), Request.Range.GetOffset(), ANSITEXTVIEW("-"), Request.Range.GetOffset() + Request.Range.GetLength() - 1));
+		TAnsiStringBuilder<128> HeaderValue;
+
+		if (Request.BuildRangeHeaderValue(HeaderValue))
+		{
+			HttpRequest.Header(ANSITEXTVIEW("range"), *HeaderValue);
+		}
 	}
 
 	// If the request is new then we need to assign an identifier
@@ -477,6 +838,11 @@ FMultiEndpointHttpClient::FHttpTicketId FMultiEndpointHttpClient::IssueRequest(F
 				FResponse& HttpResponse		= TicketStatus.GetResponse();
 				Request.StatusCode			= HttpResponse.GetStatusCode();
 				Request.bIsChunkedTransfer	= HttpResponse.GetContentLength() == -1;
+				
+				if(Request.Ranges.Num() > 1)
+				{
+					Request.MultiRangeBoundary = ParseMultiRangeBoundary(HttpResponse); // TODO - Only do this is we sent multi-range?
+				}
 
 				if (EnumHasAnyFlags(Request.Flags, EMultiEndpointRequestFlags::ResponseHeaders))
 				{
@@ -590,8 +956,6 @@ void FMultiEndpointHttpClient::CompleteRequest(FRequest&& Request, const UE::IoS
 	check(Request.SinkCounter == Request.RetryCount);
 #endif // DO_CHECK
 
-	Log(Request, TicketStatus);
-
 	FMultiEndpointHttpClientResponse Response
 	{
 		.Headers				= MoveTemp(Request.ResponseHeaders),
@@ -603,27 +967,33 @@ void FMultiEndpointHttpClient::CompleteRequest(FRequest&& Request, const UE::IoS
 
 	if (IsHttpStatusOk(Request.StatusCode) || Config.bResponseBodyOnError)
 	{
-		if (Request.StatusCode != 206 && Request.Range.IsValid())
+		if (Request.StatusCode == 206)
 		{
-			// A partial read was requested but the server sent us the full response body so we need to cut down
-			// the data that we return. Note that the partial buffer we return will retain a ref count on the full
-			// sized buffer until it is also destroyed, so we don't need to do any allocations/copies at this point.
-			FMemoryView PartialRange = Request.Body.GetView().Mid(Request.Range.GetOffset(), Request.Range.GetLength());
-			Response.Body = FIoBuffer(PartialRange, Request.Body);
+			if (Request.MultiRangeBoundary.Len() > 0)
+			{
+				if (FIoStatus Status = Response.Body.Parse(MoveTemp(Request.Body), Request.MultiRangeBoundary); !Status.IsOk())
+				{
+					Response.StatusCode = 0;
+					Response.Reason = Status.ToString();
+				}
 
-			// Try to limit how frequently we log warnings about this to avoid spam
-			const uint64 LogSpamLimit = 50; 
-			static uint64 LogCounter = 0;
-			UE_CLOGF(LogCounter++ % LogSpamLimit == 0, LogIoStoreOnDemand, Warning, "A partial content request failed and the full data blob was downloaded, check if the CDN supports this");
+				//((char*)Response.Body.GetData())[104] = '\0'; // <- Test that the string view parsing code is safe
+			}
+			else
+			{
+				// TODO - Validate the range returned?
+				check(!Request.Ranges.IsEmpty());
+				Response.Body = FResponseBody(MoveTemp(Request.Body), Request.Ranges[0]);
+			}
 		}
 		else
 		{
-			Response.Body = MoveTemp(Request.Body);
+			Response.Body = FResponseBody(MoveTemp(Request.Body));
 		}
 	}
 	else if (Config.bResponseBodyOnError)
 	{
-		Response.Body = MoveTemp(Request.Body);
+		Response.Body = FResponseBody(MoveTemp(Request.Body));
 	}
 
 	if (TicketStatus.GetId() == FTicketStatus::EId::Content)
@@ -633,7 +1003,7 @@ void FMultiEndpointHttpClient::CompleteRequest(FRequest&& Request, const UE::IoS
 	
 	if (TicketStatus.GetId() == FTicketStatus::EId::Error)
 	{
-		Response.Reason = TicketStatus.GetError().Reason;
+		Response.Reason = WriteToString<128>(TicketStatus.GetError().Reason, TEXT(" ("), TicketStatus.GetError().Code, TEXT(")"));
 	}
 	else if (TicketStatus.GetId() == FTicketStatus::EId::Cancelled)
 	{
@@ -650,6 +1020,8 @@ void FMultiEndpointHttpClient::CompleteRequest(FRequest&& Request, const UE::IoS
 		TicketLookupMap.Remove(Request.TicketId);
 	}
 
+	Log(Response, Request);
+
 	FOnHttpResponse OnResponse = MoveTemp(Request.OnResponse);
 	OnResponse(MoveTemp(Response));
 }
@@ -859,7 +1231,7 @@ uint32 FMultiEndpointHttpClient::GetRetryLimitForRequest(const FRequest& Request
 	return Config.MaxRetryCount == -1 ? Request.Connection.HostGroup.Hosts().Num() : Config.MaxRetryCount;
 }
 
-void FMultiEndpointHttpClient::Log(const FRequest& Request, const HTTP::FTicketStatus& TicketStatus) const
+void FMultiEndpointHttpClient::Log(const FMultiEndpointHttpClientResponse& Response, const FRequest& Request) const
 {
 #if !NO_LOGGING
 
@@ -867,7 +1239,7 @@ void FMultiEndpointHttpClient::Log(const FRequest& Request, const HTTP::FTicketS
 
 	ELogVerbosity::Type Verbosity = ELogVerbosity::VeryVerbose;
 	// We count  a request as failed if it's status code is bad AND the request was not canceled
-	if (GIaxWarnOnFailure && !IsHttpStatusOk(Request.StatusCode) && TicketStatus.GetId() != FTicketStatus::EId::Cancelled)
+	if (GIaxWarnOnFailure && !IsHttpStatusOk(Response.StatusCode) && !Response.bCanceled)
 	{
 		Verbosity = ELogVerbosity::Type::Warning;
 	}
@@ -877,26 +1249,22 @@ void FMultiEndpointHttpClient::Log(const FRequest& Request, const HTTP::FTicketS
 		return;
 	}
 
-	const uint64 DurationMs = (uint64)FPlatformTime::ToMilliseconds64(FPlatformTime::Cycles64() - Request.StartTime);
-	const uint64 Size = Request.Body.GetSize() >> 10;
+	const uint32 DurationMs = Response.Sample.GetTotalMs();
+	const uint64 Size = Response.Body.GetSize() >> 10;
 
 	TStringBuilder<256> Reason;
-	if (TicketStatus.GetId() == FTicketStatus::EId::Error)
+	if (!Response.Reason.IsEmpty())
 	{
-		Reason << TEXT(" ") << TicketStatus.GetError().Reason << TEXT(" (") << TicketStatus.GetError().Code << TEXT(")");
+		Reason << TEXT(" ") << Response.Reason;
 	}
-	else if (TicketStatus.GetId() == FTicketStatus::EId::Cancelled)
+	else if (Response.bCanceled)
 	{
 		Reason << TEXT(" Canceled");
 	}
-	else if (!Request.ResponseMessage.IsEmpty())
-	{
-		Reason << TEXT(" ") << Request.ResponseMessage;
-	}
 
 	FMsg::Logf(__FILE__, __LINE__, LogHttpIoDispatcher.GetCategoryName(), Verbosity,
-		TEXT("http-%3u: %5" UINT64_FMT "ms %5" UINT64_FMT "KiB [%4s] %s (Attempt %u/%u)%s"),
-		Request.StatusCode,
+		TEXT("http-%3u: %5dms %5" UINT64_FMT "KiB [%4s] %s (Attempt %u/%u)%s"),
+		Response.StatusCode,
 		DurationMs,
 		Size,
 		CDNCacheStatusToString(Request.CDNCacheStatus),
diff --git a/Engine/Source/Runtime/Experimental/IoStore/OnDemand/Private/OnDemandHttpClient.h b/Engine/Source/Runtime/Experimental/IoStore/OnDemand/Private/OnDemandHttpClient.h
--- a/Engine/Source/Runtime/Experimental/IoStore/OnDemand/Private/OnDemandHttpClient.h
+++ b/Engine/Source/Runtime/Experimental/IoStore/OnDemand/Private/OnDemandHttpClient.h
@@ -29,9 +29,6 @@ namespace UE::IoStore
 
 class FOnDemandHttpStats;
 
-/** Utility function to join two FIoBuffers together and return the combined result as a new FIoBuffer */
-FIoBuffer JoinIoBuffers(const FIoBuffer& LHS, const FIoBuffer& RHS);
-
 FIoStatus LoadDefaultHttpCertificates(bool& bWasLoaded);
 
 /** Returns if a Http status code is okay (2XX) or not */
@@ -77,6 +74,72 @@ enum class EMultiEndpointRequestFlags : uint8
 };
 ENUM_CLASS_FLAGS(EMultiEndpointRequestFlags);
 
+struct FResponseBody
+{
+public:
+	FResponseBody() = default;
+
+	FResponseBody(const FIoBuffer& InBody)
+		: Body(InBody)
+	{
+	}
+
+	FResponseBody(FIoBuffer&& InBody)
+		: Body(MoveTemp(InBody))
+	{
+	}
+
+	FResponseBody(FIoBuffer&& InBody, const FIoOffsetAndLength& Range);
+
+	FIoStatus Parse(FIoBuffer&& Body, FAnsiStringView Boundary);
+
+	/** 
+	 * Returns the total length of the body that was downloaded.
+	 * Note that if the request was multi-part range then the body will include the
+	 * header for each range that was returned, so this size may be larger than the
+	 * sum of the requested ranges. If the whole body or a single range was requested
+	 * then 
+	 */
+	uint64 GetSize() const
+	{
+		return Body.GetSize();
+	}
+
+	FIoBuffer GetBody() const
+	{
+		return Body;
+	}
+
+	FIoBuffer GetData(const FIoOffsetAndLength& Range) const;
+
+private:
+
+	void Reset()
+	{
+		Body = FIoBuffer();
+		DataRanges.Empty();
+	}
+
+	FIoBuffer					Body;		// The full body from the response
+
+	/**
+	 * 
+	 *
+	 */
+	struct FRange
+	{
+		/* Offset into the Body buffer to where this range begins */
+		uint32 BodyOffset	= MAX_uint32;
+
+		/** */
+		uint32 StartByte	= MAX_uint32;
+		/** */
+		uint32 EndByte		= 0;
+	};
+
+	TArray<FRange> DataRanges; // Views into the body of the originally requested ranges
+};
+
 struct FMultiEndpointHttpClientResponse
 {
 	bool IsOk() const
@@ -89,7 +152,8 @@ struct FMultiEndpointHttpClientResponse
 		return bCanceled;
 	}
 
-	FIoBuffer					Body;
+	FResponseBody				Body;
+
 	TArray<FAnsiString>			Headers;
 	FString						Reason;
 	HTTP::FTicketPerf::FSample	Sample;
@@ -126,8 +190,8 @@ public:
 	[[nodiscard]] static TIoStatusOr<FMultiEndpointHttpClientResponse> Get(FAnsiStringView Url, const FMultiEndpointHttpClientConfig& Config);
 
 	[[nodiscard]] FHttpTicketId Get(const FOnDemandHostGroup& HostGroup, FAnsiStringView RelativeUrl, FOnHttpResponse&& OnResponse);
-	[[nodiscard]] FHttpTicketId Get(const FOnDemandHostGroup& HostGroup, FAnsiStringView RelativeUrl, const FIoOffsetAndLength& ChunkRange, FOnHttpResponse&& OnResponse);
-	[[nodiscard]] FHttpTicketId Get(const FOnDemandHostGroup& HostGroup, FAnsiStringView RelativeUrl, const FIoOffsetAndLength& ChunkRange, TArray<FAnsiString>&& Headers, EMultiEndpointRequestFlags Flags, FOnHttpResponse&& OnResponse);
+	[[nodiscard]] FHttpTicketId Get(const FOnDemandHostGroup& HostGroup, FAnsiStringView RelativeUrl, TArray<FIoOffsetAndLength>&& ChunkRanges, FOnHttpResponse&& OnResponse);
+	[[nodiscard]] FHttpTicketId Get(const FOnDemandHostGroup& HostGroup, FAnsiStringView RelativeUrl, TArray<FIoOffsetAndLength>&& ChunkRanges, TArray<FAnsiString>&& Headers, EMultiEndpointRequestFlags Flags, FOnHttpResponse&& OnResponse);
 
 	/**
 	 * Process the underlying http client.
@@ -170,7 +234,8 @@ private:
 		TArray<FAnsiString>			RequestHeaders;
 		TArray<FAnsiString>			ResponseHeaders;
 		FAnsiString					RelativeUrl;
-		FIoOffsetAndLength			Range;
+		TArray<FIoOffsetAndLength>	Ranges; // Consider inline storage?
+		FAnsiString					MultiRangeBoundary;
 		FHttpTicketId				TicketId = 0;
 		FConnection&				Connection;
 		/** Stored the response body */
@@ -194,6 +259,15 @@ private:
 		/** Used to make sure that the sink is only called with a completed state one per request attempt */
 		int32						SinkCounter = 0;
 #endif // DO_CHECK
+
+		/** Returns the total length of all the ranges in the provided request */
+		uint64 GetTotalRangeLength() const;
+
+		/**
+		 * Builds the value for a http 'range' header based on the ranges from the provided request.
+		 * Returns true if a valid range exists and false if not.
+		 */
+		bool BuildRangeHeaderValue(FAnsiStringBuilderBase& OutString)const;
 	};
 
 										FMultiEndpointHttpClient(const FMultiEndpointHttpClientConfig& Config);
@@ -211,7 +285,7 @@ private:
 
 	uint32								GetRetryLimitForRequest(const FRequest& Request) const;
 
-	void								Log(const FRequest& Request, const HTTP::FTicketStatus& TicketStatus) const;
+	void								Log(const FMultiEndpointHttpClientResponse& Response, const FRequest& Request) const;
 	void								LogError(const FRequest& Request, const FString& ErrorReason) const;
 
 	FMultiEndpointHttpClientConfig		Config;
diff --git a/Engine/Source/Runtime/Experimental/IoStore/OnDemand/Private/OnDemandHttpIoDispacher.cpp b/Engine/Source/Runtime/Experimental/IoStore/OnDemand/Private/OnDemandHttpIoDispacher.cpp
--- a/Engine/Source/Runtime/Experimental/IoStore/OnDemand/Private/OnDemandHttpIoDispacher.cpp
+++ b/Engine/Source/Runtime/Experimental/IoStore/OnDemand/Private/OnDemandHttpIoDispacher.cpp
@@ -253,7 +253,8 @@ int32 GIaxHttpMaxInflight = 8;
 static FAutoConsoleVariableRef CVar_IaxHttpMaxInflight(
 	TEXT("iax.HttpMaxInflight"),
 	GIaxHttpMaxInflight,
-	TEXT("Number of requests to issue at once")
+	TEXT("Number of requests to issue at once"),
+	FConsoleVariableDelegate::CreateStatic(&OnHttpClientCVarChanged)
 );
 
 int32 GIaxHttpOneRequestsPerConnection = 1;
@@ -713,20 +714,46 @@ static bool HttpRequestSortPredicate(const FHttpRequest& LHS, const FHttpRequest
 uint32 FHttpRequest::NextSeqNo = 0;
 
 ////////////////////////////////////////////////////////////////////////////////
-FIoHttpRange GetTotalRange(FHttpRequest* Request)
+
+TArray<FIoOffsetAndLength> GetTotalRange(FHttpRequest* Request)
 {
-	FIoHttpRange TotalRange;
-	for (FHttpRequest* It = Request; It != nullptr; It = It->Next)
+	TArray<FIoOffsetAndLength> Ranges;
+
+	if (GHttpIoDispatcherMaxRangeKiB < 0)
+	{
+		// If GHttpIoDispatcherMaxRangeKiB is negative then we will make multipart range request
+		// covering each request in the list. We can rely on the lower level systems to optimize
+		// the requests later.
+		for (FHttpRequest* It = Request; It != nullptr; It = It->Next)
+		{
+			// If the range is invalid we fetch the entire resource
+			if (It->Options.GetRange().IsValid() == false)
+			{
+				return TArray<FIoOffsetAndLength>();
+			}
+
+			Ranges.Add(It->Options.GetRange().ToOffsetAndLength());
+		}
+	}
+	else
 	{
-		// If the range is invalid we fetch the entire resource
-		if (It->Options.GetRange().IsValid() == false)
+		// If GHttpIoDispatcherMaxRangeKiB is 0 or positive then we want to make a partial range request
+		// and should only produce a single range to read.
+		FIoHttpRange TotalRange;
+		for (FHttpRequest* It = Request; It != nullptr; It = It->Next)
 		{
-			return FIoHttpRange();
+			// If the range is invalid we fetch the entire resource
+			if (It->Options.GetRange().IsValid() == false)
+			{
+				return TArray<FIoOffsetAndLength>();
+			}
+			TotalRange += It->Options.GetRange();
 		}
-		TotalRange += It->Options.GetRange();
+
+		Ranges.Add(TotalRange.ToOffsetAndLength());
 	}
 
-	return TotalRange;
+	return Ranges;
 }
 
 ////////////////////////////////////////////////////////////////////////////////
@@ -769,30 +796,6 @@ FHttpRequestList CreateBatch(
 	return Batch;
 }
 
-////////////////////////////////////////////////////////////////////////////////
-void DebugPrintRange(FHttpRequest* Requests, const UE::FIoHttpRange& TotalRange)
-{
-#if !UE_BUILD_SHIPPING 
-	if (Requests->Next == nullptr)
-	{
-		return; 
-	}
-
-	int32 Count = 0;
-	TStringBuilder<256> Sb;
-	for (FHttpRequest* It = Requests; It != nullptr; It = It->Next)
-	{
-		const UE::FIoHttpRange& Range = It->Options.GetRange();
-		Sb << TEXT("[") << Range.GetMin() << TEXT("-") << Range.GetMax() << TEXT("] ");
-		++Count;
-	}
-
-	Sb << TEXT("-> ") << TEXT("[") << TotalRange.GetMin() << TEXT("-") << TotalRange.GetMax() << TEXT("]");
-	UE_LOGF(LogHttpIoDispatcher, VeryVerbose, "Merging %d requests, %ls, total size = %.2lf KiB",
-		Count, Sb.ToString(), double(TotalRange.GetSize()) / 1024.0);
-#endif
-}
-
 ////////////////////////////////////////////////////////////////////////////////
 #if UE_TRACE_HTTP_ENABLED
 
@@ -1367,11 +1370,11 @@ private:
 	bool							ProcessHttpRequests();
 	void							CancelRemainingHttpRequests();
 	FHttpRequest*					DequeueHttpRequests();
-	void							CompleteHttpRequest(FHttpRequest& Request, FMultiEndpointHttpClientResponse&& HttpResponse, const FIoHttpRange& TotalRange);
-	void							CompleteRequest(FHttpRequest& Request, TArray<FAnsiString>&& Headers, FIoBuffer Body, const FIoHttpRange& TotalRange, uint32 StatusCode, bool bCached);
-	void							CompleteRequest(FHttpRequest& Request, FIoBuffer& Body, const FIoHttpRange& TotalRange, uint32 StatusCode, bool bCached)
+	void							CompleteHttpRequest(FHttpRequest& Request, FMultiEndpointHttpClientResponse&& HttpResponse);
+	void							CompleteRequest(FHttpRequest& Request, TArray<FAnsiString>&& Headers, const FResponseBody& Data, uint32 StatusCode, bool bCached);
+	void							CompleteRequest(FHttpRequest& Request, const FIoBuffer& Body, uint32 StatusCode, bool bCached)
 									{
-										CompleteRequest(Request, TArray<FAnsiString>(), Body, TotalRange, StatusCode, bCached);
+										CompleteRequest(Request, TArray<FAnsiString>(), FResponseBody(Body), StatusCode, bCached);
 									}
 
 	void							ApplyCancellationToken(FHttpRequest* RequestChain, FMultiEndpointHttpClient::FHttpTicketId TicketId);
@@ -1801,7 +1804,7 @@ void FHttpDispatcher::CancelRequest(FIoHttpRequest Handle)
 		InflightRequestCount.fetch_add(1, std::memory_order_relaxed);
 		UE::Tasks::Launch(UE_SOURCE_LOCATION, [this, Request]
 		{
-			CompleteRequest(*Request, Request->Buffer, FIoHttpRange(), 0, false);
+			CompleteRequest(*Request, Request->Buffer, 0, false);
 		}, UE::Tasks::ETaskPriority::BackgroundLow);
 	}
 
@@ -1941,7 +1944,7 @@ bool FHttpDispatcher::TryReadFromCache(FHttpRequest* Request)
 
 		UE::Tasks::Launch(UE_SOURCE_LOCATION, [this, Request]
 		{
-			CompleteRequest(*Request, Request->Buffer, FIoHttpRange(), 200, true);
+			CompleteRequest(*Request, Request->Buffer, 200, true);
 		}, UE::Tasks::ETaskPriority::BackgroundLow);
 
 		return true;
@@ -1956,11 +1959,11 @@ bool FHttpDispatcher::TryReadFromCache(FHttpRequest* Request)
 		{
 			if (Request->CacheStatus == EIoErrorCode::Ok)
 			{
-				CompleteRequest(*Request, Request->Buffer, FIoHttpRange(), 200, true);
+				CompleteRequest(*Request, Request->Buffer, 200, true);
 			}
 			else if ((Request->CacheStatus == EIoErrorCode::Cancelled) || Request->Flags.HasAny(EHttpRequestFlags::CancelRequested) || bStopRequested.load(std::memory_order_relaxed))
 			{
-				CompleteRequest(*Request, Request->Buffer, FIoHttpRange(), 0, false);
+				CompleteRequest(*Request, Request->Buffer, 0, false);
 			}
 			else
 			{
@@ -2018,7 +2021,7 @@ bool FHttpDispatcher::ProcessHttpRequests()
 				// then we are simulating an error.
 				const bool bWasCanceled = !HostGroup.IsConnected();
 				FOnDemandIoBackendStats::Get()->OnHttpCompleted(HttpRequest->GetHttpRequestType(), 0, 0, 0, 0, 0, bWasCanceled);
-				CompleteRequest(*HttpRequest, HttpRequest->Buffer, FIoHttpRange(), 0, false);
+				CompleteRequest(*HttpRequest, HttpRequest->Buffer, 0, false);
 
 				check(NextHttpRequest == nullptr);
 				NextHttpRequest = DequeueHttpRequests();
@@ -2029,7 +2032,7 @@ bool FHttpDispatcher::ProcessHttpRequests()
 			HashString.Reset();
 
 			TArray<FAnsiString> Headers;
-			FIoHttpRange		TotalRange;
+			TArray<FIoOffsetAndLength> Ranges;
 
 			const bool bGenericHttpRequest = HttpRequest->ChunkHash.IsZero();
 			if (bGenericHttpRequest)
@@ -2043,8 +2046,7 @@ bool FHttpDispatcher::ProcessHttpRequests()
 				// I/O store chunk request 
 				HashString << HttpRequest->ChunkHash;
 				Url << *HttpRequest->RelativeUrl << "/" << HashString.ToView().Left(2) << "/" << HashString << ANSITEXTVIEW(".iochunk");
-				TotalRange = GetTotalRange(HttpRequest);
-				//DebugPrintRange(HttpRequest, TotalRange);
+				Ranges = GetTotalRange(HttpRequest);
 			}
 
 #if UE_ALLOW_ERROR_SIMULATION
@@ -2061,15 +2063,15 @@ bool FHttpDispatcher::ProcessHttpRequests()
 
 			// Note that the requests issued here must be completed before this method can return so the capturing of locals is safe.
 			TRACE_REQUEST_STARTED(HttpRequest, Url.ToString());
-			FMultiEndpointHttpClient::FHttpTicketId TicketId = HttpClient->Get(HostGroup.GetUnderlyingHostGroup(), Url, TotalRange.ToOffsetAndLength(), MoveTemp(Headers), Flags,
-				[this, HttpRequest, TotalRange, &NumConcurrentRequests, HostGroup = MoveTemp(HostGroup)]
+			FMultiEndpointHttpClient::FHttpTicketId TicketId = HttpClient->Get(HostGroup.GetUnderlyingHostGroup(), Url, MoveTemp(Ranges), MoveTemp(Headers), Flags,
+				[this, HttpRequest, &NumConcurrentRequests, HostGroup = MoveTemp(HostGroup)]
 				(FMultiEndpointHttpClientResponse&& HttpResponse) mutable
 				{
 					TRACE_REQUEST_COMPLETED(HttpRequest, HttpResponse, HostGroup);
 					FOnDemandIoBackendStats::Get()->OnHttpCompleted(
 						HttpRequest->GetHttpRequestType(),
 						HttpResponse.StatusCode,
-						HttpResponse.Body.DataSize(),
+						HttpResponse.Body.GetSize(),
 						HttpResponse.Sample.GetTotalMs(),
 						HttpResponse.RetryCount,
 						HttpResponse.CDNCacheStatus,
@@ -2081,9 +2083,9 @@ bool FHttpDispatcher::ProcessHttpRequests()
 
 					UE::Tasks::Launch(
 						UE_SOURCE_LOCATION,
-						[this, HttpRequest, HttpResponse = MoveTemp(HttpResponse), TotalRange]() mutable
+						[this, HttpRequest, HttpResponse = MoveTemp(HttpResponse)]() mutable
 						{
-							CompleteHttpRequest(*HttpRequest, MoveTemp(HttpResponse), TotalRange);
+							CompleteHttpRequest(*HttpRequest, MoveTemp(HttpResponse));
 						}, UE::Tasks::ETaskPriority::BackgroundLow);
 				});
 
@@ -2132,7 +2134,7 @@ void FHttpDispatcher::CancelRemainingHttpRequests()
 
 		InflightRequestCount.fetch_add(1, std::memory_order_relaxed);
 		FIoBuffer Buffer;
-		CompleteRequest(*Requests, Buffer, FIoHttpRange(), 0, false);
+		CompleteRequest(*Requests, Buffer, 0, false);
 	}
 }
 
@@ -2153,16 +2155,17 @@ FHttpRequest* FHttpDispatcher::DequeueHttpRequests()
 	return nullptr;
 }
 
-void FHttpDispatcher::CompleteHttpRequest(FHttpRequest& Request, FMultiEndpointHttpClientResponse&& HttpResponse, const FIoHttpRange& TotalRange)
+void FHttpDispatcher::CompleteHttpRequest(FHttpRequest& Request, FMultiEndpointHttpClientResponse&& HttpResponse)
 {
-	CompleteRequest(Request, MoveTemp(HttpResponse.Headers), HttpResponse.Body, TotalRange, HttpResponse.StatusCode, false);
+	CompleteRequest(Request, MoveTemp(HttpResponse.Headers), HttpResponse.Body, HttpResponse.StatusCode, false);
 }
 
-void FHttpDispatcher::CompleteRequest(FHttpRequest& Request, TArray<FAnsiString>&& Headers, FIoBuffer Body, const FIoHttpRange& TotalRange, const uint32 StatusCode, bool bCached)
+void FHttpDispatcher::CompleteRequest(FHttpRequest& Request, TArray<FAnsiString>&& Headers, const FResponseBody& ResponseBody, const uint32 StatusCode, bool bCached)
 {
 #if UE_ALLOW_CACHE_POISONING
 	if (GIasPoisonCache && bCached)
 	{
+		FIoBuffer Body = ResponseBody.GetBody();
 		for (uint64 Index = 0; Index < Body.GetSize(); Index++)
 		{
 			Body.GetData()[Index] = 0x4d;
@@ -2171,6 +2174,7 @@ void FHttpDispatcher::CompleteRequest(FHttpRequest& Request, TArray<FAnsiString>
 #endif // UE_ALLOW_CACHE_POISONING
 
 	FHttpRequest* NextRequest = &Request;
+
 	while (NextRequest != nullptr)
 	{
 		FHttpRequest* ToComplete	= NextRequest;
@@ -2185,37 +2189,28 @@ void FHttpDispatcher::CompleteRequest(FHttpRequest& Request, TArray<FAnsiString>
 
 		const FIoHttpRange& Range		= ToComplete->Options.GetRange();
 		EIoErrorCode CompletionStatus	= IsHttpStatusOk(StatusCode) ? EIoErrorCode::Ok : EIoErrorCode::ReadError;
-		FIoBuffer PartialBody			= Body;
+		
+		FIoBuffer Body;
 
 		if (ToComplete->Flags.HasAny(EHttpRequestFlags::CancelRequested))
 		{
 			CompletionStatus	= EIoErrorCode::Cancelled;
-			PartialBody			= FIoBuffer();
 			RequestStatusCode = 0; // Mark as a none http response related error
 		}
-
-		// If the request is NOT cached and specified a range we retrieve the partial body
-		if (!bCached && CompletionStatus == EIoErrorCode::Ok && Range.IsValid())
+		else if (bCached)
 		{
-			uint64			Offset	= Range.GetMin();
-			const uint64	Size	= Range.GetSize();
-
-			// If the total range is NOT valid we fetched the entire resource
-			if (TotalRange.IsValid())
-			{
-				ensure(Range.GetMin() >= TotalRange.GetMin());
-				ensure(Range.GetMax() <= TotalRange.GetMax());
-				Offset -= TotalRange.GetMin();
-			}
-
-			PartialBody = FIoBuffer(Body.GetView().Mid(Offset, Size), Body);
+			Body = ResponseBody.GetData(FIoOffsetAndLength());
+		}
+		else
+		{
+			Body = ResponseBody.GetData(Range.ToOffsetAndLength());
 		}
 
 		const EIoHttpResponseFlags Flags	= bCached ? EIoHttpResponseFlags::Cached : EIoHttpResponseFlags::None;
 		FIoHttpRequestCompleted OnCompleted = MoveTemp(ToComplete->OnCompleted);
 		TArray<FAnsiString> HeadersCopy		= Headers; // TODO: Only used for generic HTTP request but still bad 
 
-		OnCompleted(FIoHttpResponse(ToComplete->CacheKey, FIoHttpHeaders::Create(MoveTemp(HeadersCopy)), PartialBody, CompletionStatus, RequestStatusCode, Flags));
+		OnCompleted(FIoHttpResponse(ToComplete->CacheKey, FIoHttpHeaders::Create(MoveTemp(HeadersCopy)), MoveTemp(Body), CompletionStatus, RequestStatusCode, Flags));
 
 		ToComplete->Flags.Add(EHttpRequestFlags::Completed);
 		ToComplete->CompletionStatus.store(CompletionStatus, std::memory_order_seq_cst);
diff --git a/Engine/Source/Runtime/Experimental/IoStore/OnDemandCore/Public/IO/HttpIoDispatcher.h b/Engine/Source/Runtime/Experimental/IoStore/OnDemandCore/Public/IO/HttpIoDispatcher.h
--- a/Engine/Source/Runtime/Experimental/IoStore/OnDemandCore/Public/IO/HttpIoDispatcher.h
+++ b/Engine/Source/Runtime/Experimental/IoStore/OnDemandCore/Public/IO/HttpIoDispatcher.h
@@ -314,7 +314,7 @@ public:
 							~FIoHttpResponse() = default;
 
 	inline					FIoHttpResponse(const FIoHash& InCacheKey, const FIoBuffer& InBody, EIoErrorCode InErrorCode, uint32 InStatusCode, EIoHttpResponseFlags InFlags);
-	inline					FIoHttpResponse(const FIoHash& InCacheKey, FIoHttpHeaders&& InHeaders, const FIoBuffer& InBody, EIoErrorCode InErrorCode, uint32 InStatusCode, EIoHttpResponseFlags InFlags);
+	inline					FIoHttpResponse(const FIoHash& InCacheKey, FIoHttpHeaders&& InHeaders, FIoBuffer&& InBody, EIoErrorCode InErrorCode, uint32 InStatusCode, EIoHttpResponseFlags InFlags);
 							FIoHttpResponse(EIoErrorCode InErrorCode, uint32 InStatusCode)
 								: ErrorCode(InErrorCode)
 								, StatusCode(InStatusCode) { }
@@ -347,10 +347,10 @@ FIoHttpResponse::FIoHttpResponse(const FIoHash& InCacheKey, const FIoBuffer& InB
 {
 }
 
-FIoHttpResponse::FIoHttpResponse(const FIoHash& InCacheKey, FIoHttpHeaders&& InHeaders, const FIoBuffer& InBody, EIoErrorCode InErrorCode, uint32 InStatusCode, EIoHttpResponseFlags InFlags)
+FIoHttpResponse::FIoHttpResponse(const FIoHash& InCacheKey, FIoHttpHeaders&& InHeaders, FIoBuffer&& InBody, EIoErrorCode InErrorCode, uint32 InStatusCode, EIoHttpResponseFlags InFlags)
 	: CacheKey(InCacheKey)
 	, Headers(MoveTemp(InHeaders))
-	, Body(InBody)
+	, Body(MoveTemp(InBody))
 	, ErrorCode(InErrorCode)
 	, StatusCode(InStatusCode)
 	, Flags(InFlags)
