Dataflow : fix a crash when using dataflow and recompiling code using Live Coding

#Synced-CL 52497132

diff --git a/Engine/Source/Runtime/Experimental/Dataflow/Engine/Private/Dataflow/DataflowObject.cpp b/Engine/Source/Runtime/Experimental/Dataflow/Engine/Private/Dataflow/DataflowObject.cpp
--- a/Engine/Source/Runtime/Experimental/Dataflow/Engine/Private/Dataflow/DataflowObject.cpp
+++ b/Engine/Source/Runtime/Experimental/Dataflow/Engine/Private/Dataflow/DataflowObject.cpp
@@ -167,7 +167,10 @@ void UDataflow::AddReferencedObjects(UObject* InThis, FReferenceCollector& Colle
 		Collector.AddReferencedObject(Target);
 	}
 
-	This->Dataflow->AddReferencedObjects(Collector);
+	if (This->Dataflow)
+	{
+		This->Dataflow->AddReferencedObjects(Collector);
+	}
 	Super::AddReferencedObjects(InThis, Collector);
 }
 
