Dataflow : update get asset path node to be more general and be able to return the path of the bound asset

#Synced-CL 52497132

diff --git a/Engine/Source/Runtime/Experimental/Dataflow/Engine/Private/Dataflow/DataflowAssetPathNodes.cpp b/Engine/Source/Runtime/Experimental/Dataflow/Engine/Private/Dataflow/DataflowAssetPathNodes.cpp
--- a/Engine/Source/Runtime/Experimental/Dataflow/Engine/Private/Dataflow/DataflowAssetPathNodes.cpp
+++ b/Engine/Source/Runtime/Experimental/Dataflow/Engine/Private/Dataflow/DataflowAssetPathNodes.cpp
@@ -1,8 +1,10 @@
 // Copyright Epic Games, Inc. All Rights Reserved.
 
 #include "Dataflow/DataflowAssetPathNodes.h"
-#include "Dataflow/DataflowNodeFactory.h"
+
 #include "Dataflow/DataflowAnyTypeRegistry.h"
+#include "Dataflow/DataflowNodeFactory.h"
+#include "Dataflow/DataflowObjectInterface.h"
 
 #include UE_INLINE_GENERATED_CPP_BY_NAME(DataflowAssetPathNodes)
 
@@ -48,13 +50,22 @@ void FDataflowGetAssetPathNode::Evaluate(UE::Dataflow::FContext& Context, const
 
 		TObjectPtr<UObject> InAsset = GetValue(Context, &Asset);
 		if (InAsset == nullptr)
+		{
+			// check if there's any bound asset to the context 
+			if (const UE::Dataflow::FEngineContext* EngineContext = Context.AsType<UE::Dataflow::FEngineContext>())
+			{
+				InAsset = EngineContext->Owner;
+			}
+		}
+		if (InAsset == nullptr)
 		{
 			Context.Warning(LOCTEXT("NullAsset", "Asset is null, empty path will be returned"), this, Out);
-			
 		}
 		else
 		{
-			OutPath = InAsset->GetPathName();
+			// Make sure we get the main object and not a subObject ( that can happen in teh case of attachements ) 
+			FSoftObjectPath SoftPath(InAsset);
+			OutPath = SoftPath.GetAssetPath().ToString();
 		}
 		SetValue(Context, OutPath, &Path);
 	}
diff --git a/Engine/Source/Runtime/Experimental/Dataflow/Engine/Private/Dataflow/DataflowAssetPathNodes.h b/Engine/Source/Runtime/Experimental/Dataflow/Engine/Private/Dataflow/DataflowAssetPathNodes.h
--- a/Engine/Source/Runtime/Experimental/Dataflow/Engine/Private/Dataflow/DataflowAssetPathNodes.h
+++ b/Engine/Source/Runtime/Experimental/Dataflow/Engine/Private/Dataflow/DataflowAssetPathNodes.h
@@ -11,6 +11,7 @@
 
 /*
 * Get the path from an asset
+* If no asset is specified via the input or in the properties the asset bound to the dataflow will be used 
 */
 USTRUCT()
 struct FDataflowGetAssetPathNode : public FDataflowNode
