Support for setting a custom cache key for on-demand HTTP requests.

#rnx

#Synced-CL 52551117

diff --git a/Engine/Source/Runtime/Experimental/IoStore/OnDemand/Private/OnDemandHttpIoDispacher.cpp b/Engine/Source/Runtime/Experimental/IoStore/OnDemand/Private/OnDemandHttpIoDispacher.cpp
--- a/Engine/Source/Runtime/Experimental/IoStore/OnDemand/Private/OnDemandHttpIoDispacher.cpp
+++ b/Engine/Source/Runtime/Experimental/IoStore/OnDemand/Private/OnDemandHttpIoDispacher.cpp
@@ -648,7 +648,6 @@ struct FHttpRequest final
 	virtual inline FRefCountValue	Release() const override;
 
 	FIoHash							ChunkHash;
-	FIoHash							CacheKey; //TODO make this smaller the IAS cache only uses 64 bit cache keys
 	FIoHttpRequestCompleted 		OnCompleted;
 	FIoHttpHeaders					Headers;
 	FIoHttpOptions					Options;
@@ -1727,13 +1726,19 @@ FIoHttpRequest FHttpDispatcher::CreateRequest(
 
 	if (ChunkHash != nullptr)
 	{
-		Request->ChunkHash	= *ChunkHash;
-		Request->CacheKey	= GetCacheKey(*ChunkHash, Options.GetRange());
+		Request->ChunkHash = *ChunkHash;
+		if (Options.GetCacheKey().IsZero())
+		{
+			Request->Options.SetCacheKey(GetCacheKey(*ChunkHash, Options.GetRange()));
+		}
 	}
 	else
 	{
-		Request->ChunkHash	= FIoHash::Zero;
-		Request->CacheKey	= GetCacheKey(HostGroupName, Request->RelativeUrl, Request->Headers, Request->Options);
+		Request->ChunkHash = FIoHash::Zero;
+		if (Options.GetCacheKey().IsZero())
+		{
+			Request->Options.SetCacheKey(GetCacheKey(HostGroupName, Request->RelativeUrl, Request->Headers, Request->Options));
+		}
 	}
 
 	if (Last == nullptr)
@@ -1937,7 +1942,7 @@ bool FHttpDispatcher::TryReadFromCache(FHttpRequest* Request)
 	}
 
 	Request->Flags.Add(EHttpRequestFlags::CacheInflight);
-	const EIoErrorCode CacheStatus = Cache->Get(Request->CacheKey, Request->Buffer);
+	const EIoErrorCode CacheStatus = Cache->Get(Request->Options.GetCacheKey(), Request->Buffer);
 	if (CacheStatus == EIoErrorCode::Ok)
 	{
 		InflightRequestCount.fetch_add(1, std::memory_order_relaxed);
@@ -1978,7 +1983,7 @@ bool FHttpDispatcher::TryReadFromCache(FHttpRequest* Request)
 			}
 		}, OnReadyEvent, UE::Tasks::ETaskPriority::BackgroundLow);
 
-		Cache->Materialize(Request->CacheKey, Request->Buffer, Request->CacheStatus, MoveTemp(OnReadyEvent));
+		Cache->Materialize(Request->Options.GetCacheKey(), Request->Buffer, Request->CacheStatus, MoveTemp(OnReadyEvent));
 
 		return true;
 	}
@@ -2210,7 +2215,13 @@ void FHttpDispatcher::CompleteRequest(FHttpRequest& Request, TArray<FAnsiString>
 		FIoHttpRequestCompleted OnCompleted = MoveTemp(ToComplete->OnCompleted);
 		TArray<FAnsiString> HeadersCopy		= Headers; // TODO: Only used for generic HTTP request but still bad 
 
-		OnCompleted(FIoHttpResponse(ToComplete->CacheKey, FIoHttpHeaders::Create(MoveTemp(HeadersCopy)), MoveTemp(Body), CompletionStatus, RequestStatusCode, Flags));
+		OnCompleted(FIoHttpResponse(
+			ToComplete->Options.GetCacheKey(),
+			FIoHttpHeaders::Create(MoveTemp(HeadersCopy)),
+			MoveTemp(Body),
+			CompletionStatus,
+			RequestStatusCode,
+			Flags));
 
 		ToComplete->Flags.Add(EHttpRequestFlags::Completed);
 		ToComplete->CompletionStatus.store(CompletionStatus, std::memory_order_seq_cst);
diff --git a/Engine/Source/Runtime/Experimental/IoStore/OnDemand/Private/OnDemandHttpIoDispatcherBackend.cpp b/Engine/Source/Runtime/Experimental/IoStore/OnDemand/Private/OnDemandHttpIoDispatcherBackend.cpp
--- a/Engine/Source/Runtime/Experimental/IoStore/OnDemand/Private/OnDemandHttpIoDispatcherBackend.cpp
+++ b/Engine/Source/Runtime/Experimental/IoStore/OnDemand/Private/OnDemandHttpIoDispatcherBackend.cpp
@@ -6,6 +6,7 @@
 #include "HAL/PlatformTime.h"
 #include "IO/HttpIoDispatcher.h"
 #include "IO/IoAllocators.h"
+#include "IO/IoHash.h"
 #include "IO/IoOffsetLength.h"
 #include "IO/IoStatus.h"
 #include "IO/IoStoreOnDemand.h"
@@ -39,6 +40,17 @@ static FAutoConsoleVariableRef CVar_IasFastCancelEnabled(
 	GIasFastCancelEnabled,
 	TEXT("Enables canceling IAS requests before the completion callback has been triggered.")
 );
+
+////////////////////////////////////////////////////////////////////////////////
+static FIoHash GetHttpCacheKey(const FIoHash& ChunkHash, const FIoHttpRange& Range)
+{
+	FIoHashBuilder HashBuilder;
+	HashBuilder.Update(ChunkHash.GetBytes(), sizeof(FIoHash::ByteArray));
+	HashBuilder.Update(&Range, sizeof(FIoHttpRange));
+
+	return HashBuilder.Finalize();
+}
+
 ////////////////////////////////////////////////////////////////////////////////
 class FHttpIoDispatcherBackend
 	: public IOnDemandHttpIoDispatcherBackend
@@ -213,14 +225,18 @@ void FHttpIoDispatcherBackend::ResolveIoRequests(FIoRequestList Requests, FIoReq
 				ResolvedSize).ConsumeValueOrDie();
 		}
 
-		const FIoHttpRange HttpRange = FIoHttpRange::FromOffsetAndLength(ChunkRange);
-		FResolvedRequest* ResolvedRequest = FResolvedRequest::Create(*Request, MoveTemp(ChunkInfo), ChunkRange);
+		FResolvedRequest*	ResolvedRequest = FResolvedRequest::Create(*Request, MoveTemp(ChunkInfo), ChunkRange);
+		const FIoHttpRange	HttpRange = FIoHttpRange::FromOffsetAndLength(ChunkRange);
+		const FIoHash		CacheKey = GetHttpCacheKey(ResolvedRequest->ChunkInfo.ChunkEntry().Hash, HttpRange);
+		const FIoHash&		ResourceHash = ResolvedRequest->ChunkInfo.ChunkEntry().Hash; // The content hash is the CDN resource hash
+		const EIoHttpFlags	Flags = EIoHttpFlags::Default; // Using default flags will fetch the request from the local cache
+
 		ResolvedRequest->HttpRequest = Batch.Get(
 			ResolvedRequest->ChunkInfo.HostGroupName(),
 			ResolvedRequest->ChunkInfo.RelativeUrl(),
 			FIoHttpHeaders(),
-			FIoHttpOptions(Request->Priority, HttpRetryCount, EIoHttpFlags::Default, HttpRange, FIoHttpMetadata(Request->ChunkId)),
-			ResolvedRequest->ChunkInfo.ChunkEntry().Hash,
+			FIoHttpOptions(Request->Priority, HttpRetryCount, Flags, HttpRange, CacheKey, FIoHttpMetadata(Request->ChunkId)),
+			ResourceHash,
 			[this, ResolvedRequest](FIoHttpResponse&& HttpResponse)
 			{
 				// Callbacks are always triggered from the task pool
diff --git a/Engine/Source/Runtime/Experimental/IoStore/OnDemandCore/Public/IO/HttpIoDispatcher.h b/Engine/Source/Runtime/Experimental/IoStore/OnDemandCore/Public/IO/HttpIoDispatcher.h
--- a/Engine/Source/Runtime/Experimental/IoStore/OnDemandCore/Public/IO/HttpIoDispatcher.h
+++ b/Engine/Source/Runtime/Experimental/IoStore/OnDemandCore/Public/IO/HttpIoDispatcher.h
@@ -144,13 +144,15 @@ private:
 struct FIoHttpOptions
 {
 public:
-	inline						FIoHttpOptions(int32 InPriority, int32 InRetryCount, EIoHttpFlags InFlags, const FIoHttpRange& InRange, FIoHttpMetadata Metadata);
+	inline						FIoHttpOptions(int32 InPriority, int32 InRetryCount, EIoHttpFlags InFlags, const FIoHttpRange& InRange, const FIoHash& CacheKey, FIoHttpMetadata Metadata);
+	inline						FIoHttpOptions(int32 InPriority, int32 InRetryCount, EIoHttpFlags InFlags, const FIoHttpRange& InRange, const FIoHash& CacheKey);
 	inline						FIoHttpOptions(int32 InPriority, int32 InRetryCount, EIoHttpFlags InFlags, const FIoHttpRange& InRange);
 	inline						FIoHttpOptions(int32 InPriority, int32 InRetryCount, EIoHttpFlags InFlags);
 	inline						FIoHttpOptions(int32 InPriority, int32 InRetryCount);
 	inline explicit				FIoHttpOptions(int32 InPriority);
 								FIoHttpOptions() = default;
 
+	const FIoHash				GetCacheKey()	const { return CacheKey; }
 	const FIoHttpRange&			GetRange()		const { return Range; }
 	int32						GetPriority()	const { return Priority; }
 	int32						GetRetryCount()	const { return RetryCount; }
@@ -158,6 +160,7 @@ public:
 	uint8						GetCategory()	const { return Category; }
 	const FIoHttpMetadata&		GetMetadata()	const { return Metadata; }
 
+	void						SetCacheKey(const FIoHash& InCacheKey)	{ CacheKey = InCacheKey; }
 	void						SetRange(const FIoHttpRange InRange)	{ Range = InRange; }
 	void						SetPriority(int32 InPriority)			{ Priority = InPriority; }
 	void						SetRetryCount(int32 InRetryCount)		{ RetryCount = InRetryCount; }
@@ -167,6 +170,7 @@ public:
 	UE_API static const FIoHttpOptions Default;
 
 private:
+	FIoHash			CacheKey = FIoHash::Zero;
 	FIoHttpRange	Range;
 	int32			Priority = 0;
 	int32			RetryCount = 0;
@@ -175,8 +179,9 @@ private:
 	FIoHttpMetadata	Metadata;
 };
 
-FIoHttpOptions::FIoHttpOptions(int32 InPriority, int32 InRetryCount, EIoHttpFlags InFlags, const FIoHttpRange& InRange, FIoHttpMetadata InMetadata)
-	: Range(InRange)
+FIoHttpOptions::FIoHttpOptions(int32 InPriority, int32 InRetryCount, EIoHttpFlags InFlags, const FIoHttpRange& InRange, const FIoHash& InCacheKey, FIoHttpMetadata InMetadata)
+	: CacheKey(InCacheKey)
+	, Range(InRange)
 	, Priority(InPriority)
 	, RetryCount(InRetryCount)
 	, Flags(InFlags)
@@ -184,6 +189,15 @@ FIoHttpOptions::FIoHttpOptions(int32 InPriority, int32 InRetryCount, EIoHttpFlag
 {
 }
 
+FIoHttpOptions::FIoHttpOptions(int32 InPriority, int32 InRetryCount, EIoHttpFlags InFlags, const FIoHttpRange& InRange, const FIoHash& InCacheKey)
+	: CacheKey(InCacheKey)
+	, Range(InRange)
+	, Priority(InPriority)
+	, RetryCount(InRetryCount)
+	, Flags(InFlags)
+{
+}
+
 FIoHttpOptions::FIoHttpOptions(int32 InPriority, int32 InRetryCount, EIoHttpFlags InFlags, const FIoHttpRange& InRange)
 	: Range(InRange)
 	, Priority(InPriority)
