Add HashBuffer to FGenericHash.

#rnx

#Synced-CL 52551117

diff --git a/Engine/Source/Runtime/Experimental/IoStore/OnDemandCore/Internal/IO/GenericHash.cpp b/Engine/Source/Runtime/Experimental/IoStore/OnDemandCore/Internal/IO/GenericHash.cpp
--- a/Engine/Source/Runtime/Experimental/IoStore/OnDemandCore/Internal/IO/GenericHash.cpp
+++ b/Engine/Source/Runtime/Experimental/IoStore/OnDemandCore/Internal/IO/GenericHash.cpp
@@ -3,6 +3,7 @@
 #include "GenericHash.h"
 
 #include "Containers/ArrayView.h"
+#include "Hash/Blake3.h"
 #include "IO/IoHash.h"
 #include "Misc/StringBuilder.h"
 #include "String/BytesToHex.h"
@@ -30,4 +31,14 @@ FString ToHex(FMemoryView Memory)
 	return FString(Sb);
 }
 
+void HashBuffer(FMutableMemoryView Dst, FMemoryView Src)
+{
+	check(Dst.GetSize() > 0);
+	check(Src.GetSize() > 0);
+	FBlake3Hash		Blake = FBlake3::HashBuffer(Src);
+	const uint64	BlakeSize = uint64(sizeof(FBlake3Hash::ByteArray));
+
+	Dst.CopyFrom(FMemoryView(Blake.GetBytes(), FMath::Min(Dst.GetSize(), BlakeSize)));
+}
+
 } // namesapce UE::GenericHash
diff --git a/Engine/Source/Runtime/Experimental/IoStore/OnDemandCore/Internal/IO/GenericHash.h b/Engine/Source/Runtime/Experimental/IoStore/OnDemandCore/Internal/IO/GenericHash.h
--- a/Engine/Source/Runtime/Experimental/IoStore/OnDemandCore/Internal/IO/GenericHash.h
+++ b/Engine/Source/Runtime/Experimental/IoStore/OnDemandCore/Internal/IO/GenericHash.h
@@ -22,6 +22,7 @@ namespace GenericHash
 
 UE_API FString				ToHex(FMemoryView Memory);
 UE_API FStringBuilderBase&	ToHex(FMemoryView Memory, FStringBuilderBase& Out);
+UE_API void					HashBuffer(FMutableMemoryView Dst, FMemoryView Src);
 
 } // namespace UE::GenericHash
 
@@ -36,6 +37,9 @@ struct THash
 	SIZE_T				GetSize() const { return Size; }
 	FMemoryView			GetView() const { return FMemoryView(Hash, Size); }
 
+	static THash		HashBuffer(FMemoryView Memory);
+	static void			HashBuffer(FMemoryView Memory, THash& Out);
+
 	static THash		From(FMemoryView Memory);
 	static THash		From(const uint8* Data, SIZE_T DataSize);
 	static THash		From(const FIoHash& IoHash);
@@ -108,6 +112,20 @@ THash<BitCount> THash<BitCount>::From(const FIoHash& IoHash)
 	return Hash;
 }
 
+template<SIZE_T BitCount>
+void THash<BitCount>::HashBuffer(FMemoryView Memory, THash<BitCount>& Out)
+{
+	GenericHash::HashBuffer(MakeMemoryView(Out.Hash, Out.GetSize()), Memory);
+}
+
+template<SIZE_T BitCount>
+THash<BitCount> THash<BitCount>::HashBuffer(FMemoryView Memory)
+{
+	THash<BitCount> Hash;
+	GenericHash::HashBuffer(MakeMemoryView(Hash.Hash, Hash.GetSize()), Memory);
+	return Hash;
+}
+
 using FHash32	= THash<32>;
 using FHash64	= THash<64>;
 using FHash96	= THash<96>;
