Chaos API: Updating triangle mesh geometry to hide the include for the implicit.
#Synced-CL 52632702

diff --git a/Engine/Source/Runtime/Experimental/RigidPhysics/Private/RigidPhysics/Geometry/TriangleMeshGeometry.cpp b/Engine/Source/Runtime/Experimental/RigidPhysics/Private/RigidPhysics/Geometry/TriangleMeshGeometry.cpp
--- a/Engine/Source/Runtime/Experimental/RigidPhysics/Private/RigidPhysics/Geometry/TriangleMeshGeometry.cpp
+++ b/Engine/Source/Runtime/Experimental/RigidPhysics/Private/RigidPhysics/Geometry/TriangleMeshGeometry.cpp
@@ -3,8 +3,16 @@
 
 #if UE_RIGIDPHYSICS_API_ENABLED
 
+#include "Chaos/TriangleMeshImplicitObject.h"
+
 namespace UE::Physics
 {
+	FTriangleMeshGeometry::~FTriangleMeshGeometry() = default;
+	FTriangleMeshGeometry::FTriangleMeshGeometry(const FTriangleMeshGeometry&) = default;
+	FTriangleMeshGeometry::FTriangleMeshGeometry(FTriangleMeshGeometry&&) = default;
+	FTriangleMeshGeometry& FTriangleMeshGeometry::operator=(const FTriangleMeshGeometry&) = default;
+	FTriangleMeshGeometry& FTriangleMeshGeometry::operator=(FTriangleMeshGeometry&&) = default;
+
 	Chaos::FTrimeshBVH::EPerFaceBoundsRepresentation Convert(FTriangleMeshGeometrySetup::EPerFaceBoundsRepresentation InValue)
 	{
 		switch (InValue)
@@ -25,7 +33,7 @@ namespace UE::Physics
 		}
 	}
 
-	FTriangleMeshGeometry::FTriangleMeshGeometry(FTriangleMeshImplicitObjectPtr InImplicit)
+	FTriangleMeshGeometry::FTriangleMeshGeometry(TRefCountPtr<Chaos::FTriangleMeshImplicitObject> InImplicit)
 		: Implicit(InImplicit)
 	{
 	}
@@ -132,7 +140,7 @@ namespace UE::Physics
 		return TEXT("TriangleMesh");
 	}
 
-	Chaos::FTriangleMeshImplicitObjectPtr FTriangleMeshGeometry::GetImplicit() const
+	TRefCountPtr<Chaos::FTriangleMeshImplicitObject> FTriangleMeshGeometry::GetImplicit() const
 	{
 		return Implicit;
 	}
diff --git a/Engine/Source/Runtime/Experimental/RigidPhysics/Public/RigidPhysics/Geometry/TriangleMeshGeometry.h b/Engine/Source/Runtime/Experimental/RigidPhysics/Public/RigidPhysics/Geometry/TriangleMeshGeometry.h
--- a/Engine/Source/Runtime/Experimental/RigidPhysics/Public/RigidPhysics/Geometry/TriangleMeshGeometry.h
+++ b/Engine/Source/Runtime/Experimental/RigidPhysics/Public/RigidPhysics/Geometry/TriangleMeshGeometry.h
@@ -5,9 +5,13 @@
 
 #if UE_RIGIDPHYSICS_API_ENABLED
 
-#include "Chaos/TriangleMeshImplicitObject.h"
 #include "RigidPhysics/RigidMaterialIndex.h"
 
+namespace Chaos
+{
+	class FTriangleMeshImplicitObject;
+} // namespace Chaos
+
 UE_EXPERIMENTAL(5.8, "The new Chaos API is experimental")
 namespace UE::Physics
 {
@@ -31,12 +35,15 @@ namespace UE::Physics
 
 	class FTriangleMeshGeometry
 	{
-		using FTriangleMeshImplicitObjectPtr = Chaos::FTriangleMeshImplicitObjectPtr;
-
 	public:
 		using EPerFaceBoundsRepresentation = FTriangleMeshGeometrySetup::EPerFaceBoundsRepresentation;
 
-		UE_INTERNAL RIGIDPHYSICS_API FTriangleMeshGeometry(FTriangleMeshImplicitObjectPtr InImplicit);
+		RIGIDPHYSICS_API ~FTriangleMeshGeometry();
+		RIGIDPHYSICS_API FTriangleMeshGeometry(const FTriangleMeshGeometry&);
+		RIGIDPHYSICS_API FTriangleMeshGeometry(FTriangleMeshGeometry&&);
+		RIGIDPHYSICS_API FTriangleMeshGeometry& operator=(const FTriangleMeshGeometry&);
+		RIGIDPHYSICS_API FTriangleMeshGeometry& operator=(FTriangleMeshGeometry&&);
+		UE_INTERNAL RIGIDPHYSICS_API FTriangleMeshGeometry(TRefCountPtr<Chaos::FTriangleMeshImplicitObject> InImplicit);
 		RIGIDPHYSICS_API FTriangleMeshGeometry(FTriangleMeshGeometrySetup&& InSetup);
 
 		bool operator==(const FTriangleMeshGeometry&) const = default;
@@ -62,13 +69,13 @@ namespace UE::Physics
 
 		RIGIDPHYSICS_API FString ToString() const;
 
-		UE_INTERNAL RIGIDPHYSICS_API FTriangleMeshImplicitObjectPtr GetImplicit() const;
+		UE_INTERNAL RIGIDPHYSICS_API TRefCountPtr<Chaos::FTriangleMeshImplicitObject> GetImplicit() const;
 
 	private:
 		template <typename IndexType>
 		void BuildInternal(FTriangleMeshGeometrySetup&& InSetup);
 
-		Chaos::FTriangleMeshImplicitObjectPtr Implicit;
+		TRefCountPtr<Chaos::FTriangleMeshImplicitObject> Implicit;
 	};
 } // namespace UE::Physics
 
diff --git a/Engine/Source/Runtime/Experimental/RigidPhysics/Public/RigidPhysics/RigidShapeInstanceSetup.h b/Engine/Source/Runtime/Experimental/RigidPhysics/Public/RigidPhysics/RigidShapeInstanceSetup.h
--- a/Engine/Source/Runtime/Experimental/RigidPhysics/Public/RigidPhysics/RigidShapeInstanceSetup.h
+++ b/Engine/Source/Runtime/Experimental/RigidPhysics/Public/RigidPhysics/RigidShapeInstanceSetup.h
@@ -8,6 +8,7 @@
 #include "RigidPhysics/Geometry/AnyGeometry.h"
 #include "RigidPhysics/RigidMaterials.h"
 #include "BodySetupEnums.h"
+#include "Chaos/CollisionFilterData.h"
 
 UE_EXPERIMENTAL(5.8, "The new Chaos API is experimental")
 namespace UE::Physics
