Prevent unwritable values from creating malformed optionals or maps from JsonObjectGraph::Stringify

#rnx
#Synced-CL 52588543

diff --git a/Engine/Source/Runtime/Experimental/JsonObjectGraph/Private/JsonStringifyImpl.cpp b/Engine/Source/Runtime/Experimental/JsonObjectGraph/Private/JsonStringifyImpl.cpp
--- a/Engine/Source/Runtime/Experimental/JsonObjectGraph/Private/JsonStringifyImpl.cpp
+++ b/Engine/Source/Runtime/Experimental/JsonObjectGraph/Private/JsonStringifyImpl.cpp
@@ -923,10 +923,21 @@ void FJsonStringifyImpl::WriteMapToJson(const void* MapInstance, const FMapPrope
 		// age old question - how do you want to encode a tuple in json? I've
 		// chosen named tuple:
 		Writer->WriteObjectStartInline();
-		Writer->WriteIdentifierPrefix(UE_JSON_TMAP_KEY_KEY_TCHAR);
-		WriteValueToJson(Key, nullptr, KeyProperty);
-		Writer->WriteIdentifierPrefix(UE_JSON_TMAP_VALUE_KEY_TCHAR);
-		WriteValueToJson(Value, nullptr, ValueProperty);
+		
+		{
+			FPendingScope PropertyIdentifier(this, [Impl = this]()
+				{
+					Impl->Writer->WriteIdentifierPrefix(UE_JSON_TMAP_KEY_KEY_TCHAR);
+				});
+			WriteValueToJson(Key, nullptr, KeyProperty);
+		}
+		{
+			FPendingScope PropertyIdentifier(this, [Impl = this]()
+				{
+					Impl->Writer->WriteIdentifierPrefix(UE_JSON_TMAP_VALUE_KEY_TCHAR);
+				});
+			WriteValueToJson(Value, nullptr, ValueProperty);
+		}
 		Writer->WriteObjectEnd();
 	}
 
@@ -942,7 +953,10 @@ void FJsonStringifyImpl::WriteOptionalToJson(const void* OptionalInstnace, const
 	Writer->WriteObjectStartInline();
 	if (const void* ValueAddress = static_cast<const void*>(OptionalProperty->GetValuePointerForReadOrReplaceIfSet(OptionalInstnace)))
 	{
-		Writer->WriteIdentifierPrefix(UE_JSON_OPTIONAL_VALUE_KEY_TCHAR);
+		FPendingScope PropertyIdentifier(this, [Impl = this]()
+			{
+				Impl->Writer->WriteIdentifierPrefix(UE_JSON_OPTIONAL_VALUE_KEY_TCHAR);
+			});
 		WriteValueToJson(ValueAddress, nullptr, OptionalProperty->GetValueProperty());
 	}
 	Writer->WriteObjectEnd();
