Mark Chaos debug draw interface as UE_INTERNAL

#Synced-CL 52663992

diff --git a/Engine/Source/Runtime/Experimental/Chaos/Public/Chaos/ChaosDebugDrawInterface.h b/Engine/Source/Runtime/Experimental/Chaos/Public/Chaos/ChaosDebugDrawInterface.h
--- a/Engine/Source/Runtime/Experimental/Chaos/Public/Chaos/ChaosDebugDrawInterface.h
+++ b/Engine/Source/Runtime/Experimental/Chaos/Public/Chaos/ChaosDebugDrawInterface.h
@@ -7,7 +7,7 @@
 
 namespace Chaos
 {
-	class IDebugDrawInterface
+	class UE_INTERNAL IDebugDrawInterface
 	{
 	public:
 		// Define an interface for rendering a mesh
