Dataflow:
- Added mesh methods to EngineUtil

#Synced-CL 52659293

diff --git a/Engine/Source/Runtime/Experimental/Dataflow/Engine/Private/Dataflow/DataflowEngineUtil.cpp b/Engine/Source/Runtime/Experimental/Dataflow/Engine/Private/Dataflow/DataflowEngineUtil.cpp
--- a/Engine/Source/Runtime/Experimental/Dataflow/Engine/Private/Dataflow/DataflowEngineUtil.cpp
+++ b/Engine/Source/Runtime/Experimental/Dataflow/Engine/Private/Dataflow/DataflowEngineUtil.cpp
@@ -7,6 +7,9 @@
 #include "Dataflow/DataflowTypePolicy.h"
 #include "ReferenceSkeleton.h"
 #include "UObject/UnrealType.h"
+#include "DynamicMesh/DynamicMesh3.h"
+#include "DynamicMesh/MeshNormals.h"
+#include "DynamicMesh/DynamicMeshAttributeSet.h"
 
 namespace UE::Dataflow
 {
@@ -242,5 +245,143 @@ namespace UE::Dataflow
 
 			return LoadObject<UMaterialInterface>(nullptr, Path);
 		}
+
+		UMaterialInterface* GetDataflowCheckerBoardMaterial(const bool bUseVertexColor)
+		{
+			const FStringView Path = bUseVertexColor
+				? TEXT("/Engine/EditorMaterials/Dataflow/M_Dataflow_Vertex_CheckerBoard.M_Dataflow_Vertex_CheckerBoard")
+				: TEXT("/Engine/EditorMaterials/Dataflow/M_Dataflow_CheckerBoard.M_Dataflow_CheckerBoard");
+
+			return LoadObject<UMaterialInterface>(nullptr, Path);
+		}
+
+		UMaterialInterface* GetDataflowVertexColorMaterial()
+		{
+			const FStringView Path = TEXT("/Engine/EditorMaterials/Dataflow/DataflowVertexMaterial.DataflowVertexMaterial");
+
+			return LoadObject<UMaterialInterface>(nullptr, Path);
+		}
+
+	}
+
+	namespace Mesh
+	{
+		bool GetMeshVertices(const UE::Geometry::FDynamicMesh3* InMeshPtr, TArray<FVector3f>& OutVertices)
+		{
+			OutVertices.Reset();
+
+			if (InMeshPtr)
+			{
+				const int32 NumVertices = InMeshPtr->VertexCount();
+
+				OutVertices.Reserve(NumVertices);
+
+				for (const int32 VertexID : InMeshPtr->VertexIndicesItr())
+				{
+					if (InMeshPtr->IsVertex(VertexID))
+					{
+						OutVertices.Add((FVector3f)InMeshPtr->GetVertex(VertexID));
+					}
+				}
+
+				return true;
+			}
+
+			return false;
+		}
+
+		bool GetMeshVertexNormals(const UE::Geometry::FDynamicMesh3* InMeshPtr, TArray<FVector3f>& OutVertexNormals)
+		{
+			OutVertexNormals.Reset();
+
+			if (InMeshPtr)
+			{
+				const int32 NumVertices = InMeshPtr->VertexCount();
+
+				OutVertexNormals.Reserve(NumVertices);
+
+				if (InMeshPtr->HasVertexNormals())
+				{
+					for (const int32 VertexId : InMeshPtr->VertexIndicesItr())
+					{
+						if (InMeshPtr->IsVertex(VertexId))
+						{
+							OutVertexNormals.Add(InMeshPtr->GetVertexNormal(VertexId));
+						}
+					}
+				}
+				else if (InMeshPtr->HasAttributes() && InMeshPtr->Attributes()->PrimaryNormals())
+				{
+					const UE::Geometry::FDynamicMeshNormalOverlay* const NormalOverlay = InMeshPtr->Attributes()->PrimaryNormals();
+
+					TArray<int> OverlayElements;
+					for (const int32 VertexId : InMeshPtr->VertexIndicesItr())
+					{
+						if (InMeshPtr->IsVertex(VertexId))
+						{
+							OverlayElements.Reset();
+							NormalOverlay->GetVertexElements(VertexId, OverlayElements);
+
+							FVector3f AvgNormal(0.0f, 0.0f, 0.0f);
+							if (OverlayElements.Num() > 0)
+							{
+								for (int32 ElementID : OverlayElements)
+								{
+									AvgNormal += NormalOverlay->GetElement(ElementID);
+								}
+								AvgNormal /= (float)OverlayElements.Num();
+
+								OutVertexNormals.Add(AvgNormal);
+							}
+						}
+					}
+				}
+
+				return true;
+			}
+
+			return false;
+		}
+
+		bool GetMeshUVs(const UE::Geometry::FDynamicMesh3* InDynMeshPtr, TArray<FVector2f>& OutUVs, const int32 InUVLayer)
+		{
+			if (InDynMeshPtr && InDynMeshPtr->HasAttributes() && InDynMeshPtr->Attributes()->NumUVLayers() > 0)
+			{
+				OutUVs.Init(FVector2f::ZeroVector, InDynMeshPtr->MaxVertexID()); // RenderFacade->DynamicMesh code expects all vertices to have full UV sets.
+
+				int32 NumUVLayers = InDynMeshPtr->Attributes()->NumUVLayers();
+
+				if (InUVLayer >= NumUVLayers)
+				{
+					return false;
+				}
+
+				TArray<const UE::Geometry::FDynamicMeshUVOverlay*> UVLayers;
+				UVLayers.SetNum(NumUVLayers);
+
+				for (int32 Layer = 0; Layer < NumUVLayers; ++Layer)
+				{
+					UVLayers[Layer] = InDynMeshPtr->Attributes()->GetUVLayer(Layer);
+				}
+
+				for (const int32 TriangleIndex : InDynMeshPtr->TriangleIndicesItr())
+				{
+					const UE::Geometry::FIndex3i Tri = InDynMeshPtr->GetTriangle(TriangleIndex);
+
+					for (int32 TriangleVertexIndex = 0; TriangleVertexIndex < 3; ++TriangleVertexIndex)
+					{
+						const int32 VertexIndex = Tri[TriangleVertexIndex];
+						if (InDynMeshPtr->IsVertex(VertexIndex))
+						{
+							OutUVs[VertexIndex] = UVLayers[InUVLayer]->GetElementAtVertex(TriangleIndex, VertexIndex);
+						}
+					}
+				}
+
+				return true;
+			}
+
+			return false;
+		}
 	}
 }
diff --git a/Engine/Source/Runtime/Experimental/Dataflow/Engine/Public/Dataflow/DataflowEngineUtil.h b/Engine/Source/Runtime/Experimental/Dataflow/Engine/Public/Dataflow/DataflowEngineUtil.h
--- a/Engine/Source/Runtime/Experimental/Dataflow/Engine/Public/Dataflow/DataflowEngineUtil.h
+++ b/Engine/Source/Runtime/Experimental/Dataflow/Engine/Public/Dataflow/DataflowEngineUtil.h
@@ -24,6 +24,11 @@
 struct FDataflowConnection;
 struct FInstancedPropertyBag;
 
+namespace UE::Geometry
+{
+	class FDynamicMesh3;
+}
+
 namespace UE::Dataflow
 {
 	namespace Reflection
@@ -107,6 +112,47 @@ namespace UE::Dataflow
 	{
 		UE_API UMaterialInterface* GetDataflowColorMaterial(const bool bTwoSided = false);
 		UE_API UMaterialInterface* GetDataflowColorWireframeMaterial(const bool bTwoSided = false);
+		UE_API UMaterialInterface* GetDataflowCheckerBoardMaterial(const bool bUseVertexColor = false);
+		UE_API UMaterialInterface* GetDataflowVertexColorMaterial();
+	}
+
+	namespace Mesh
+	{
+		template <typename T>
+		T InterpolateAttribute(const FVector& InVertexA,
+			const FVector& InVertexB,
+			const FVector& InVertexC,
+			const FVector& InPointInTriangle,
+			const T& InAttributeAtVertexA,
+			const T& InAttributeAtVertexB,
+			const T& InAttributeAtVertexC,
+			const T& InDefaultValue)
+		{
+			double AreaTriangleABC = 0.5 * FVector::CrossProduct(InVertexB - InVertexA, InVertexC - InVertexA).Length();
+			double AreaTriangleABP = 0.5 * FVector::CrossProduct(InVertexA - InPointInTriangle, InVertexB - InPointInTriangle).Length();
+			double AreaTriangleACP = 0.5 * FVector::CrossProduct(InVertexA - InPointInTriangle, InVertexC - InPointInTriangle).Length();
+			double AreaTriangleBCP = 0.5 * FVector::CrossProduct(InVertexB - InPointInTriangle, InVertexC - InPointInTriangle).Length();
+
+			if (AreaTriangleABC > SMALL_NUMBER)
+			{
+				const double U = AreaTriangleBCP / AreaTriangleABC;
+				const double V = AreaTriangleACP / AreaTriangleABC;
+				const double W = AreaTriangleABP / AreaTriangleABC;
+
+				if (FMath::Abs(U + V + W - 1.0) < UE_SMALL_NUMBER)
+				{
+					return U * InAttributeAtVertexA + V * InAttributeAtVertexB + W * InAttributeAtVertexC;
+				}
+			}
+
+			return InDefaultValue;
+		}
+
+		UE_API bool GetMeshVertices(const UE::Geometry::FDynamicMesh3* InMeshPtr, TArray<FVector3f>& OutVertices);
+
+		UE_API bool GetMeshVertexNormals(const UE::Geometry::FDynamicMesh3* InMeshPtr, TArray<FVector3f>& OutVertexNormals);
+
+		UE_API bool GetMeshUVs(const UE::Geometry::FDynamicMesh3* InDynMeshPtr, TArray<FVector2f>& OutUVs, const int32 InUVLayer);
 	}
 }
 
