Add a new cvar 'iax.HttpClientOptimizeRanges' that enables/disables range request optimization in the http client (defaults to true).

#rnx

- The http client attempts to merge concurrent and overlapping range requests into single requests so we send as few ranges to the CDN as possible.
- This can be an optimization it also helps keep the results of the GET to be a bit more consistent as we have seen very different behaviors from some servers when facing concurrent, overlapping or out of order range requests.
- This cvar lets us disable the optimization pass in which case the client will send the requests in the exact form that the higher level code provided.

#Synced-CL 52609747

diff --git a/Engine/Source/Runtime/Experimental/IoStore/OnDemand/Private/OnDemandHttpClient.cpp b/Engine/Source/Runtime/Experimental/IoStore/OnDemand/Private/OnDemandHttpClient.cpp
--- a/Engine/Source/Runtime/Experimental/IoStore/OnDemand/Private/OnDemandHttpClient.cpp
+++ b/Engine/Source/Runtime/Experimental/IoStore/OnDemand/Private/OnDemandHttpClient.cpp
@@ -40,6 +40,13 @@ static FAutoConsoleVariableRef CVar_IaxAllowPartialContentRequests(
 	TEXT("Enable/disable the use of partial content http requests")
 );
 
+bool GHttpClientOptimizeRanges = true;
+static FAutoConsoleVariableRef CVar_IaxHttpClientOptimizeRanges(
+	TEXT("iax.HttpClientOptimizeRanges"),
+	GHttpClientOptimizeRanges,
+	TEXT("Should the http client attempt an optimization pass on ranged requests or process the requests as provided")
+);
+
 /**
  * Use CA root certificates from the system store if available.
  */
@@ -382,26 +389,26 @@ bool MergeOverlappingRanges(FIoOffsetAndLength& MergeDst, const FIoOffsetAndLeng
 	}
 }
 
-TArray<FIoOffsetAndLength> SortAndMergeRanges(TArray<FIoOffsetAndLength>& UnsortedRanges)
+TArray<FIoOffsetAndLength> SortAndMergeRanges(TArray<FIoOffsetAndLength>& RangesToMerge)
 {
-	if (UnsortedRanges.IsEmpty())
+	if (!GHttpClientOptimizeRanges || RangesToMerge.IsEmpty())
 	{
-		return TArray<FIoOffsetAndLength>();
+		return RangesToMerge;
 	}
 
-	UnsortedRanges.Sort([](const FIoOffsetAndLength& LHS, const FIoOffsetAndLength& RHS)
+	RangesToMerge.Sort([](const FIoOffsetAndLength& LHS, const FIoOffsetAndLength& RHS)
 	{
 		return LHS.GetOffset() < RHS.GetOffset();
 	});
 
 	TArray<FIoOffsetAndLength> SortedRanges;
-	SortedRanges.Add(UnsortedRanges[0]);
+	SortedRanges.Add(RangesToMerge[0]);
 
-	for (int32 Index = 1; Index < UnsortedRanges.Num(); ++Index)
+	for (int32 Index = 1; Index < RangesToMerge.Num(); ++Index)
 	{
-		if (!MergeOverlappingRanges(SortedRanges.Last(), UnsortedRanges[Index]))
+		if (!MergeOverlappingRanges(SortedRanges.Last(), RangesToMerge[Index]))
 		{
-			SortedRanges.Add(UnsortedRanges[Index]);
+			SortedRanges.Add(RangesToMerge[Index]);
 		}
 	}
 	return SortedRanges;
