Finish/clean up the code documentation for FResponseBody.

#rnx

#Synced-CL 52609747

diff --git a/Engine/Source/Runtime/Experimental/IoStore/OnDemand/Private/OnDemandHttpClient.h b/Engine/Source/Runtime/Experimental/IoStore/OnDemand/Private/OnDemandHttpClient.h
--- a/Engine/Source/Runtime/Experimental/IoStore/OnDemand/Private/OnDemandHttpClient.h
+++ b/Engine/Source/Runtime/Experimental/IoStore/OnDemand/Private/OnDemandHttpClient.h
@@ -95,21 +95,32 @@ public:
 
 	/** 
 	 * Returns the total length of the body that was downloaded.
-	 * Note that if the request was multi-part range then the body will include the
-	 * header for each range that was returned, so this size may be larger than the
-	 * sum of the requested ranges. If the whole body or a single range was requested
-	 * then 
+	 * 
+	 * Note that if the request was multi-part range then the body will include additional
+	 * overhead as each range has it's own small header and so the  size will be larger
+	 * than the sum of the requested ranged.
 	 */
 	uint64 GetSize() const
 	{
 		return Body.GetSize();
 	}
 
+	/**
+	 * Returns the entire body in a single buffer as it was returned from the CDN.
+	 * 
+	 * Note that if the request was multi-part range then this buffer will include
+	 * the headers for each range. It is up to the caller to manage that.
+	 */
 	FIoBuffer GetBody() const
 	{
 		return Body;
 	}
 
+	/**
+	 * Returns a view of the data for the given range.
+	 * 
+	 * The requested range must be within the ranges for the initial request.
+	 */
 	FIoBuffer GetData(const FIoOffsetAndLength& Range) const;
 
 private:
