Chaos Visual Debugger: Build fix

#Synced-CL 52666112

diff --git a/Engine/Source/Runtime/Experimental/Chaos/Public/ChaosVisualDebugger/ChaosVDStructCollectionMemWriterReader.h b/Engine/Source/Runtime/Experimental/Chaos/Public/ChaosVisualDebugger/ChaosVDStructCollectionMemWriterReader.h
--- a/Engine/Source/Runtime/Experimental/Chaos/Public/ChaosVisualDebugger/ChaosVDStructCollectionMemWriterReader.h
+++ b/Engine/Source/Runtime/Experimental/Chaos/Public/ChaosVisualDebugger/ChaosVDStructCollectionMemWriterReader.h
@@ -61,7 +61,7 @@ namespace Chaos::VisualDebugger
 	private:
 		void SerializeObjectRef(UObject* Obj)
 		{
-			UScriptStruct* ScriptStruct = Obj ? Cast<UScriptStruct>(Obj) : nullptr;
+			UScriptStruct* ScriptStruct = (Obj && Obj->IsA<UScriptStruct>()) ? static_cast<UScriptStruct*>(Obj) : nullptr;
 			bool bIsScriptStruct = (ScriptStruct != nullptr);
 			*this << bIsScriptStruct;
 			if (ScriptStruct)
