[SecureTraceTransport][TraceBasedDebuggers] Fix relay data loss on close and improve diagnostics

1. Fix FRelayTraceDataWriter::Close() discarding incomplete bunch.
When trace data accumulated in CurrentPendingBunch never reached the bunch-size threshold (default 4096 bytes), Close() silently discarded it.
Now Close() flushes the incomplete bunch into the queue and fires NewDataAvailableDelegate so the relay transport can send it.

2. Fix shutdown ordering in FRuntimeModule::StopRecording().
ClearRelayDataPumpDelegates() was called before RelayWriter->Close(), so the delegate fired by the flush had no listener.
Reordered so Close() runs while the pump delegate is still bound.

3. Add relay pipeline diagnostics.
RelayTraceDataWriter logs write call count, total bytes, bunches completed, pending queue size, and flushed bunch size at Close().
SecureSocketTraceRelayTransport logs bunch count and bytes per relay pump call (Log in Test builds, Verbose otherwise).

4. Improve TLS pre-verification failure logging.
When Context.native_handle() returns null, fall back to ERR_peek_last_error() to surface the OpenSSL error queue.

#Synced-CL 52680077

diff --git a/Engine/Source/Runtime/Experimental/TraceBasedDebuggers/Private/RelayTraceDataWriter.cpp b/Engine/Source/Runtime/Experimental/TraceBasedDebuggers/Private/RelayTraceDataWriter.cpp
--- a/Engine/Source/Runtime/Experimental/TraceBasedDebuggers/Private/RelayTraceDataWriter.cpp
+++ b/Engine/Source/Runtime/Experimental/TraceBasedDebuggers/Private/RelayTraceDataWriter.cpp
@@ -120,6 +120,9 @@ namespace UE::TraceBasedDebuggers
 			CurrentPendingBunch = CreateBunch();
 		}
 
+		++TotalWriteCalls;
+		TotalBytesWritten += Size;
+
 		uint32 RemainingSizeToCopy = Size;
 
 		while (RemainingSizeToCopy > 0)
@@ -131,6 +134,7 @@ namespace UE::TraceBasedDebuggers
 				QueuedBytesNum += CurrentPendingBunch->Num();
 				PendingBunchesToRelay.Enqueue(MoveTemp(CurrentPendingBunch));
 				++PendingBunchesCount;
+				++TotalBunchesCompleted;
 				CurrentPendingBunch = CreateBunch();
 				AvailableSpaceInCurrentBunch = MaxPendingBytesPerBunch;
 
@@ -185,13 +189,21 @@ namespace UE::TraceBasedDebuggers
 
 	void FRelayTraceDataWriter::Close()
 	{
-		int64 PendingBytes = QueuedBytesNum;
-		if (PendingBytes > 0)
+		// Flush the current incomplete bunch so it is available for the relay transport to send
+		const int32 CurrentBunchBytes = CurrentPendingBunch ? CurrentPendingBunch->Num() : 0;
+		if (CurrentPendingBunch && CurrentBunchBytes > 0)
 		{
-			// TODO: We should handle the case where we close the trace session but we still have data pending in the send queue (UE-320563).
-			// We can ignore it for now as it is not a critical failure
-			UE_LOG_REF(LogCategory, Warning, TEXT("Closing Trace Session with [%lld] bytes in the send queue. This data will be discarded"), PendingBytes);
+			QueuedBytesNum += CurrentBunchBytes;
+			PendingBunchesToRelay.Enqueue(MoveTemp(CurrentPendingBunch));
+			++PendingBunchesCount;
+			++TotalBunchesCompleted;
+
+			NewDataAvailableDelegate.ExecuteIfBound();
 		}
+		CurrentPendingBunch.Reset();
+
+		UE_LOG_REF(LogCategory, Log, TEXT("[%hs] Closing relay writer | WriteCalls [%d] | TotalBytesWritten [%lld] | BunchesCompleted [%d] | PendingQueueBytes [%lld] | FlushedIncompleteBunch [%d bytes]"),
+			__func__, TotalWriteCalls, TotalBytesWritten, TotalBunchesCompleted, static_cast<int64>(QueuedBytesNum), CurrentBunchBytes);
 
 		bClosed = true;
 	}
diff --git a/Engine/Source/Runtime/Experimental/TraceBasedDebuggers/Private/TraceBasedDebuggerRuntime.cpp b/Engine/Source/Runtime/Experimental/TraceBasedDebuggers/Private/TraceBasedDebuggerRuntime.cpp
--- a/Engine/Source/Runtime/Experimental/TraceBasedDebuggers/Private/TraceBasedDebuggerRuntime.cpp
+++ b/Engine/Source/Runtime/Experimental/TraceBasedDebuggers/Private/TraceBasedDebuggerRuntime.cpp
@@ -349,14 +349,17 @@ void FRuntimeModule::StopRecording()
 		RecordingTimerHandle.Reset();
 	}
 
-	ClearRelayDataPumpDelegates();
-
-
+	// Close the relay writer before clearing the data pump delegates so the final
+	// incomplete bunch flush in Close() can still trigger the relay transport via
+	// NewDataAvailableDelegate. Without this ordering, data accumulated in the
+	// current bunch that never reached the bunch-size threshold is silently lost.
 	if (RelayWriter)
 	{
 		RelayWriter->Close();
 	}
 
+	ClearRelayDataPumpDelegates();
+
 	ActiveRecordingRequesterId = FGuid{};
 	AccumulatedRecordingTime = 0.0f;
 
diff --git a/Engine/Source/Runtime/Experimental/TraceBasedDebuggers/Public/RelayTraceDataWriter.h b/Engine/Source/Runtime/Experimental/TraceBasedDebuggers/Public/RelayTraceDataWriter.h
--- a/Engine/Source/Runtime/Experimental/TraceBasedDebuggers/Public/RelayTraceDataWriter.h
+++ b/Engine/Source/Runtime/Experimental/TraceBasedDebuggers/Public/RelayTraceDataWriter.h
@@ -123,6 +123,11 @@ private:
 	uint32 MaxPendingBytesPerBunch = 0;
 	std::atomic<int64> QueuedBytesNum = 0;
 	uint64 MaxAllowedPendingBytes = 1073741824;
+
+	// Diagnostic counters — logged at Close() to diagnose relay data pipeline issues
+	int64 TotalBytesWritten = 0;
+	int32 TotalWriteCalls = 0;
+	int32 TotalBunchesCompleted = 0;
 };
 
 } // UE::TraceBasedDebuggers
