[Chaos][Parallel constraint solver] Link tasks in ConstraintGroupSolver to task in IslandGroupManager

#Synced-CL 52354501

diff --git a/Engine/Source/Runtime/Experimental/Chaos/Private/Chaos/Evolution/ConstraintGroupSolver.cpp b/Engine/Source/Runtime/Experimental/Chaos/Private/Chaos/Evolution/ConstraintGroupSolver.cpp
--- a/Engine/Source/Runtime/Experimental/Chaos/Private/Chaos/Evolution/ConstraintGroupSolver.cpp
+++ b/Engine/Source/Runtime/Experimental/Chaos/Private/Chaos/Evolution/ConstraintGroupSolver.cpp
@@ -91,7 +91,10 @@ namespace Chaos
 
 		void FPBDConstraintGroupSolver::AddConstraintsAndBodies()
 		{
-			AddConstraintsImpl();
+			{
+				QUICK_SCOPE_CYCLE_COUNTER(STAT_Evolution_AddConstraints);
+				AddConstraintsImpl();
+			}
 
 			for (int32 ContainerIndex = 0; ContainerIndex < ConstraintContainerSolvers.Num(); ++ContainerIndex)
 			{
@@ -99,10 +102,14 @@ namespace Chaos
 				{
 					if (CVars::Chaos_SolverPartitionManager_Partitioning)
 					{
+						QUICK_SCOPE_CYCLE_COUNTER(STAT_Evolution_PrepareSolverPartitions);
 						MaxSerialBatches = FMath::Max(MaxSerialBatches, ConstraintContainerSolvers[ContainerIndex]->PrepareSolverPartitions());
 					}
 
-					ConstraintContainerSolvers[ContainerIndex]->AddBodies(SolverBodyContainer);
+					{
+						QUICK_SCOPE_CYCLE_COUNTER(STAT_Evolution_AddBodies);
+						ConstraintContainerSolvers[ContainerIndex]->AddBodies(SolverBodyContainer);
+					}
 
 					TotalNumConstraints += ConstraintContainerSolvers[ContainerIndex]->GetNumConstraints();
 				}
@@ -580,7 +587,7 @@ namespace Chaos
 			}
 		}
 
-		void FPBDConstraintGroupSolver::ApplyConstraintsParallelTasks(const FReal Dt)
+		TArray<UE::Tasks::FTask> FPBDConstraintGroupSolver::ApplyConstraintsParallelTasks(const FReal Dt)
 		{
 			QUICK_SCOPE_CYCLE_COUNTER(STAT_ConstraintGroupSolver_ApplyConstraintsParallelTasks);
 
@@ -614,9 +621,7 @@ namespace Chaos
 			PreApplyProjectionConstraintsParallelTasks(Dt);
 			ApplyProjectionConstraintsParallelTasks(Dt);
 
-			// @todo(Chaos): We cannot call Wait from within a task on some platforms with low number of core.
-			// Need to chain this to the existing task-parallel constraint group solver.
-			UE::Tasks::Wait(ConstraintSolvePrerequisites.Last());
+			return ConstraintSolvePrerequisites.Last();
 		}
 
 		//
diff --git a/Engine/Source/Runtime/Experimental/Chaos/Private/Chaos/Island/IslandGroupManager.cpp b/Engine/Source/Runtime/Experimental/Chaos/Private/Chaos/Island/IslandGroupManager.cpp
--- a/Engine/Source/Runtime/Experimental/Chaos/Private/Chaos/Island/IslandGroupManager.cpp
+++ b/Engine/Source/Runtime/Experimental/Chaos/Private/Chaos/Island/IslandGroupManager.cpp
@@ -246,7 +246,59 @@ namespace Chaos
 			}
 		}
 
-		void FPBDIslandGroupManager::SolveGroupConstraints(int32 GroupIndex, const FReal Dt)
+		void FPBDIslandGroupManager::SolveGroupConstraintsSerial(int32 GroupIndex, const FReal Dt)
+		{
+			FPBDIslandConstraintGroupSolver* IslandGroup = GetGroup(GroupIndex);
+			QUICK_SCOPE_CYCLE_COUNTER(STAT_Evolution_ApplyConstraintsAllPhases);
+			{
+				SCOPE_CYCLE_COUNTER(STAT_Evolution_ApplyConstraintsPhase1);
+				CSV_SCOPED_ISLANDGROUP_TIMING_STAT(PerIslandSolve_ApplyTotalSerialized, GroupIndex);
+
+				IslandGroup->PreApplyPositionConstraints(Dt);
+				IslandGroup->ApplyPositionConstraints(Dt);
+			}
+			{
+				SCOPE_CYCLE_COUNTER(STAT_Evolution_ApplyConstraintsPhase2);
+				CSV_SCOPED_ISLANDGROUP_TIMING_STAT(PerIslandSolve_ApplyPushOutTotalSerialized, GroupIndex);
+
+				IslandGroup->PreApplyVelocityConstraints(Dt);
+				IslandGroup->ApplyVelocityConstraints(Dt);
+			}
+			{
+				SCOPE_CYCLE_COUNTER(STAT_Evolution_ApplyConstraintsPhase3);
+				CSV_SCOPED_ISLANDGROUP_TIMING_STAT(PerIslandSolve_ApplyProjectionTotalSerialized, GroupIndex);
+
+				IslandGroup->PreApplyProjectionConstraints(Dt);
+				IslandGroup->ApplyProjectionConstraints(Dt);
+			}
+		}
+
+		void FPBDIslandGroupManager::SolveGroupConstraintsParallelCalledFromTask(int32 GroupIndex, const FReal Dt)
+		{
+			QUICK_SCOPE_CYCLE_COUNTER(STAT_Evolution_ApplyConstraintsAllPhases);
+
+			FPBDIslandConstraintGroupSolver* IslandGroup = GetGroup(GroupIndex);
+			const TArray<UE::Tasks::FTask>& FinalTasks = IslandGroup->ApplyConstraintsParallelTasks(Dt);
+			
+			// AddNested() will prevent the calling task from terminating until the last constraint group solver task has finished.				
+			for (const UE::Tasks::FTask& Prerequisite : FinalTasks)
+			{
+				UE::Tasks::AddNested(Prerequisite);
+			}
+		}
+
+		void FPBDIslandGroupManager::SolveGroupConstraintsParallelNotCalledFromTask(int32 GroupIndex, const FReal Dt)
+		{
+			QUICK_SCOPE_CYCLE_COUNTER(STAT_Evolution_ApplyConstraintsAllPhases);
+
+			FPBDIslandConstraintGroupSolver* IslandGroup = GetGroup(GroupIndex);
+			const TArray<UE::Tasks::FTask>& FinalTasks = IslandGroup->ApplyConstraintsParallelTasks(Dt);
+
+			// Global Wait() needed if not called from within another task.
+			UE::Tasks::Wait(FinalTasks);
+		}
+
+		void FPBDIslandGroupManager::SolveGroupConstraintsCalledFromTask(int32 GroupIndex, const FReal Dt)
 		{
 			CSV_SCOPED_TIMING_STAT(Chaos, ApplyConstraints);
 
@@ -254,33 +306,27 @@ namespace Chaos
 
 			if (IslandGroup->SolveInParallel())
 			{
-				QUICK_SCOPE_CYCLE_COUNTER(STAT_Evolution_ApplyConstraintsAllPhases);
-				IslandGroup->ApplyConstraintsParallelTasks(Dt);
+				SolveGroupConstraintsParallelCalledFromTask(GroupIndex, Dt);
 			}
 			else
 			{
-				QUICK_SCOPE_CYCLE_COUNTER(STAT_Evolution_ApplyConstraintsAllPhases);
-				{
-					SCOPE_CYCLE_COUNTER(STAT_Evolution_ApplyConstraintsPhase1);
-					CSV_SCOPED_ISLANDGROUP_TIMING_STAT(PerIslandSolve_ApplyTotalSerialized, GroupIndex);
+				SolveGroupConstraintsSerial(GroupIndex, Dt);
+			}
+		}
 
-					IslandGroup->PreApplyPositionConstraints(Dt);
-					IslandGroup->ApplyPositionConstraints(Dt);
-				}
-				{
-					SCOPE_CYCLE_COUNTER(STAT_Evolution_ApplyConstraintsPhase2);
-					CSV_SCOPED_ISLANDGROUP_TIMING_STAT(PerIslandSolve_ApplyPushOutTotalSerialized, GroupIndex);
+		void FPBDIslandGroupManager::SolveGroupConstraintsNotCalledFromTask(int32 GroupIndex, const FReal Dt)
+		{
+			CSV_SCOPED_TIMING_STAT(Chaos, ApplyConstraints);
 
-					IslandGroup->PreApplyVelocityConstraints(Dt);
-					IslandGroup->ApplyVelocityConstraints(Dt);
-				}
-				{
-					SCOPE_CYCLE_COUNTER(STAT_Evolution_ApplyConstraintsPhase3);
-					CSV_SCOPED_ISLANDGROUP_TIMING_STAT(PerIslandSolve_ApplyProjectionTotalSerialized, GroupIndex);
+			FPBDIslandConstraintGroupSolver* IslandGroup = GetGroup(GroupIndex);
 
-					IslandGroup->PreApplyProjectionConstraints(Dt);
-					IslandGroup->ApplyProjectionConstraints(Dt);
-				}
+			if (IslandGroup->SolveInParallel())
+			{
+				SolveGroupConstraintsParallelNotCalledFromTask(GroupIndex, Dt);
+			}
+			else
+			{
+				SolveGroupConstraintsSerial(GroupIndex, Dt);
 			}
 		}
 
@@ -301,7 +347,7 @@ namespace Chaos
 					IslandGroup->GatherConstraints(Dt, 0, IslandGroup->GetNumSolverConstraints());
 				}
 
-				SolveGroupConstraints(GroupIndex, Dt);
+				SolveGroupConstraintsNotCalledFromTask(GroupIndex, Dt);
 
 				{
 					SCOPE_CYCLE_COUNTER(STAT_Evolution_Scatter);
@@ -366,7 +412,7 @@ namespace Chaos
 				{
 					FPBDIslandConstraintGroupSolver* IslandGroup = GetGroup(GroupIndex);
 
-					SolveGroupConstraints(GroupIndex, Dt);
+					SolveGroupConstraintsNotCalledFromTask(GroupIndex, Dt);
 				});
 
 			{
@@ -396,36 +442,59 @@ namespace Chaos
 			}
 		}
 
+		template<typename Lambda>
+		TArray<UE::Tasks::FTask> FPBDIslandGroupManager::DispatchBatchedTasks(
+			Lambda BatchedTask, 
+			const int32 NumTasks, 
+			const int32 NumItemsPerTask,
+			const int32 NumItems,
+			const FReal Dt,
+			const TArray<UE::Tasks::FTask>& Prereqs
+			)
+		{
+			TArray<UE::Tasks::FTask> BatchedTasks;
+			BatchedTasks.Reserve(NumTasks);
+			for (int32 BeginIndex = 0; BeginIndex < NumItems; BeginIndex += NumItemsPerTask)
+			{
+				const int32 EndIndex = FMath::Min(BeginIndex + NumItemsPerTask, NumItems);
+				BatchedTasks.Add(UE::Tasks::Launch(UE_SOURCE_LOCATION,
+					[BatchedTask, BeginIndex, EndIndex, Dt]()
+					{
+						BatchedTask(Dt, BeginIndex, EndIndex);
+					},
+					Prereqs, LowLevelTasks::ETaskPriority::High));
+			}
+			return MoveTemp(BatchedTasks);
+		}
 
 		void FPBDIslandGroupManager::SolveParallelTasks(const FReal Dt)
 		{
 			SCOPE_CYCLE_COUNTER(STAT_SolveConstraints);
 
-			FGraphEventArray GroupSolveEvents;
-			GroupSolveEvents.Reserve(NumActiveGroups);
+			TArray<UE::Tasks::FTask> GroupSolveTasks;
+			GroupSolveTasks.Reserve(NumActiveGroups);
 
-			// Each IslandGroup is independent of the otyher so they can all be solved in parallel.
+			// Each IslandGroup is independent of the others so they can all be solved in parallel.
 			// In addition, the Gather and Scatter work within each group is highly parallelizable.
 			//
-			// In order to utilize all workers when we have a small number of groups some are much
-			// larger than others, we allow each group to generate multiple gather/scatter tasks. 
-			// The number of tasks each group generates depends on how many constraints it has and 
-			// TargetNumConstraintsPerTask, which was calculated in BuildGroups such that we end 
+			// In order to utilize all workers when we have a small number of groups or some are much
+			// larger than others, we allow each group to generate multiple gather/scatter tasks.
+			// The number of tasks each group generates depends on how many constraints it has and
+			// TargetNumConstraintsPerTask, which was calculated in BuildGroups such that we end
 			// up with the right number of gather tasks below.
 			//
 			for (int32 GroupIndex = 0; GroupIndex < NumActiveGroups; ++GroupIndex)
 			{
 				// NOTE: This task spawns other child tasks. We must run AddConstraintsAndBodies before we know how many
 				// Gather and Scatter tasks we want to create.
-				FGraphEventRef GroupSolveEvent = FFunctionGraphTask::CreateAndDispatchWhenReady(
-					[this, GroupIndex, Dt](ENamedThreads::Type CurrentThread, const FGraphEventRef& MyCompletionGraphEvent)
+				UE::Tasks::FTask GroupSolveTask = UE::Tasks::Launch(UE_SOURCE_LOCATION,
+					[this, GroupIndex, Dt]()
 					{
 						FPBDIslandConstraintGroupSolver* IslandGroup = GetGroup(GroupIndex);
 
-						// Add bodies and constraints to the solvers but do no real work. Cannot be further parallelized
+						// Add bodies and constraints to the solvers but do no real work. Cannot be further parallelized.
 						{
 							SCOPE_CYCLE_COUNTER(STAT_Evolution_Gather_Prepare);
-
 							IslandGroup->AddConstraintsAndBodies();
 						}
 
@@ -439,128 +508,107 @@ namespace Chaos
 						const int32 NumConstraintGatherTasks = FMath::Max(1, FMath::DivideAndRoundNearest(NumConstraints, TargetNumConstraintsPerTask));
 						const int32 NumConstraintsPerTask = FMath::DivideAndRoundUp(NumConstraints, NumConstraintGatherTasks);
 
-						// Gather Bodies. Can run as wide as we like
-						FGraphEventArray GatherBodyEvents;
-						GatherBodyEvents.Reserve(NumBodyGatherTasks);
-						for (int32 BodyBeginIndex = 0; BodyBeginIndex < NumBodies; BodyBeginIndex += NumBodiesPerTask)
-						{
-							const int32 BodyEndIndex = FMath::Min(BodyBeginIndex + NumBodiesPerTask, NumBodies);
-
-							FGraphEventRef GatherBodyEvent = FFunctionGraphTask::CreateAndDispatchWhenReady(
-								[IslandGroup, BodyBeginIndex, BodyEndIndex, Dt]()
-								{
-									CSV_SCOPED_TIMING_STAT(PhysicsVerbose, StepSolver_Gather);
-
-									IslandGroup->GatherBodies(Dt, BodyBeginIndex, BodyEndIndex);
-								}
-								, GET_STATID(STAT_Evolution_Gather_Bodies));
-
-							GatherBodyEvents.Add(GatherBodyEvent);
-						}
-
-						// Wait for Gather Bodies
-						FGraphEventRef GatherBodiesCompletionEvent = TGraphTask<FNullGraphTask>::CreateTask(&GatherBodyEvents).ConstructAndDispatchWhenReady(TStatId(), ENamedThreads::AnyThread);
-
-						// Gather Constraints after body gather is complete. Can run as wide as wek like
-						FGraphEventArray GatherConstraintEvents;
-						GatherConstraintEvents.Reserve(NumConstraintGatherTasks);
-						for (int32 ConstraintBeginIndex = 0; ConstraintBeginIndex < NumConstraints; ConstraintBeginIndex += NumConstraintsPerTask)
-						{
-							const int32 ConstraintEndIndex = FMath::Min(ConstraintBeginIndex + NumConstraintsPerTask, NumConstraints);
-
-							FGraphEventRef GatherConstraintsEvent = FFunctionGraphTask::CreateAndDispatchWhenReady(
-								[this, IslandGroup, ConstraintBeginIndex, ConstraintEndIndex, Dt]()
-								{
-									CSV_SCOPED_TIMING_STAT(PhysicsVerbose, StepSolver_Gather);
-
-									IslandGroup->GatherConstraints(Dt, ConstraintBeginIndex, ConstraintEndIndex);
-								}
-								, GET_STATID(STAT_Evolution_Gather_Constraints)
-								, GatherBodiesCompletionEvent);
-
-							GatherConstraintEvents.Add(GatherConstraintsEvent);
-						}
-
-						// Wait for Gather Constraints
-						FGraphEventRef GatherConstraintsCompletionEvent = TGraphTask<FNullGraphTask>::CreateTask(&GatherConstraintEvents).ConstructAndDispatchWhenReady(TStatId(), ENamedThreads::AnyThread);
-
-						// Solve Constraints after gathering is complete. Single task - no further parallelization (without coloring)
-						// @todo(chaos): add coloring support
-						FGraphEventRef SolveConstraintsCompletionEvent = FFunctionGraphTask::CreateAndDispatchWhenReady(
-							[this, GroupIndex, Dt]()
+						// Gather Bodies. Can run as wide as we like.
+						TArray<UE::Tasks::FTask> GatherBodyTasks = DispatchBatchedTasks(
+							[IslandGroup](const FReal Dt, const int32 BeginIndex, const int32 EndIndex) 
 							{
-								SolveGroupConstraints(GroupIndex, Dt);
+								SCOPE_CYCLE_COUNTER(STAT_Evolution_Gather_Bodies);
+								CSV_SCOPED_TIMING_STAT(PhysicsVerbose, StepSolver_Gather);
+								IslandGroup->GatherBodies(Dt, BeginIndex, EndIndex);
+							}, 
+							NumBodyGatherTasks, 
+							NumBodiesPerTask, 
+							NumBodies, 
+							Dt, 
+							{ 
+								UE::Tasks::FTask() 
 							}
-							, TStatId()
-							, GatherConstraintsCompletionEvent);
+						);
+						
+						// Gather Constraints after body gather is complete. Can run as wide as we like.
+						// Passing GatherBodyTasks directly as prerequisites replaces the FNullGraphTask fan-in.
+						TArray<UE::Tasks::FTask> GatherConstraintTasks = DispatchBatchedTasks(
+							[IslandGroup](const FReal Dt, const int32 BeginIndex, const int32 EndIndex)
+							{
+								SCOPE_CYCLE_COUNTER(STAT_Evolution_Gather_Constraints);
+								CSV_SCOPED_TIMING_STAT(PhysicsVerbose, StepSolver_Gather);
+								IslandGroup->GatherConstraints(Dt, BeginIndex, EndIndex);
+							},
+							NumConstraintGatherTasks,
+							NumConstraintsPerTask,
+							NumConstraints,
+							Dt,
+							GatherBodyTasks
+						);
+
+						// Solve Constraints after gathering is complete. 
+						// Can be further parallelized with SolveTask if IslandGroup->SolveInParallel().
+						UE::Tasks::FTask SolveTask = UE::Tasks::Launch(UE_SOURCE_LOCATION,
+							[this, GroupIndex, Dt]()
+							{
+								SolveGroupConstraintsCalledFromTask(GroupIndex, Dt);
+							},
+							GatherConstraintTasks, LowLevelTasks::ETaskPriority::High);
 
 						// Scatter Constraints after the solve completes.
-
-						FGraphEventArray ScatterConstraintEvents;
-						ScatterConstraintEvents.Reserve(NumConstraintGatherTasks);
-
-						for (int32 ConstraintBeginIndex = 0; ConstraintBeginIndex < NumConstraints; ConstraintBeginIndex += NumConstraintsPerTask)
-						{
-							const int32 ConstraintEndIndex = FMath::Min(ConstraintBeginIndex + NumConstraintsPerTask, NumConstraints);
-
-							FGraphEventRef ScatterConstraintsEvent = FFunctionGraphTask::CreateAndDispatchWhenReady(
-								[this, IslandGroup, ConstraintBeginIndex, ConstraintEndIndex, Dt]()
-								{
-									CSV_SCOPED_TIMING_STAT(PhysicsVerbose, StepSolver_Scatter);
-
-									IslandGroup->ScatterConstraints(Dt, ConstraintBeginIndex, ConstraintEndIndex);
-								}
-								, GET_STATID(STAT_Evolution_Scatter_Constraints)
-								, SolveConstraintsCompletionEvent);
-
-							ScatterConstraintEvents.Add(ScatterConstraintsEvent);
-						}
+						TArray<UE::Tasks::FTask> ScatterConstraintTasks = DispatchBatchedTasks(
+							[IslandGroup](const FReal Dt, const int32 BeginIndex, const int32 EndIndex)
+							{
+								SCOPE_CYCLE_COUNTER(STAT_Evolution_Scatter_Constraints);
+								CSV_SCOPED_TIMING_STAT(PhysicsVerbose, StepSolver_Scatter);
+								IslandGroup->ScatterConstraints(Dt, BeginIndex, EndIndex);
+							},
+							NumConstraintGatherTasks,
+							NumConstraintsPerTask,
+							NumConstraints,
+							Dt,
+							{ 
+								SolveTask 
+							}
+						);
 
 						// Scatter Bodies.
-						// 
-						// Note: Executed after ScatterConstraints to avoid race condition given that ScatterConstraints reads FSolverBody P/Q and ScatterBodies writes to them.
-
-						FGraphEventRef BodyScatterDependency = TGraphTask<FNullGraphTask>::CreateTask(&ScatterConstraintEvents).ConstructAndDispatchWhenReady(TStatId(), ENamedThreads::AnyThread);
-
-						FGraphEventArray ScatterBodyEvents;
-						ScatterBodyEvents.Reserve(NumBodyGatherTasks);
-
-						for (int32 BodyBeginIndex = 0; BodyBeginIndex < NumBodies; BodyBeginIndex += NumBodiesPerTask)
-						{
-							const int32 BodyEndIndex = FMath::Min(BodyBeginIndex + NumBodiesPerTask, NumBodies);
-
-							FGraphEventRef ScatterBodyEvent = FFunctionGraphTask::CreateAndDispatchWhenReady(
-								[this, IslandGroup, BodyBeginIndex, BodyEndIndex, Dt]()
-								{
-									CSV_SCOPED_TIMING_STAT(PhysicsVerbose, StepSolver_Scatter);
-
-									IslandGroup->ScatterBodies(Dt, BodyBeginIndex, BodyEndIndex);
-								}
-								, GET_STATID(STAT_Evolution_Scatter_Bodies)
-								, BodyScatterDependency);
-
-							ScatterBodyEvents.Add(ScatterBodyEvent);
-						}
-
-						// Completion event
-						if (ScatterBodyEvents.IsEmpty())
+						//
+						// Note: Executed after ScatterConstraints to avoid race condition given that ScatterConstraints
+						// reads FSolverBody P/Q and ScatterBodies writes to them.
+						TArray<UE::Tasks::FTask> ScatterBodyTasks = DispatchBatchedTasks(
+							[IslandGroup](const FReal Dt, const int32 BeginIndex, const int32 EndIndex)
+							{
+								SCOPE_CYCLE_COUNTER(STAT_Evolution_Scatter_Bodies);
+								CSV_SCOPED_TIMING_STAT(PhysicsVerbose, StepSolver_Scatter);
+								IslandGroup->ScatterBodies(Dt, BeginIndex, EndIndex);
+							},
+							NumBodyGatherTasks,
+							NumBodiesPerTask,
+							NumBodies,
+							Dt,
+							ScatterConstraintTasks
+						);
+
+						// Add nested task dependency so that GroupSolveTask can only finish when all internally called tasks have finished.
+						if (ScatterBodyTasks.IsEmpty())
 						{
-							// NOTE: If no solver bodies present, ScatterBodyEvents will be empty.
-							// Our completion event will need to depend on the previous task to wait until all constraint scatter tasks have finished.
-							MyCompletionGraphEvent->DontCompleteUntil(BodyScatterDependency);
+							// NOTE: If no solver bodies present, ScatterBodyTasks will be empty, link to ScatterConstraintTasks instead.
+							for (UE::Tasks::FTask& ScatterConstraintTask : ScatterConstraintTasks)
+							{
+								UE::Tasks::AddNested(ScatterConstraintTask);
+							}
 						}
 						else
 						{
-							FGraphEventRef ScatterCompletionEvent = TGraphTask<FNullGraphTask>::CreateTask(&ScatterBodyEvents).ConstructAndDispatchWhenReady(TStatId(), ENamedThreads::AnyThread);
-							MyCompletionGraphEvent->DontCompleteUntil(ScatterCompletionEvent);
+							for (UE::Tasks::FTask& ScatterBodyTask : ScatterBodyTasks)
+							{
+								UE::Tasks::AddNested(ScatterBodyTask);
+							}
 						}
-					});
+					},
+					LowLevelTasks::ETaskPriority::High);
 
-				GroupSolveEvents.Add(GroupSolveEvent);
+				GroupSolveTasks.Add(GroupSolveTask);
 			}
 
-			// Wait for all tasks to complete and help out
-			FTaskGraphInterface::Get().WaitUntilTasksComplete(GroupSolveEvents);
+			// Wait for all tasks to complete.
+			UE::Tasks::Wait(GroupSolveTasks);
 		}
 
 		void FPBDIslandGroupManager::SetIterationSettings(const FIterationSettings& InIterations)
diff --git a/Engine/Source/Runtime/Experimental/Chaos/Public/Chaos/Evolution/ConstraintGroupSolver.h b/Engine/Source/Runtime/Experimental/Chaos/Public/Chaos/Evolution/ConstraintGroupSolver.h
--- a/Engine/Source/Runtime/Experimental/Chaos/Public/Chaos/Evolution/ConstraintGroupSolver.h
+++ b/Engine/Source/Runtime/Experimental/Chaos/Public/Chaos/Evolution/ConstraintGroupSolver.h
@@ -140,8 +140,14 @@ namespace Chaos
 			 * This is equivalent to calling the public pre-apply and apply constraints functions,
 			 * however, it parallelizes the constraint solve by scheduling and executing tasks
 			 * provided their respective prerequisite tasks.
+			 * 
+			 * The function returns a TArray<UE::Tasks::FTask> of the tasks to be executed last.
+			 * The caller of this function must do one of the following:
+			 *   1) if there will be more tasks launched, use the returned tasks as a prerequisites for the next task
+			 *   2) if these were the last tasks and this is NOT called from within task, call UE::Tasks::Wait.
+			 *   3) if these were the last tasks and this IS called from within task, call UE::Tasks::AddNested.
 			*/
-			CHAOS_API void ApplyConstraintsParallelTasks(const FReal Dt);
+			CHAOS_API TArray<UE::Tasks::FTask> ApplyConstraintsParallelTasks(const FReal Dt);
 
 			/**
 			 * Push results from all solver bodies back to their respective particles.
diff --git a/Engine/Source/Runtime/Experimental/Chaos/Public/Chaos/Island/IslandGroupManager.h b/Engine/Source/Runtime/Experimental/Chaos/Public/Chaos/Island/IslandGroupManager.h
--- a/Engine/Source/Runtime/Experimental/Chaos/Public/Chaos/Island/IslandGroupManager.h
+++ b/Engine/Source/Runtime/Experimental/Chaos/Public/Chaos/Island/IslandGroupManager.h
@@ -162,8 +162,23 @@ namespace Chaos
 			CHAOS_API void SolveParallelTasks(const FReal Dt);
 
 			CHAOS_API void BuildGatherBatches(TArray<FIslandGroupRange>& BodyRanges, TArray<FIslandGroupRange>& ConstraintRanges);
-			CHAOS_API void SolveGroupConstraints(const int32 GroupIndex, const FReal Dt);
 
+			CHAOS_API void SolveGroupConstraintsSerial(const int32 GroupIndex, const FReal Dt);
+			CHAOS_API void SolveGroupConstraintsParallelCalledFromTask(const int32 GroupIndex, const FReal Dt);
+			CHAOS_API void SolveGroupConstraintsParallelNotCalledFromTask(const int32 GroupIndex, const FReal Dt);
+
+			CHAOS_API void SolveGroupConstraintsCalledFromTask(const int32 GroupIndex, const FReal Dt);
+			CHAOS_API void SolveGroupConstraintsNotCalledFromTask(const int32 GroupIndex, const FReal Dt);
+
+			template<typename Lambda>
+			TArray<UE::Tasks::FTask> DispatchBatchedTasks(
+				Lambda BatchedTask,
+				const int32 NumTasks,
+				const int32 NumItemsPerTask,
+				const int32 NumItems,
+				const FReal Dt,
+				const TArray<UE::Tasks::FTask>& Prereqs
+			);
 
 			FPBDIslandManager& IslandManager;
 			TArray<TUniquePtr<FPBDIslandConstraintGroupSolver>> IslandGroups;
