Improve error validation when parsing the returned ranged from a multipart ranged request.

#rnx

- We don't really support values > 32bit as it would be too slow to process, so make sure that values returned are within that range.
- Error if the start and end bytes are the wrong way around, although that would indicate a server error.
- Add a temp check to make sure that the range is within int32 size. This can be removed when the parsing is fixed to use the full 32bit range.
#Synced-CL 52609747

diff --git a/Engine/Source/Runtime/Experimental/IoStore/OnDemand/Private/OnDemandHttpClient.cpp b/Engine/Source/Runtime/Experimental/IoStore/OnDemand/Private/OnDemandHttpClient.cpp
--- a/Engine/Source/Runtime/Experimental/IoStore/OnDemand/Private/OnDemandHttpClient.cpp
+++ b/Engine/Source/Runtime/Experimental/IoStore/OnDemand/Private/OnDemandHttpClient.cpp
@@ -168,16 +168,13 @@ FAnsiStringView ExtractTextLine(FAnsiStringView View)
 
 FAnsiStringView ExtractData(FAnsiStringView View, FAnsiStringView Boundary, TPair<uint32, uint32> ByteRange)
 {
-	int32 DataLength = (int32)(ByteRange.Value - ByteRange.Key) + 1; // TODO -  Safer casting/truncation?
-	FAnsiStringView Data = View.Left(DataLength);
-
-	// TODO - Validate?
-
-	return Data;
+	// We know that ByteRange.Value >= ByteRange.Key as that was validated (in ParseContentRange) when the byte range was parsed
+	int32 DataLength = (int32)(ByteRange.Value - ByteRange.Key) + 1;
+	return View.Left(DataLength);
 }
 
 // TODO - LexFromString (need FAnsiStringView support)
-void ExtractFromString(uint32& OutValue, const FAnsiStringView& InString)
+void ExtractFromString(uint64& OutValue, const FAnsiStringView& InString)
 {
 	TAnsiStringBuilder<64> Builder;
 	LexFromString(OutValue, *(Builder << InString));
@@ -185,20 +182,24 @@ void ExtractFromString(uint32& OutValue, const FAnsiStringView& InString)
 
 FIoStatus ParseContentRange(FAnsiStringView Line, TPair<uint32,uint32>& ByteRange)
 {
-	// Content-Range: bytes 1000-1023/
-
 	constexpr FAnsiStringView Bytes = ANSITEXTVIEW("bytes ");
 
 	int32 BytesIndex = Line.Find(Bytes);
 	if (BytesIndex == INDEX_NONE)
 	{
-		return FIoStatus(EIoErrorCode::HttpReadError, TEXT("Multirange request: Unable to bytes for 'Content-Range'"));
+		return FIoStatus(EIoErrorCode::HttpReadError, TEXT("Multirange request: Unable to find bytes for 'Content-Range'"));
 	}
 
 	FAnsiStringView DataView = Line.RightChop(BytesIndex + Bytes.Len());
 	DataView.TrimEndInline();
 
-	ExtractFromString(ByteRange.Key, DataView);
+	uint64 StartByte = 0;
+	ExtractFromString(StartByte, DataView);
+
+	if (StartByte > MAX_uint32)
+	{
+		return FIoStatus(EIoErrorCode::HttpReadError, *WriteToString<32>(TEXT("Multirange request: Start byte is too large: "), StartByte));
+	}
 
 	int32 SecondByteOffset = 0;
 	if (!DataView.FindChar('-', SecondByteOffset))
@@ -208,9 +209,26 @@ FIoStatus ParseContentRange(FAnsiStringView Line, TPair<uint32,uint32>& ByteRang
 
 	DataView.RightChopInline(SecondByteOffset + 1);
 
-	ExtractFromString(ByteRange.Value, DataView);
+	uint64 EndByte = 0;
+	ExtractFromString(EndByte, DataView);
+
+	if (EndByte > MAX_uint32)
+	{
+		return FIoStatus(EIoErrorCode::HttpReadError, *WriteToString<32>(TEXT("Multirange request: End byte is too large: "), EndByte));
+	}
+
+	if (EndByte < StartByte)
+	{
+		return FIoStatus(EIoErrorCode::HttpReadError, TEXT("Multirange request: Malformed byte range for 'Content-Range'"));
+	}
+
+	if ((EndByte - StartByte) + 1 > MAX_int32) // Temp validation until string view (int32 limit) is replaced.
+	{
+		return FIoStatus(EIoErrorCode::HttpReadError, TEXT("Multirange request:Byte range is too large"));
+	}
 
-	// TODO - Better validation that we parsed byte range values
+	ByteRange.Key = (uint32)StartByte;
+	ByteRange.Value = (uint32)EndByte;
 
 	return EIoErrorCode::Ok;
 }
