[CVD] Fix frame / stage issues with resimulation.

#Synced-CL 52722506

diff --git a/Engine/Source/Runtime/Experimental/Chaos/Private/Chaos/Framework/PhysicsSolverBase.cpp b/Engine/Source/Runtime/Experimental/Chaos/Private/Chaos/Framework/PhysicsSolverBase.cpp
--- a/Engine/Source/Runtime/Experimental/Chaos/Private/Chaos/Framework/PhysicsSolverBase.cpp
+++ b/Engine/Source/Runtime/Experimental/Chaos/Private/Chaos/Framework/PhysicsSolverBase.cpp
@@ -219,7 +219,18 @@ namespace Chaos
 
 	void FPhysicsSolverAdvanceTask::AdvanceSolver()
 	{
-		CVD_SCOPE_TRACE_SOLVER_FRAME(FPhysicsSolverBase, Solver);
+#if WITH_CHAOS_VISUAL_DEBUGGER
+		// TOptional lets us conditionally construct the scoped CVD frame tracer while
+		// keeping it alive for the entire AdvanceSolver() scope. During resimulation the CVD
+		// frame is managed by ConditionalApplyRewind_Internal instead, so that all
+		// resim stages (RewindToFrame, StepNonResimParticles, ApplyTargets, solver
+		// evolution, etc.) appear as stages within a single frame per physics step.
+		TOptional<FChaosVDScopeSolverFrame<FPhysicsSolverBase>> ScopeSolverFrame;
+		if (!Solver.IsResimming())
+		{
+			ScopeSolverFrame.Emplace(Solver);
+		}
+#endif
 		LLM_SCOPE(ELLMTag::ChaosUpdate);
 		SCOPE_CYCLE_COUNTER(STAT_ChaosTick);
 		CSV_SCOPED_TIMING_STAT_EXCLUSIVE(Physics);
diff --git a/Engine/Source/Runtime/Experimental/Chaos/Private/PBDRigidsSolver.cpp b/Engine/Source/Runtime/Experimental/Chaos/Private/PBDRigidsSolver.cpp
--- a/Engine/Source/Runtime/Experimental/Chaos/Private/PBDRigidsSolver.cpp
+++ b/Engine/Source/Runtime/Experimental/Chaos/Private/PBDRigidsSolver.cpp
@@ -2219,15 +2219,23 @@ namespace Chaos
 				{
 					if (MRewindData->AllowRewindToFrame(ResimStep))
 					{
+						SetIsResimming(true);
+						CurrentFrame = ResimStep;
+
+#if WITH_CHAOS_VISUAL_DEBUGGER
+						const bool bIsCVDTracing = FChaosVisualDebuggerTrace::IsTracing();
+						const bool bIsTracingResimDetails = bIsCVDTracing && CVDDC_ResimDetails->IsChannelEnabled();
+						if (bIsCVDTracing)
+						{
+							CVD_TRACE_SOLVER_START_FRAME(FPhysicsSolverBase, *this);
+						}
+#endif
 						MRewindData->RewindToFrame(ResimStep);
 
 #if DEBUG_REWIND_DATA
 						UE_LOGF(LogChaos, Warning, "CLIENT | PT | ConditionalApplyRewind_Internal | PERFORMING RESIMULATION | Resim From Frame = %d | Num Steps = %d | To Current Frame: %d", ResimStep, NumResimSteps, CurrentFrame);
 #endif
 
-						SetIsResimming(true);
-						CurrentFrame = ResimStep;
-
 						FPushPhysicsData* ConsumerData = MarshallingManager.GetConsumerData_Internal();
 						TArray<FPushPhysicsData*> RecordedPushData = MarshallingManager.StealHistory_Internal(NumResimSteps);
 						bool bFirst = true;
@@ -2237,25 +2245,18 @@ namespace Chaos
 						{
 							QUICK_SCOPE_CYCLE_COUNTER(ChaosRewindAndResim_Resim);
 
-#if WITH_CHAOS_VISUAL_DEBUGGER
-							bool bIsTracingResimDetails = (FChaosVisualDebuggerTrace::IsTracing() && CVDDC_ResimDetails->IsChannelEnabled());
-							if (bIsTracingResimDetails)
-							{
-								// This frame is redundant with the one in AdvanceSolver, but if we don't do this, subsequent particle traces and steps will appear as the last steps of the current frame we rewind from
-								// and CVD will not correctly advance the server timeline to the frame we are landing on, making it harder to compare the target set locally with server state at the same frame
-								CVD_TRACE_SOLVER_START_FRAME(FPhysicsSolverBase, *this);
-								CVD_SCOPE_TRACE_SOLVER_STEP(CVDDC_ResimDetails, TEXT("Rewind To Frame"));
-								MRewindData->TraceParticlesModifiedDuringRewind();
-							}
-#endif
-
 							// Do rollback as necessary
 							for (int32 Step = ResimStep; Step <= LastStep; ++Step)
 							{
 #if WITH_CHAOS_VISUAL_DEBUGGER
-								if (!bFirst && bIsTracingResimDetails)
+								if (bIsCVDTracing && !bFirst)
 								{
+									CVD_TRACE_SOLVER_END_FRAME(FPhysicsSolverBase, *this);
 									CVD_TRACE_SOLVER_START_FRAME(FPhysicsSolverBase, *this);
+								}
+
+								if (bIsTracingResimDetails)
+								{
 									CVD_SCOPE_TRACE_SOLVER_STEP(CVDDC_ResimDetails, TEXT("Resim Step Start"));
 									CVD_TRACE_PARTICLES_SOA(GetEvolution()->GetParticles(), &GetEvolution()->GetRigidClustering());
 								}
@@ -2311,27 +2312,11 @@ namespace Chaos
 										}
 									}
 #endif
-#if WITH_CHAOS_VISUAL_DEBUGGER
-									if (bFirst && bIsTracingResimDetails)
-									{
-										CVD_TRACE_SOLVER_END_FRAME(FPhysicsSolverBase, *this);
-										CVD_TRACE_SOLVER_START_FRAME(FPhysicsSolverBase, *this);
-										CVD_SCOPE_TRACE_SOLVER_STEP(CVDDC_ResimDetails, TEXT("Resim Step Start"));
-										CVD_TRACE_PARTICLES_SOA(GetEvolution()->GetParticles(), &GetEvolution()->GetRigidClustering());
-									}
-#endif
-
 									{
 										QUICK_SCOPE_CYCLE_COUNTER(ChaosRewindAndResim_RewindCallback_PreResimStep_Internal);
 										MRewindCallback->PreResimStep_Internal(Step, bFirst);
 									}
 
-#if WITH_CHAOS_VISUAL_DEBUGGER
-									if (bIsTracingResimDetails)
-									{
-										CVD_TRACE_SOLVER_END_FRAME(FPhysicsSolverBase, *this);
-									}
-#endif
 									// Run the advance but omit the game thread callbacks as we're executing during the physics tick
 									FSolverTasksPTOnly ImmediateTask(*this, PushData);
 									//ensure(bSolverHasFrozenGameThreadCallbacks == false);	//We don't support this for resim as it's very expensive and difficult to schedule
@@ -2349,6 +2334,13 @@ namespace Chaos
 						// Clean up after resim. This is where particles are unfrozen if they were frozen due to bubble resim.
 						MRewindData->EndResim();
 
+#if WITH_CHAOS_VISUAL_DEBUGGER
+						if (bIsCVDTracing)
+						{
+							CVD_TRACE_SOLVER_END_FRAME(FPhysicsSolverBase, *this);
+						}
+#endif
+
 						MarshallingManager.SetConsumerData_Internal(ConsumerData);
 						GetEvolution()->GetIslandManager().ResetParticleResimFrame();
 
diff --git a/Engine/Source/Runtime/Experimental/Chaos/Private/RewindData.cpp b/Engine/Source/Runtime/Experimental/Chaos/Private/RewindData.cpp
--- a/Engine/Source/Runtime/Experimental/Chaos/Private/RewindData.cpp
+++ b/Engine/Source/Runtime/Experimental/Chaos/Private/RewindData.cpp
@@ -1003,6 +1003,8 @@ void FRewindData::RewindToFrame(int32 RewindFrame)
 		return;
 	}
 
+	CVD_SCOPE_TRACE_SOLVER_STEP(CVDDC_ResimDetails, TEXT("Rewind To Frame"));
+
 	// Start measuring resim time
 	ResimStartCycleCounter = FPlatformTime::Cycles64();
 
@@ -1221,6 +1223,10 @@ void FRewindData::RewindToFrame(int32 RewindFrame)
 	// Bubble resim only: "Freeze" particles outside resim bubbles
 	FreezeParticlesOutsideResimBubbles();
 
+#if WITH_CHAOS_VISUAL_DEBUGGER
+	TraceParticlesModifiedDuringRewind();
+#endif
+
 #if !UE_BUILD_SHIPPING
 	// For now, just ensure that the joints are InSync
 	for(FDirtyJointInfo& DirtyJointInfo : DirtyJoints)
