Add descriptions to some more trace channels

#Synced-CL 52439714

diff --git a/Engine/Source/Runtime/Core/Private/ProfilingDebugging/CookStats.cpp b/Engine/Source/Runtime/Core/Private/ProfilingDebugging/CookStats.cpp
--- a/Engine/Source/Runtime/Core/Private/ProfilingDebugging/CookStats.cpp
+++ b/Engine/Source/Runtime/Core/Private/ProfilingDebugging/CookStats.cpp
@@ -13,7 +13,7 @@ CORE_API void FCookStatsManager::LogCookStats(AddStatFuncRef AddStat)
 	CookStatsCallbacks.Broadcast(AddStat);
 }
 
-UE_TRACE_CHANNEL_DEFINE(CookChannel)
+UE_TRACE_CHANNEL_DEFINE(CookChannel, "Package cook events: package names and IDs, asset class names, and per-package timing scopes for each cook processing stage. Enables the Unreal Cooking Insights (Packages) tab in Timing Insights profiler.")
 
 UE_TRACE_EVENT_BEGIN(CookTrace, Package)
 	UE_TRACE_EVENT_FIELD(uint64, Id)
diff --git a/Engine/Source/Runtime/Core/Private/ProfilingDebugging/HttpTrace.cpp b/Engine/Source/Runtime/Core/Private/ProfilingDebugging/HttpTrace.cpp
--- a/Engine/Source/Runtime/Core/Private/ProfilingDebugging/HttpTrace.cpp
+++ b/Engine/Source/Runtime/Core/Private/ProfilingDebugging/HttpTrace.cpp
@@ -9,7 +9,7 @@
 #include "Trace/Trace.h"
 
 ////////////////////////////////////////////////////////////////////////////////////////////////////
-UE_TRACE_CHANNEL_DEFINE(HttpChannel)
+UE_TRACE_CHANNEL_DEFINE(HttpChannel, "HTTP dispatcher and request lifecycle events: dispatcher/category creation, request start (URL, priority), chunk range additions, and request completion (status code, content length).")
 
 ////////////////////////////////////////////////////////////////////////////////////////////////////
 UE_TRACE_EVENT_BEGIN(Http, DispatcherCreated, NoSync)
diff --git a/Engine/Source/Runtime/Core/Private/ProfilingDebugging/PlatformFileTrace.cpp b/Engine/Source/Runtime/Core/Private/ProfilingDebugging/PlatformFileTrace.cpp
--- a/Engine/Source/Runtime/Core/Private/ProfilingDebugging/PlatformFileTrace.cpp
+++ b/Engine/Source/Runtime/Core/Private/ProfilingDebugging/PlatformFileTrace.cpp
@@ -16,7 +16,7 @@
 
 #define PLATFORMFILETRACE_DEBUG_ENABLED 0
 
-UE_TRACE_CHANNEL(FileChannel)
+UE_TRACE_CHANNEL(FileChannel, "Platform file I/O operations: open, close, read, and write events with file handle, offset, size, and cycle timestamps.")
 
 UE_TRACE_EVENT_BEGIN(PlatformFile, BeginOpen)
 	UE_TRACE_EVENT_FIELD(uint64, Cycle)
diff --git a/Engine/Source/Runtime/Core/Private/ProfilingDebugging/SpatialTrace.cpp b/Engine/Source/Runtime/Core/Private/ProfilingDebugging/SpatialTrace.cpp
--- a/Engine/Source/Runtime/Core/Private/ProfilingDebugging/SpatialTrace.cpp
+++ b/Engine/Source/Runtime/Core/Private/ProfilingDebugging/SpatialTrace.cpp
@@ -11,7 +11,7 @@
 #define UE_TRACE_SPATIAL_COMPONENT_DELTAS 0
 #endif // UE_TRACE_SPATIAL_COMPONENT_DELTAS
 
-UE_TRACE_MINIMAL_CHANNEL_DEFINE(SpatialChannel);
+UE_TRACE_MINIMAL_CHANNEL_DEFINE(SpatialChannel, "Named spatial point data (position, velocity, direction) sampled over time. Allows tracking the spatial state of arbitrary points for debugging movement or physics.");
 
 UE_TRACE_MINIMAL_EVENT_BEGIN(Spatial, PointSpec, NoSync | Important)
 	UE_TRACE_MINIMAL_EVENT_FIELD(uint32, Id)
diff --git a/Engine/Source/Runtime/Core/Private/Stats/StatsTrace.cpp b/Engine/Source/Runtime/Core/Private/Stats/StatsTrace.cpp
--- a/Engine/Source/Runtime/Core/Private/Stats/StatsTrace.cpp
+++ b/Engine/Source/Runtime/Core/Private/Stats/StatsTrace.cpp
@@ -11,7 +11,7 @@
 
 #if STATSTRACE_ENABLED
 
-UE_TRACE_CHANNEL(StatsChannel)
+UE_TRACE_CHANNEL(StatsChannel, "Named stat declarations and batched stat operations (increment, decrement, add, set) for integer and float stats declared with DECLARE_STATS_GROUP. Viewable in Unreal Insights alongside the Counters channel.")
 
 UE_TRACE_EVENT_BEGIN(Stats, Spec, NoSync|Important)
 	UE_TRACE_EVENT_FIELD(uint32, Id)
diff --git a/Engine/Source/Runtime/Engine/Private/ObjectPropertyTrace.cpp b/Engine/Source/Runtime/Engine/Private/ObjectPropertyTrace.cpp
--- a/Engine/Source/Runtime/Engine/Private/ObjectPropertyTrace.cpp
+++ b/Engine/Source/Runtime/Engine/Private/ObjectPropertyTrace.cpp
@@ -8,7 +8,7 @@
 #if WITH_ENGINE
 
 #if OBJECT_PROPERTY_TRACE_ENABLED
-UE_TRACE_CHANNEL(ObjectProperties)
+UE_TRACE_CHANNEL(ObjectProperties, "Per-frame snapshots of UObject property values (type, name, parent hierarchy, value) for registered objects. Requires the Object channel. Enables property inspection over time in Unreal Insights.")
 
 UE_TRACE_EVENT_BEGIN(Object, ClassPropertyStringId, NoSync|Important)
 	UE_TRACE_EVENT_FIELD(uint32, Id)
diff --git a/Engine/Source/Runtime/Engine/Private/ObjectTrace.cpp b/Engine/Source/Runtime/Engine/Private/ObjectTrace.cpp
--- a/Engine/Source/Runtime/Engine/Private/ObjectTrace.cpp
+++ b/Engine/Source/Runtime/Engine/Private/ObjectTrace.cpp
@@ -14,7 +14,7 @@
 #include "UObject/UObjectAnnotation.h"
 #endif // WITH_EDITOR
 
-UE_TRACE_MINIMAL_CHANNEL(ObjectChannel)
+UE_TRACE_MINIMAL_CHANNEL(ObjectChannel, "UObject and UWorld lifecycle (creation, destruction), type hierarchy, transforms, events, controller-pawn possession, and per-frame recording info. Required by Animation, ObjectProperties, and VisualLogger channels.")
 
 UE_TRACE_MINIMAL_EVENT_BEGIN(Object, Type)
 	UE_TRACE_MINIMAL_EVENT_FIELD(uint64, Id)
diff --git a/Engine/Source/Runtime/Engine/Private/VisualLogger/VisualLoggerTraceDevice.cpp b/Engine/Source/Runtime/Engine/Private/VisualLogger/VisualLoggerTraceDevice.cpp
--- a/Engine/Source/Runtime/Engine/Private/VisualLogger/VisualLoggerTraceDevice.cpp
+++ b/Engine/Source/Runtime/Engine/Private/VisualLogger/VisualLoggerTraceDevice.cpp
@@ -9,7 +9,7 @@
 #include "VisualLogger/VisualLoggerCustomVersion.h"
 #include "UObject/UE5MainStreamObjectVersion.h"
 
-UE_TRACE_MINIMAL_CHANNEL_DEFINE(VisualLoggerChannel);
+UE_TRACE_MINIMAL_CHANNEL_DEFINE(VisualLoggerChannel, "Visual Logger entries (owner object, timing, serialized log data) recorded into the trace stream. Allows reviewing visual log data in Unreal Insights instead of a standalone .vlog file.");
 
 UE_TRACE_MINIMAL_EVENT_BEGIN(VisualLogger, VisualLogEntryHeader)
 	UE_TRACE_MINIMAL_EVENT_FIELD(uint32, ID)
diff --git a/Engine/Source/Runtime/Experimental/Chaos/Private/ChaosInsights/ChaosInsights.cpp b/Engine/Source/Runtime/Experimental/Chaos/Private/ChaosInsights/ChaosInsights.cpp
--- a/Engine/Source/Runtime/Experimental/Chaos/Private/ChaosInsights/ChaosInsights.cpp
+++ b/Engine/Source/Runtime/Experimental/Chaos/Private/ChaosInsights/ChaosInsights.cpp
@@ -10,7 +10,7 @@
 #include "Misc/CString.h"
 #include "HAL/PlatformTime.h"
 
-UE_TRACE_CHANNEL(ChaosLocksChannel)
+UE_TRACE_CHANNEL(ChaosLocksChannel, "Chaos physics read/write lock acquisition events (begin, acquired, end) with cycle timestamps. Useful for diagnosing lock contention in the physics solver.")
 
 UE_TRACE_EVENT_BEGIN(Chaos, LockAcquireBegin)
 	UE_TRACE_EVENT_FIELD(uint64, Cycle)
diff --git a/Engine/Source/Runtime/Experimental/Chaos/Private/ChaosVisualDebugger/ChaosVisualDebuggerTrace.cpp b/Engine/Source/Runtime/Experimental/Chaos/Private/ChaosVisualDebugger/ChaosVisualDebuggerTrace.cpp
--- a/Engine/Source/Runtime/Experimental/Chaos/Private/ChaosVisualDebugger/ChaosVisualDebuggerTrace.cpp
+++ b/Engine/Source/Runtime/Experimental/Chaos/Private/ChaosVisualDebugger/ChaosVisualDebuggerTrace.cpp
@@ -43,7 +43,7 @@ UE_TRACE_MINIMAL_EVENT_DEFINE(ChaosVDLogger, ChaosVDSolverFrameStart)
 
 UE_TRACE_MINIMAL_EVENT_DEFINE(ChaosVDLogger, ChaosVDSolverFrameEnd)
 
-UE_TRACE_MINIMAL_CHANNEL_DEFINE(ChaosVDChannel);
+UE_TRACE_MINIMAL_CHANNEL_DEFINE(ChaosVDChannel, "Chaos physics solver frames, simulation steps, particles, constraints (collision, joints, character ground), scene queries, implicit geometry, and debug shapes. Required for the Chaos Visual Debugger (CVD) tool.");
 UE_TRACE_MINIMAL_EVENT_DEFINE(ChaosVDLogger, ChaosVDParticle)
 UE_TRACE_MINIMAL_EVENT_DEFINE(ChaosVDLogger, ChaosVDParticleDestroyed)
 UE_TRACE_MINIMAL_EVENT_DEFINE(ChaosVDLogger, ChaosVDSolverStepStart)
diff --git a/Engine/Source/Runtime/Net/Core/Private/Net/Core/Trace/Reporters/NetTraceReporter.cpp b/Engine/Source/Runtime/Net/Core/Private/Net/Core/Trace/Reporters/NetTraceReporter.cpp
--- a/Engine/Source/Runtime/Net/Core/Private/Net/Core/Trace/Reporters/NetTraceReporter.cpp
+++ b/Engine/Source/Runtime/Net/Core/Private/Net/Core/Trace/Reporters/NetTraceReporter.cpp
@@ -14,7 +14,7 @@
 uint32 FNetTraceReporter::NetTraceReporterVersion = 1;
 
 NETCORE_API UE_TRACE_CHANNEL_EXTERN(NetChannel)
-UE_TRACE_CHANNEL_DEFINE(NetChannel)
+UE_TRACE_CHANNEL_DEFINE(NetChannel, "Network connections, packets (sent/received/dropped), packet content, replicated object lifecycle, and per-connection stats. Viewable in the Network Insights profiler.")
 
 // We always output this event first to make sure we have a version number for backwards compatibility
 UE_TRACE_EVENT_BEGIN(NetTrace, InitEvent)
