Allow multipart ranged requests to parse the response up to the full range of uint32 sizes which we support.

#rnx

- Originally I used FAnsiStringView to parse the full buffer as it had the best support for finding sequences and slicing ranges but the string view class is l;imited at in32 for size and we are supposed to support up to uint32 data sizes in IAX.
- Replaced FAnsiStringView with FMemoryView which uses uint64 for size and can cover all our use cases.
- We continue to use the stringview class when parsing header lines as they should fit within bounds and it will make the data easier to visualize in the debugger.
- Added ::FindSequence for finding a byte seqeuence (string) in a FMemoryView more easily.

#Synced-CL 52609747

diff --git a/Engine/Source/Runtime/Experimental/IoStore/OnDemand/Private/OnDemandHttpClient.cpp b/Engine/Source/Runtime/Experimental/IoStore/OnDemand/Private/OnDemandHttpClient.cpp
--- a/Engine/Source/Runtime/Experimental/IoStore/OnDemand/Private/OnDemandHttpClient.cpp
+++ b/Engine/Source/Runtime/Experimental/IoStore/OnDemand/Private/OnDemandHttpClient.cpp
@@ -136,40 +136,76 @@ static FAnsiString ParseMultiRangeBoundary(UE::IoStore::HTTP::FResponse& Respons
 	return Output;
 }
 
-FAnsiStringView SkipToNextLine(FAnsiStringView View, int32 Offset = 0)
+/**
+ * Searches for the first occurrence of Sequence bytes in View, starting at Offset.
+ * Returns INDEX_NONE if not found.
+ */
+static uint64 FindSequence(FMemoryView View, FAnsiStringView Sequence, uint64 Offset = 0)
+{
+	if (Sequence.IsEmpty())
+	{
+		return Offset;
+	}
+
+	const uint64 DataSize = View.GetSize();
+	const uint64 SequenceLen = (uint64)Sequence.Len();
+
+	if (DataSize < Offset + SequenceLen)
+	{
+		return INDEX_NONE;
+	}
+
+	const uint8* Data = static_cast<const uint8*>(View.GetData());
+	const uint8* SequenceData = reinterpret_cast<const uint8*>(Sequence.GetData());
+
+	for (uint64 Index = Offset; Index + SequenceLen <= DataSize; ++Index)
+	{
+		if (FMemory::Memcmp(Data + Index, SequenceData, SequenceLen) == 0)
+		{
+			return Index;
+		}
+	}
+
+	return INDEX_NONE;
+}
+
+FMemoryView SkipToNextLine(FMemoryView View, uint64 Offset = 0)
 {
 	constexpr FAnsiStringView CRLF = ANSITEXTVIEW("\r\n");
 
-	int32 LineEndingIndex = View.Find(CRLF, Offset, ESearchCase::IgnoreCase);
+	const uint64 LineEndingIndex = FindSequence(View, CRLF, Offset);
 	if (LineEndingIndex != INDEX_NONE)
 	{
-		return View.RightChop(LineEndingIndex + 2);
+		return View.RightChop(LineEndingIndex + CRLF.Len());
 	}
 	else
 	{
-		return FAnsiStringView();
+		return FMemoryView();
 	}
 }
 
-FAnsiStringView ExtractTextLine(FAnsiStringView View)
+TIoStatusOr<FAnsiStringView> ExtractTextLine(FMemoryView View)
 {
 	constexpr FAnsiStringView CRLF = ANSITEXTVIEW("\r\n");
 
-	int32 LineEndingIndex = View.Find(CRLF, 0, ESearchCase::IgnoreCase);
+	const uint64 LineEndingIndex = FindSequence(View, CRLF, 0);
 	if (LineEndingIndex != INDEX_NONE)
 	{
-		return View.Left(LineEndingIndex);
+		View = View.Left(LineEndingIndex);
 	}
-	else
+
+	if (View.GetSize() > (uint64)MAX_int32)
 	{
-		return View;
+		return FIoStatus(EIoErrorCode::HttpReadError, TEXT("Multirange request: Oversized line indicates malformed data"));
 	}
+
+	return FAnsiStringView(static_cast<const char*>(View.GetData()), (int32)View.GetSize());
 }
 
-FAnsiStringView ExtractData(FAnsiStringView View, FAnsiStringView Boundary, TPair<uint32, uint32> ByteRange)
+FMemoryView ExtractData(FMemoryView View, FAnsiStringView Boundary, TPair<uint32, uint32> ByteRange)
 {
 	// We know that ByteRange.Value >= ByteRange.Key as that was validated (in ParseContentRange) when the byte range was parsed
-	int32 DataLength = (int32)(ByteRange.Value - ByteRange.Key) + 1;
+	const uint64 DataLength = (uint64)(ByteRange.Value - ByteRange.Key) + 1;
 	return View.Left(DataLength);
 }
 
@@ -222,26 +258,26 @@ FIoStatus ParseContentRange(FAnsiStringView Line, TPair<uint32,uint32>& ByteRang
 		return FIoStatus(EIoErrorCode::HttpReadError, TEXT("Multirange request: Malformed byte range for 'Content-Range'"));
 	}
 
-	if ((EndByte - StartByte) + 1 > MAX_int32) // Temp validation until string view (int32 limit) is replaced.
-	{
-		return FIoStatus(EIoErrorCode::HttpReadError, TEXT("Multirange request:Byte range is too large"));
-	}
-
 	ByteRange.Key = (uint32)StartByte;
 	ByteRange.Value = (uint32)EndByte;
 
 	return EIoErrorCode::Ok;
 }
 
-static FIoStatus ParseHeader(FAnsiStringView& BodyView, FAnsiStringView Boundary, TPair<uint32, uint32>& OutByteRange)
+static FIoStatus ParseHeader(FMemoryView& BodyView, FAnsiStringView Boundary, TPair<uint32, uint32>& OutByteRange)
 {
 	constexpr FAnsiStringView ContentRange = ANSITEXTVIEW("Content-Range");
 
 	bool bReadRange = false;
 	while (!BodyView.IsEmpty())
 	{
-		FAnsiStringView Line = ExtractTextLine(BodyView);
-		BodyView = SkipToNextLine(BodyView, Line.Len());
+		TIoStatusOr<FAnsiStringView> LineResult = ExtractTextLine(BodyView);
+		if (!LineResult.IsOk())
+		{
+			return LineResult.Status();
+		}
+		FAnsiStringView Line = LineResult.ConsumeValueOrDie();
+		BodyView = SkipToNextLine(BodyView, (uint64)Line.Len());
 
 		// TODO: Clean up the parsing, should be less logic to determine that we read the correct format.
 		if (Line.StartsWith(ContentRange))
@@ -291,14 +327,18 @@ FIoStatus FResponseBody::Parse(FIoBuffer&& InBody, FAnsiStringView Boundary)
 {
 	constexpr FAnsiStringView DelimEnd = ANSITEXTVIEW("--");
 
-	// Note that in theory the body data can contain '\0' if the content type is binary, but the string view
-	// code will ignore that if we provide the full size and makes it easier to parse the headers.
-	FAnsiStringView BodyView((const char*)InBody.GetData(), (int32)InBody.GetSize()); // TODO - Safer cast
+	FMemoryView BodyView(InBody.GetData(), InBody.GetSize());
 
 	while (!BodyView.IsEmpty())
 	{
-		FAnsiStringView Line = ExtractTextLine(BodyView);
-		BodyView = SkipToNextLine(BodyView, Line.Len());
+		TIoStatusOr<FAnsiStringView> LineResult = ExtractTextLine(BodyView);
+		if (!LineResult.IsOk())
+		{
+			Reset();
+			return LineResult.Status();
+		}
+		FAnsiStringView Line = LineResult.ConsumeValueOrDie();
+		BodyView = SkipToNextLine(BodyView, (uint64)Line.Len());
 
 		if (Line.StartsWith(Boundary))
 		{
@@ -317,11 +357,11 @@ FIoStatus FResponseBody::Parse(FIoBuffer&& InBody, FAnsiStringView Boundary)
 			}
 
 			// Now parse the data
-			FAnsiStringView DataView = ExtractData(BodyView, Boundary, ByteRange);
-			BodyView = SkipToNextLine(BodyView, DataView.Len());
+			FMemoryView DataView = ExtractData(BodyView, Boundary, ByteRange);
+			BodyView = SkipToNextLine(BodyView, DataView.GetSize());
 
 			// TODO: Safer cast
-			const uint32 DataOffset = (uint32)((uint8*)DataView.GetData() - InBody.GetData());
+			const uint32 DataOffset = (uint32)((const uint8*)DataView.GetData() - InBody.GetData());
 
 			// TODO: Validate the boundary and data with ByteRange?
 			DataRanges.Add(FResponseBody::FRange
