Copy the Version field when copying an FManagedArrayCollection, otherwise it's left uninitialized by the FManagedArrayCollection copy constructor

#Synced-CL 52781822

diff --git a/Engine/Source/Runtime/Experimental/Chaos/Public/GeometryCollection/ManagedArrayCollection.h b/Engine/Source/Runtime/Experimental/Chaos/Public/GeometryCollection/ManagedArrayCollection.h
--- a/Engine/Source/Runtime/Experimental/Chaos/Public/GeometryCollection/ManagedArrayCollection.h
+++ b/Engine/Source/Runtime/Experimental/Chaos/Public/GeometryCollection/ManagedArrayCollection.h
@@ -180,6 +180,8 @@ public:
 
 	void CopyTo(FManagedArrayCollection* Collection, const TArray<FName>& GroupsToSkip = TArray<FName>(), TArray<TTuple<FName, FName>> AttributesToSkip = TArray<TTuple<FName, FName>>()) const
 	{
+		Collection->Version = Version;
+
 		if (!Map.IsEmpty())
 		{
 			for (const TTuple<FKeyType, FValueType>& Entry : Map)
